Create ProceduralVegetationEditor plugin
Note: Materials are broken on the sample assets and still need to be fixed in a separate CL

#virtualized

diff --git a/Engine/Plugins/Experimental/ProceduralVegetationEditor/Source/MegaPlants/MegaPlants.Build.cs b/Engine/Plugins/Experimental/ProceduralVegetationEditor/Source/MegaPlants/MegaPlants.Build.cs
new file mode 100644
--- /dev/null
+++ b/Engine/Plugins/Experimental/ProceduralVegetationEditor/Source/MegaPlants/MegaPlants.Build.cs
@@ -0,0 +1,40 @@
+// Copyright Epic Games, Inc. All Rights Reserved.
+
+namespace UnrealBuildTool.Rules
+{
+	public class MegaPlants : ModuleRules
+	{
+		public MegaPlants(ReadOnlyTargetRules Target) : base(Target)
+		{
+			PCHUsage = PCHUsageMode.UseExplicitOrSharedPCHs;
+
+			PublicDependencyModuleNames.AddRange(
+				new string[]
+				{
+				});
+
+			PrivateDependencyModuleNames.AddRange(
+				new string[]
+				{
+					"Core",
+					"CoreUObject",
+					"Chaos",
+					"Engine",
+					"RenderCore",
+					"GeometryCore",
+					"GeometryFramework",
+					"GeometryScriptingCore",
+					"Json",
+					"JsonUtilities",
+					"MeshDescription",
+					"MeshConversion",
+					"PlanarCut",
+					"SkeletalMeshDescription",
+					"StaticMeshDescription",
+					"ModelingComponents",
+					"PCG" // ... add private dependencies that you statically link with here ...	
+				}
+			);
+		}
+	}
+}
diff --git a/Engine/Plugins/Experimental/ProceduralVegetationEditor/Source/MegaPlants/Private/Facades/MegaPlantBoneFacade.cpp b/Engine/Plugins/Experimental/ProceduralVegetationEditor/Source/MegaPlants/Private/Facades/MegaPlantBoneFacade.cpp
new file mode 100644
--- /dev/null
+++ b/Engine/Plugins/Experimental/ProceduralVegetationEditor/Source/MegaPlants/Private/Facades/MegaPlantBoneFacade.cpp
@@ -0,0 +1,242 @@
+﻿// Copyright Epic Games, Inc. All Rights Reserved.
+
+#include "Facades/MegaPlantBoneFacade.h"
+#include "MegaPlantsModule.h"
+#include "Facades/MegaPlantAttributesNames.h"
+#include "Facades/MegaPlantBranchFacade.h"
+#include "Facades/MegaPlantPointFacade.h"
+
+class FGeometryCollection;
+
+namespace MegaPlants::Facades
+{
+	FBoneFacade::FBoneFacade(FManagedArrayCollection& InCollection)
+		: BoneName(InCollection, AttributeNames::BoneName, GroupNames::BonesGroup),
+		BoneParentIndex(InCollection, AttributeNames::BoneParentIndex, GroupNames::BonesGroup),
+		BonePose(InCollection, AttributeNames::BonePose, GroupNames::BonesGroup),
+		BonePointIndex(InCollection, AttributeNames::BonePointIndex, GroupNames::BonesGroup),
+		BoneParentPointIndex(InCollection, AttributeNames::BoneParentPointIndex, GroupNames::BonesGroup),
+		BoneBranchIndex(InCollection, AttributeNames::BoneBranchIndex, GroupNames::BonesGroup),
+		NjordPixelIds(InCollection, AttributeNames::NjordPixelId, GroupNames::VerticesGroup)
+	{}
+
+	FBoneFacade::FBoneFacade(const FManagedArrayCollection& InCollection)
+		: BoneName(InCollection, AttributeNames::BoneName, GroupNames::BonesGroup),
+		BoneParentIndex(InCollection, AttributeNames::BoneParentIndex, GroupNames::BonesGroup),
+		BonePose(InCollection, AttributeNames::BonePose, GroupNames::BonesGroup),
+		BonePointIndex(InCollection, AttributeNames::BonePointIndex, GroupNames::BonesGroup),
+		BoneParentPointIndex(InCollection, AttributeNames::BoneParentPointIndex, GroupNames::BonesGroup),
+		BoneBranchIndex(InCollection, AttributeNames::BoneBranchIndex, GroupNames::BonesGroup),
+		NjordPixelIds(InCollection, AttributeNames::NjordPixelId, GroupNames::VerticesGroup)
+	{}
+
+	void FBoneFacade::DefineSchema(FManagedArrayCollection& InCollection)
+	{
+		InCollection.AddAttribute<FString>(AttributeNames::BoneName, GroupNames::BonesGroup);
+		InCollection.AddAttribute<int32>(AttributeNames::BoneParentIndex, GroupNames::BonesGroup);
+		InCollection.AddAttribute<FTransform>(AttributeNames::BonePose, GroupNames::BonesGroup);
+		InCollection.AddAttribute<int32>(AttributeNames::BonePointIndex, GroupNames::BonesGroup);
+		InCollection.AddAttribute<int32>(AttributeNames::BoneParentPointIndex, GroupNames::BonesGroup);
+		InCollection.AddAttribute<int32>(AttributeNames::BoneBranchIndex, GroupNames::BonesGroup);
+		InCollection.AddAttribute<float>(AttributeNames::NjordPixelId, GroupNames::VerticesGroup);
+	}
+
+	bool FBoneFacade::IsValid() const
+	{
+		return NjordPixelIds.IsValid();
+	}
+
+	void FBoneFacade::GetNjordPixelIds(TArray<float>& VertexNjordIds) const
+	{
+		for (auto VertexBoneID : NjordPixelIds.Get())
+		{
+			VertexNjordIds.Add(VertexBoneID);
+		}
+	}
+
+	const TManagedArray<FString>& FBoneFacade::GetBoneNames() const
+	{
+		return BoneName.Get();
+	}
+
+	const TManagedArray<int32>& FBoneFacade::GetBoneParentIndices() const
+	{
+		return BoneParentIndex.Get();
+	}
+
+	const TManagedArray<FTransform>& FBoneFacade::GetBonePoses() const
+	{
+		return BonePose.Get();
+	}
+
+	const TManagedArray<float>& FBoneFacade::GetNjordPixelIds() const
+	{
+		return NjordPixelIds.Get();
+	}
+
+	TManagedArray<float>& FBoneFacade::ModifyNjordPixelIds()
+	{
+		return NjordPixelIds.Modify();
+	}
+
+	TArray<FBoneNode> GetParentBranchBones(const  TArray<FBoneNode>& InBoneNodes, int32 InParentIndex)
+	{
+		TArray<FBoneNode> ParentBranchBones;
+		for (const auto& Bone : InBoneNodes)
+		{
+			if (InParentIndex == Bone.BranchIndex)
+			{
+				ParentBranchBones.Add(Bone);
+			}
+		}
+		return ParentBranchBones;
+	}
+
+	FBoneNode* FindBoneForPoint(const FPointFacade& PointFacade, TArray<FBoneNode>& BoneNodes, int32 BonePoint, float MaxDistance = 1.0)
+	{
+		const TManagedArray<FVector3f>& Positions = PointFacade.GetPositions();
+		FVector BonePosition = FVector(Positions[BonePoint]);
+		float BonePScale = PointFacade.GetPointScale(BonePoint);
+				
+		FBoneNode* FoundBone = nullptr;
+				
+		for (int32 NodeIdx = 0; NodeIdx < BoneNodes.Num(); ++NodeIdx)
+		{
+			FVector NodePosition = BoneNodes[NodeIdx].AbsolutePosition;
+			FVector NextPosition = NodePosition;
+			if (BoneNodes.IsValidIndex(NodeIdx + 1))
+			{
+				NextPosition = BoneNodes[NodeIdx + 1].AbsolutePosition;
+			}
+
+			float Dist = FMath::PointDistToSegment(BonePosition, NodePosition, NextPosition);
+			if (Dist <= BonePScale * 2)
+			{
+				FoundBone = &BoneNodes[NodeIdx];
+			}
+		}
+
+		return FoundBone;
+	}
+	
+	void FBoneFacade::CreateBoneData(TArray<FBoneNode>& Bones) const
+	{
+		FPointFacade PointFacade(NjordPixelIds.GetConstCollection());
+		FBranchFacade BranchFacade(NjordPixelIds.GetConstCollection());
+
+		const TManagedArray<FVector3f>& Positions = PointFacade.GetPositions();
+		TArray<int32> SortedBranchIndices;
+		BranchFacade.GetSortedBranchIndicesByHierarchy(SortedBranchIndices);
+		
+		if (SortedBranchIndices.Num() == 0)
+		{
+			UE_LOG(LogMegaPlants, Log, TEXT("Cannot create bones, no branches found"));
+			return;
+		}
+
+		int32 RootBranch = SortedBranchIndices[0];
+		check(BranchFacade.GetPoints(RootBranch).Num() > 0);
+		const int32 RootPointIndex = BranchFacade.GetPoints(RootBranch)[0];
+		
+		check(Positions.Num() > 0 && Positions.Num() >= RootPointIndex);
+		const FVector RootPointPosition = FVector(Positions[RootPointIndex]);
+		
+		FBoneNode RootBoneNode{ "Root",0, INDEX_NONE, RootPointIndex, RootBranch, 1.0, FTransform{RootPointPosition}};
+		Bones.Emplace(MoveTemp(RootBoneNode));
+
+		int ParentBoneIndex = 0;
+		int BoneIndex = 1;
+		FVector PreviousPosition = RootPointPosition;
+		
+		for (int32 i = 0; i < SortedBranchIndices.Num(); ++i)
+		{
+			int32 BranchIndex = SortedBranchIndices[i];
+			int32 BranchNumber = BranchFacade.GetBranchNumber(BranchIndex);
+			UE_LOG(LogMegaPlants, Log, TEXT("Sorted i{%i} Branch Index{%i} Branch Number{%i}"), i, BranchIndex, BranchNumber);
+			
+			auto BranchPoints = BranchFacade.GetPoints(BranchIndex);
+			
+			check(BranchPoints.IsValidIndex(0));
+			int32 ParentIndex = BranchFacade.GetParentBranchIndices(BranchIndex).Last();
+			TArray<FBoneNode> ParentBoneNodes = GetParentBranchBones(Bones, ParentIndex);
+			
+			FBoneNode* FoundBoneNode = FindBoneForPoint(PointFacade, ParentBoneNodes, BranchPoints[0]);
+
+			if (FoundBoneNode == nullptr && i != 0)
+			{
+				UE_LOG(LogMegaPlants, Log, TEXT("Parent bone not found for branch (%i) ParentIndex {%i}"), BranchIndex, ParentIndex);
+			}
+			
+			ParentBoneIndex = FoundBoneNode ? FoundBoneNode->BoneIndex : ParentBoneIndex;
+			PreviousPosition = FoundBoneNode ? FVector(FoundBoneNode->AbsolutePosition) : RootPointPosition;
+
+			TArray<int32> ConnectingPoints;
+			
+			for (int Idx = 0; Idx < BranchPoints.Num(); Idx++)
+			{
+				const int32 PointIdx = BranchPoints[Idx];
+				check(Positions.Num() > PointIdx && PointIdx >= 0);
+				const FVector PointPosition = FVector(Positions[PointIdx]);
+				const float PScale = PointFacade.GetPointScale(PointIdx);
+				
+				FVector NextPosition = PointPosition;
+				if (BranchPoints.IsValidIndex(Idx + 1))
+				{
+					NextPosition = FVector(BranchPoints[Idx + 1]);
+				}
+				
+				float NjordPixelID = PointFacade.GetNjordPixelIndex(PointIdx);
+
+				bool IsConnectingPoint = false;
+				for (auto ChildBranch : BranchFacade.GetChildren(BranchIndex))
+				{
+					const int32 ChildBranchIndex = BranchFacade.GetBranchNumbers().Find(ChildBranch);
+					if (ChildBranchIndex != INDEX_NONE)
+					{
+						const TArray<int32>& ChildPoints = BranchFacade.GetPoints(ChildBranchIndex);
+						check(ChildPoints.Num() > 0);
+						FVector FirstPointPosition = FVector(Positions[ChildPoints[0]]);
+
+						float PointDistToSegment = FMath::PointDistToSegment(FirstPointPosition, PointPosition, NextPosition);
+						if (PointDistToSegment <= PScale && !ConnectingPoints.Contains(ChildPoints[0]))
+						{
+							ConnectingPoints.Add(ChildPoints[0]);
+							IsConnectingPoint = true;
+						}
+					}
+					else
+					{
+						UE_LOG(LogMegaPlants, Warning, TEXT("Child branch number {%i} not found in facade, for branch {%i}"), ChildBranch, BranchFacade.GetBranchNumber(BranchIndex));
+					}
+				}
+
+				bool IsLastPoint = !IsConnectingPoint && !(FMath::Frac(NjordPixelID) <= 0.0001f) && Idx == BranchPoints.Num() - 1;
+				//If this point exist in the BoneNodes set it as parent Bone
+				
+				//UE_LOG(LogMegaPlants, Warning, TEXT("Iterating Branch Points NjordPixelId, Id{%f} PointIndex{%i} PointFacade.GetBudNumber(BranchPoints[PointIndex]){%i} BranchID{%i} ParentBoneIndex{%i} bFound{%i} FMath::Frac(NjordPixelID){%f}"),
+				//	NjordPixelID, PointIdx, PointFacade.GetBudNumber(PointIdx), BranchIndex, ParentBoneIndex, FoundBoneNode ? 1 : 0 , FMath::Frac(NjordPixelID));
+				
+				if ((FMath::Frac(NjordPixelID) <= 0.0001f || IsConnectingPoint || IsLastPoint))
+				{
+					FBoneNode BoneNode;
+					BoneNode.BoneName = *("Bone_" + FString::FromInt(BoneIndex));
+					BoneNode.BoneIndex = BoneIndex;
+					BoneNode.PointIndex = PointIdx;
+					BoneNode.ParentBoneIndex = ParentBoneIndex;
+					BoneNode.BranchIndex = BranchIndex;
+					BoneNode.NjordPixelID = NjordPixelID;
+					FVector Position = FVector(PointPosition) - PreviousPosition;
+					BoneNode.AbsolutePosition = PointPosition;
+					BoneNode.BoneTransform = FTransform{Position};
+					Bones.Emplace(MoveTemp(BoneNode));
+					ParentBoneIndex = BoneIndex;
+					PreviousPosition = FVector(PointPosition);
+					BoneIndex++;
+
+					//UE_LOG(LogMegaPlants, Warning, TEXT("Creating Bone NjordPixelId, Id{%f} BoneId{%i} PointIndex{%i} ParentBoneIndex{%i} BranchID{%i}"),
+					//	BoneNode.NjordPixelID, BoneNode.BoneIndex, BoneNode.PointIndex,BoneNode.ParentBoneIndex,BoneNode.BranchIndex);
+				}
+			}
+		}
+	}
+}
diff --git a/Engine/Plugins/Experimental/ProceduralVegetationEditor/Source/MegaPlants/Private/Facades/MegaPlantBranchFacade.cpp b/Engine/Plugins/Experimental/ProceduralVegetationEditor/Source/MegaPlants/Private/Facades/MegaPlantBranchFacade.cpp
new file mode 100644
--- /dev/null
+++ b/Engine/Plugins/Experimental/ProceduralVegetationEditor/Source/MegaPlants/Private/Facades/MegaPlantBranchFacade.cpp
@@ -0,0 +1,272 @@
+﻿// Copyright Epic Games, Inc. All Rights Reserved.
+
+#include "Facades/MegaPlantBranchFacade.h"
+#include "Facades/MegaPlantAttributesNames.h"
+
+namespace MegaPlants::Facades
+{
+	FBranchFacade::FBranchFacade(FManagedArrayCollection& InCollection)
+		: Collection(&InCollection)
+		, Parents(InCollection, AttributeNames::BranchParents, GroupNames::BranchGroup)
+		, Children(InCollection, AttributeNames::BranchChildren, GroupNames::BranchGroup)
+		, BranchPoints(InCollection, AttributeNames::BranchPoints, GroupNames::BranchGroup)
+		, BranchNumbers(InCollection, AttributeNames::BranchNumber, GroupNames::BranchGroup)
+		, BranchSourceBudNumber(InCollection, AttributeNames::BranchSourceBudNumber, GroupNames::BranchGroup)
+		, BranchFoliageIDs(InCollection, AttributeNames::BranchFoliageIDs, GroupNames::BranchGroup)
+		, BranchUVMaterial(InCollection, AttributeNames::BranchUVMaterial, GroupNames::BranchGroup)
+		, BranchHierarchyNumber(InCollection, AttributeNames::BranchHierarchyNumber, GroupNames::BranchGroup)
+		, TrunkMaterialPathAttribute(InCollection, AttributeNames::TrunkMaterialPath, GroupNames::DetailsGroup)
+	{
+	}
+
+	FBranchFacade::FBranchFacade(const FManagedArrayCollection& InCollection)
+		: Collection(nullptr)
+		, Parents(InCollection, AttributeNames::BranchParents, GroupNames::BranchGroup)
+		, Children(InCollection, AttributeNames::BranchChildren, GroupNames::BranchGroup)
+		, BranchPoints(InCollection, AttributeNames::BranchPoints, GroupNames::BranchGroup)
+		, BranchNumbers(InCollection, AttributeNames::BranchNumber, GroupNames::BranchGroup)
+		, BranchSourceBudNumber(InCollection, AttributeNames::BranchSourceBudNumber, GroupNames::BranchGroup)
+		, BranchFoliageIDs(InCollection, AttributeNames::BranchFoliageIDs, GroupNames::BranchGroup)
+		, BranchUVMaterial(InCollection, AttributeNames::BranchUVMaterial, GroupNames::BranchGroup)
+		, BranchHierarchyNumber(InCollection, AttributeNames::BranchHierarchyNumber, GroupNames::BranchGroup)
+		, TrunkMaterialPathAttribute(InCollection, AttributeNames::TrunkMaterialPath, GroupNames::DetailsGroup)
+	{
+		
+	}
+
+	bool FBranchFacade::IsValid() const
+	{
+		return Parents.IsValid() && Children.IsValid() && BranchPoints.IsValid() && BranchNumbers.IsValid() && BranchSourceBudNumber.IsValid() && BranchHierarchyNumber.IsValid();
+	}
+
+	int32 FBranchFacade::GetElementCount() const
+	{
+		return Parents.Num();
+	}
+
+	const TArray<int32>& FBranchFacade::GetPoints(const int32 Index) const
+	{
+		if (BranchPoints.IsValid() && BranchPoints.IsValidIndex(Index))
+		{
+			return BranchPoints[Index];
+		}
+
+		static const TArray<int32> EmptyPoints;
+		return EmptyPoints;
+	}
+
+	void FBranchFacade::SetPoints(const int32 Index, const TArray<int32>& InPoints)
+	{
+		if (BranchPoints.IsValid() && BranchPoints.IsValidIndex(Index))
+		{
+			BranchPoints.ModifyAt(Index, InPoints);
+		}
+	}
+
+	const TArray<int32>& FBranchFacade::GetChildren(const int32 Index) const
+	{
+		if (Children.IsValid() && Children.IsValidIndex(Index))
+		{
+			return Children[Index];
+		}
+
+		static const TArray<int32> EmptyChildren;
+		return EmptyChildren;
+	}
+
+	void FBranchFacade::SetChildren(const int32 Index, const TArray<int32>& InChildren)
+	{
+		if (Children.IsValid() && Children.IsValidIndex(Index))
+		{
+			Children.ModifyAt(Index, InChildren);
+		}
+	}
+
+	const TArray<int32>& FBranchFacade::GetParents(const int32 Index) const
+	{
+		if (Parents.IsValid() && Parents.IsValidIndex(Index))
+		{
+			return Parents[Index];
+		}
+
+		static const TArray<int32> EmptyParents;
+		return EmptyParents;
+	}
+
+	void FBranchFacade::SetParents(const int32 Index, const TArray<int32>& InParents)
+	{
+		if (Parents.IsValid() && Parents.IsValidIndex(Index))
+		{
+			Parents.ModifyAt(Index, InParents);
+		}
+	}
+
+	int32 FBranchFacade::GetBranchNumber(const int32 Index) const
+	{
+		if (BranchNumbers.IsValid() && BranchNumbers.IsValidIndex(Index))
+		{
+			return BranchNumbers[Index];
+		}
+
+		return INDEX_NONE;
+	}
+
+	int32 FBranchFacade::GetBranchSourceBudNumber(const int32 Index) const
+	{
+		if (BranchSourceBudNumber.IsValid() && BranchSourceBudNumber.IsValidIndex(Index))
+		{
+			return BranchSourceBudNumber[Index];
+		}
+
+		return INDEX_NONE;
+	}
+
+	int32 FBranchFacade::GetTrunkIndex() const
+	{
+		if (IsValid() && BranchNumbers.Num() > 0)
+		{
+			return BranchNumbers.Get().Find(1);
+		}
+
+		return INDEX_NONE;
+	}
+
+	TArray<int32> FBranchFacade::GetParentBranchIndices(int BranchIndex) const
+	{
+		TArray<int32> ParentIndices;
+		if (IsValid() && Parents.Num() > 0)
+		{
+			auto ParentsBranchNumbers = Parents[BranchIndex];
+			for (auto Parent : ParentsBranchNumbers)
+			{
+				ParentIndices.Add(BranchNumbers.Get().Find(Parent));
+			}
+		}
+
+		return ParentIndices;
+	}
+
+	int32 FBranchFacade::GetGenerationNumber(const int32 Index) const
+	{
+		// The generation attributes contain incorrect information at the moment 
+		// thus the number of parents can be used to derive the correct number
+		if (Parents.IsValid() && Parents.IsValidIndex(Index))
+		{
+			return Parents[Index].Num();
+		}
+
+		return 0;
+	}
+
+	const TManagedArray<int32>& FBranchFacade::GetBranchNumbers() const
+	{
+		return BranchNumbers.Get();
+	}
+
+	int32 FBranchFacade::GetBranchUVMaterial(const int32 Index) const
+	{
+		if (BranchUVMaterial.IsValid() && BranchUVMaterial.IsValidIndex(Index))
+		{
+			return BranchUVMaterial[Index];
+		}
+
+		return INDEX_NONE;
+	}
+
+	void FBranchFacade::SetBranchUVMaterial(const int32 Index, const int32 InMaterial)
+	{
+		if (!BranchUVMaterial.IsValid())
+		{
+			BranchUVMaterial.Add();
+		}
+		if (BranchUVMaterial.IsValid() && BranchUVMaterial.IsValidIndex(Index))
+		{
+			BranchUVMaterial.ModifyAt(Index, InMaterial);
+		}
+	}
+
+	int32 FBranchFacade::GetBranchHierarchyNumber(const int32 Index) const
+	{
+		if (BranchHierarchyNumber.IsValid() && BranchHierarchyNumber.IsValidIndex(Index))
+		{
+			return BranchHierarchyNumber[Index];
+		}
+
+		return INDEX_NONE;
+	}
+
+	void FBranchFacade::SetTrunkMaterialPath(const FString& InPath)
+	{
+		check(!IsConst());
+
+		int32 NumElements = Collection->NumElements(GroupNames::DetailsGroup);
+
+		if (NumElements == 0)
+		{
+			Collection->AddElements(1, GroupNames::DetailsGroup);
+		}
+		if (!TrunkMaterialPathAttribute.IsValid())
+		{
+			TrunkMaterialPathAttribute.Add();
+		}
+
+		TrunkMaterialPathAttribute.Modify()[0] = InPath;
+	}
+
+	FString FBranchFacade::GetTrunkMaterialPath() const
+	{
+		if (TrunkMaterialPathAttribute.IsValid() && TrunkMaterialPathAttribute.IsValidIndex(0))
+		{
+			return TrunkMaterialPathAttribute[0];
+		}
+
+		return FString();
+	}
+
+	void FBranchFacade::CopyEntry(const int32 FromIndex, const int32 ToIndex)
+	{
+		if (IsValid() && BranchNumbers.IsValidIndex(FromIndex) && BranchNumbers.IsValidIndex(ToIndex))
+		{
+			Parents.ModifyAt(ToIndex, Parents[FromIndex]);
+			Children.ModifyAt(ToIndex, Children[FromIndex]);
+			BranchPoints.ModifyAt(ToIndex, BranchPoints[FromIndex]);
+			BranchNumbers.ModifyAt(ToIndex, BranchNumbers[FromIndex]);
+			BranchSourceBudNumber.ModifyAt(ToIndex, BranchSourceBudNumber[FromIndex]);
+			BranchFoliageIDs.ModifyAt(ToIndex, BranchFoliageIDs[FromIndex]);
+			BranchHierarchyNumber.ModifyAt(ToIndex, BranchHierarchyNumber[FromIndex]);
+			if(BranchUVMaterial.IsValid())
+			{
+				BranchUVMaterial.ModifyAt(ToIndex, BranchUVMaterial[FromIndex]);
+			}
+		}
+	}
+
+	void FBranchFacade::RemoveEntries(const int32 NumEntries, const int32 StartIndex)
+	{
+		if (IsValid() && BranchNumbers.IsValidIndex(StartIndex) && StartIndex + NumEntries <= BranchNumbers.Num())
+		{
+			BranchNumbers.RemoveElements(NumEntries, StartIndex);
+		}
+	}
+
+	void FBranchFacade::GetSortedBranchIndicesByHierarchy(TArray<int32>& OutSortedIndices) const
+	{
+		for (int32 Index = 0; Index < GetElementCount(); Index++)
+		{
+			OutSortedIndices.Add(Index);
+		}
+
+		OutSortedIndices.Sort([&](int32 Index1, int32 Index2)
+		{
+			int32 Index1Hierarchy = GetBranchHierarchyNumber(Index1);
+			int32 Index2Hierarchy = GetBranchHierarchyNumber(Index2);
+
+			if (Index1Hierarchy == Index2Hierarchy)
+			{
+				return Index1 < Index2;
+			}
+				
+			return Index1Hierarchy < Index2Hierarchy;
+		});
+	}
+}
diff --git a/Engine/Plugins/Experimental/ProceduralVegetationEditor/Source/MegaPlants/Private/Facades/MegaPlantBudVectorsFacade.cpp b/Engine/Plugins/Experimental/ProceduralVegetationEditor/Source/MegaPlants/Private/Facades/MegaPlantBudVectorsFacade.cpp
new file mode 100644
--- /dev/null
+++ b/Engine/Plugins/Experimental/ProceduralVegetationEditor/Source/MegaPlants/Private/Facades/MegaPlantBudVectorsFacade.cpp
@@ -0,0 +1,41 @@
+﻿// Copyright Epic Games, Inc. All Rights Reserved.
+
+#include "MegaPlantBudVectorsFacade.h"
+#include "Facades/MegaPlantAttributesNames.h"
+
+namespace MegaPlants::Facades
+{
+	FBudVectorsFacade::FBudVectorsFacade(FManagedArrayCollection& InCollection)
+			: BudDirection(InCollection, AttributeNames::BudDirection, GroupNames::PointGroup)
+	{}
+
+	FBudVectorsFacade::FBudVectorsFacade(const FManagedArrayCollection& InCollection)
+		: BudDirection(InCollection, AttributeNames::BudDirection, GroupNames::PointGroup)
+	{}
+
+	bool FBudVectorsFacade::IsValid() const
+	{
+		return BudDirection.IsValid();
+	}
+
+	const TArray<FVector3f>& FBudVectorsFacade::GetBudDirection(int32 Index) const
+	{
+		if(BudDirection.IsValid() && BudDirection.IsValidIndex(Index))
+		{
+			return BudDirection[Index];
+		}
+
+		static const TArray<FVector3f> EmptyBudDirections;
+		return EmptyBudDirections;
+	}
+
+	TManagedArray<TArray<FVector3f>>& FBudVectorsFacade::ModifyBudDirections()
+	{
+		return BudDirection.Modify();
+	}
+
+	const TManagedArray<TArray<FVector3f>>& FBudVectorsFacade::GetBudDirections() const
+	{
+		return BudDirection.Get();
+	}
+}
diff --git a/Engine/Plugins/Experimental/ProceduralVegetationEditor/Source/MegaPlants/Private/Facades/MegaPlantBudVectorsFacade.h b/Engine/Plugins/Experimental/ProceduralVegetationEditor/Source/MegaPlants/Private/Facades/MegaPlantBudVectorsFacade.h
new file mode 100644
--- /dev/null
+++ b/Engine/Plugins/Experimental/ProceduralVegetationEditor/Source/MegaPlants/Private/Facades/MegaPlantBudVectorsFacade.h
@@ -0,0 +1,31 @@
+﻿// Copyright Epic Games, Inc. All Rights Reserved.
+
+#pragma once
+
+#include "CoreMinimal.h"
+#include "GeometryCollection/ManagedArrayAccessor.h"
+
+namespace MegaPlants::Facades
+{
+	/**
+	 * FBudVectorsFacade is used to access and manipulate the BudVectors data from Point Group with in the MegaPlant's FManagedArrayCollection
+	 * BudVectors represent different growth related vectors in plant data and are associated with each point. 
+	 */
+	class FBudVectorsFacade
+	{
+	public:
+		FBudVectorsFacade(FManagedArrayCollection& InCollection);
+		FBudVectorsFacade(const FManagedArrayCollection& InCollection);
+
+		bool IsValid() const;
+
+		const TArray<FVector3f>& GetBudDirection(int32 Index) const;
+
+		TManagedArray<TArray<FVector3f>>& ModifyBudDirections();
+
+		const TManagedArray<TArray<FVector3f>>& GetBudDirections() const;
+		
+	private:
+		TManagedArrayAccessor<TArray<FVector3f>> BudDirection;
+	};
+}
diff --git a/Engine/Plugins/Experimental/ProceduralVegetationEditor/Source/MegaPlants/Private/Facades/MegaPlantFoliageFacade.cpp b/Engine/Plugins/Experimental/ProceduralVegetationEditor/Source/MegaPlants/Private/Facades/MegaPlantFoliageFacade.cpp
new file mode 100644
--- /dev/null
+++ b/Engine/Plugins/Experimental/ProceduralVegetationEditor/Source/MegaPlants/Private/Facades/MegaPlantFoliageFacade.cpp
@@ -0,0 +1,364 @@
+﻿// Copyright Epic Games, Inc. All Rights Reserved.
+
+#include "Facades/MegaPlantFoliageFacade.h"
+
+#include "MegaPlantsModule.h"
+#include "Facades/MegaPlantAttributesNames.h"
+
+namespace MegaPlants::Facades
+{
+	FFoliageFacade::FFoliageFacade(FManagedArrayCollection& InCollection, const int32 InitialSize)
+		: ConstCollection(InCollection)
+		, Collection(&InCollection)
+		, NamesAttribute(InCollection, AttributeNames::FoliageName, GroupNames::FoliageNamesGroup)
+		, NameIdsAttribute(InCollection, AttributeNames::FoliageNameID, GroupNames::FoliageGroup)
+		, BranchIdsAttribute(InCollection, AttributeNames::FoliageBranchID, GroupNames::FoliageGroup)
+		, PivotPointsAttribute(InCollection, AttributeNames::FoliagePivotPoint, GroupNames::FoliageGroup)
+		, UpVectorsAttribute(InCollection, AttributeNames::FoliageUPVector, GroupNames::FoliageGroup)
+		, NormalVectorsAttribute(InCollection, AttributeNames::FoliageNormalVector, GroupNames::FoliageGroup)
+		, ScalesAttribute(InCollection, AttributeNames::FoliageScale, GroupNames::FoliageGroup)
+		, LengthFromRootAttribute(InCollection, AttributeNames::FoliageLengthFromRoot, GroupNames::FoliageGroup)
+		, ParentBoneIdsAttribute(InCollection, AttributeNames::FoliageParentBoneID, GroupNames::FoliageGroup)
+		, FoliageIdsAttribute(InCollection, AttributeNames::BranchFoliageIDs, GroupNames::BranchGroup)
+		, FoliagePathAttribute(InCollection, AttributeNames::FoliagePath, GroupNames::DetailsGroup)
+	{
+		DefineSchema(InitialSize);
+	}
+
+	FFoliageFacade::FFoliageFacade(const FManagedArrayCollection& InCollection)
+		: ConstCollection(InCollection)
+		, Collection(nullptr)
+		, NamesAttribute(InCollection, AttributeNames::FoliageName, GroupNames::FoliageNamesGroup)
+		, NameIdsAttribute(InCollection, AttributeNames::FoliageNameID, GroupNames::FoliageGroup)
+		, BranchIdsAttribute(InCollection, AttributeNames::FoliageBranchID, GroupNames::FoliageGroup)
+		, PivotPointsAttribute(InCollection, AttributeNames::FoliagePivotPoint, GroupNames::FoliageGroup)
+		, UpVectorsAttribute(InCollection, AttributeNames::FoliageUPVector, GroupNames::FoliageGroup)
+		, NormalVectorsAttribute(InCollection, AttributeNames::FoliageNormalVector, GroupNames::FoliageGroup)
+		, ScalesAttribute(InCollection, AttributeNames::FoliageScale, GroupNames::FoliageGroup)
+		, LengthFromRootAttribute(InCollection, AttributeNames::FoliageLengthFromRoot, GroupNames::FoliageGroup)
+		, ParentBoneIdsAttribute(InCollection, AttributeNames::FoliageParentBoneID, GroupNames::FoliageGroup)
+		, FoliageIdsAttribute(InCollection, AttributeNames::BranchFoliageIDs, GroupNames::BranchGroup)
+		, FoliagePathAttribute(InCollection, AttributeNames::FoliagePath, GroupNames::DetailsGroup)
+	{
+	}
+
+	void FFoliageFacade::DefineSchema(const int32 InitialSize)
+	{
+		check(!IsConst());
+
+		if (!Collection->HasGroup(GroupNames::FoliageGroup))
+		{
+			Collection->AddGroup(GroupNames::FoliageGroup);
+			if (InitialSize > 0)
+			{
+				Collection->AddElements(InitialSize, GroupNames::FoliageGroup);
+			}
+		}
+
+		if (!Collection->HasGroup(GroupNames::FoliageNamesGroup))
+		{
+			Collection->AddGroup(GroupNames::FoliageNamesGroup);
+		}
+
+		if (!Collection->HasGroup(GroupNames::DetailsGroup))
+		{
+			Collection->AddGroup(GroupNames::DetailsGroup);
+		}
+
+		NamesAttribute.Add();
+		NameIdsAttribute.Add();
+		BranchIdsAttribute.Add();
+		PivotPointsAttribute.Add();
+		UpVectorsAttribute.Add();
+		NormalVectorsAttribute.Add();
+		ScalesAttribute.Add();
+		LengthFromRootAttribute.Add();
+		ParentBoneIdsAttribute.Add();
+		FoliageIdsAttribute.Add();
+		FoliagePathAttribute.Add();
+	}
+
+	bool FFoliageFacade::IsValid() const
+	{
+		return NamesAttribute.IsValid() && NameIdsAttribute.IsValid() && BranchIdsAttribute.IsValid() && PivotPointsAttribute.IsValid() &&
+			UpVectorsAttribute.IsValid() && NormalVectorsAttribute.IsValid() && ScalesAttribute.IsValid() && LengthFromRootAttribute.IsValid() &&
+			FoliageIdsAttribute.IsValid() && ParentBoneIdsAttribute.IsValid();
+	}
+
+	FFoliageEntryData FFoliageFacade::GetFoliageEntry(const int32 Index) const
+	{
+		FFoliageEntryData ReturnData = {};
+		if (IsValid() && Index > -1)
+		{
+			if (NameIdsAttribute.IsValidIndex(Index))
+			{
+				ReturnData.NameId = NameIdsAttribute.Get()[Index];
+			}
+			if (BranchIdsAttribute.IsValidIndex(Index))
+			{
+				ReturnData.BranchId = BranchIdsAttribute.Get()[Index];
+			}
+			if (PivotPointsAttribute.IsValidIndex(Index))
+			{
+				ReturnData.PivotPoint = PivotPointsAttribute.Get()[Index];
+			}
+			if (UpVectorsAttribute.IsValidIndex(Index))
+			{
+				ReturnData.UpVector = UpVectorsAttribute.Get()[Index];
+			}
+			if (NormalVectorsAttribute.IsValidIndex(Index))
+			{
+				ReturnData.NormalVector = NormalVectorsAttribute.Get()[Index];
+			}
+			if (ScalesAttribute.IsValidIndex(Index))
+			{
+				ReturnData.Scale = ScalesAttribute.Get()[Index];
+			}
+			if (LengthFromRootAttribute.IsValidIndex(Index))
+			{
+				ReturnData.LengthFromRoot = LengthFromRootAttribute.Get()[Index];
+			}
+			if (ParentBoneIdsAttribute.IsValidIndex(Index))
+			{
+				ReturnData.ParentBoneID = ParentBoneIdsAttribute.Get()[Index];
+			}
+		}
+		return ReturnData;
+	}
+
+	const float& FFoliageFacade::GetLengthFromRoot(const int32 Index) const
+	{
+		if (LengthFromRootAttribute.IsValid() && LengthFromRootAttribute.IsValidIndex(Index))
+		{
+			return LengthFromRootAttribute[Index];
+		}
+
+		static constexpr float DefaultLengthFromRoot = 0.0;
+		return DefaultLengthFromRoot;
+	}
+
+	void FFoliageFacade::SetLengthFromRoot(const int32 Index, const float Input)
+	{
+		check(!IsConst());
+		if (IsValid() && LengthFromRootAttribute.IsValidIndex(Index))
+		{
+			LengthFromRootAttribute.ModifyAt(Index, Input);
+		}
+	}
+
+	int32 FFoliageFacade::GetParentBoneID(const int32 Index) const
+	{
+		if (ParentBoneIdsAttribute.IsValid() && ParentBoneIdsAttribute.IsValidIndex(Index))
+		{
+			return ParentBoneIdsAttribute[Index];
+		}
+
+		return INDEX_NONE;
+	}
+
+	void FFoliageFacade::SetParentBoneID(const int32 Index, const int32 Input)
+	{
+		check(!IsConst());
+		if (IsValid() && ParentBoneIdsAttribute.IsValidIndex(Index))
+		{
+			ParentBoneIdsAttribute.ModifyAt(Index, Input);
+		}
+	}
+
+	const TArray<int32>& FFoliageFacade::GetFoliageEntryIdsForBranch(const int32 Index) const
+	{
+		if (FoliageIdsAttribute.IsValid() && FoliageIdsAttribute.IsValidIndex(Index))
+		{
+			return FoliageIdsAttribute[Index];
+		}
+
+		static const TArray<int32> EmptyArray;
+		return EmptyArray;
+	}
+
+	int32 FFoliageFacade::AddFoliageEntry(const FFoliageEntryData& InputData)
+	{
+		check(!IsConst());
+		if (IsValid())
+		{
+			const int32 NewIndex = NameIdsAttribute.AddElements(1);
+			NameIdsAttribute.Modify()[NewIndex] = InputData.NameId;
+			BranchIdsAttribute.Modify()[NewIndex] = InputData.BranchId;
+			PivotPointsAttribute.Modify()[NewIndex] = InputData.PivotPoint;
+			UpVectorsAttribute.Modify()[NewIndex] = InputData.UpVector;
+			NormalVectorsAttribute.Modify()[NewIndex] = InputData.NormalVector;
+			ScalesAttribute.Modify()[NewIndex] = InputData.Scale;
+			LengthFromRootAttribute.Modify()[NewIndex] = InputData.LengthFromRoot;
+			ParentBoneIdsAttribute.Modify()[NewIndex] = InputData.ParentBoneID;
+
+			return NewIndex;
+		}
+		return INDEX_NONE;
+	}
+
+	void FFoliageFacade::SetFoliageEntry(const int32 Index, const FFoliageEntryData& InputData)
+	{
+		check(!IsConst());
+		if (IsValid() && Index >= 0)
+		{
+			NameIdsAttribute.ModifyAt(Index, InputData.NameId);
+			BranchIdsAttribute.ModifyAt(Index, InputData.BranchId);
+			PivotPointsAttribute.ModifyAt(Index, InputData.PivotPoint);
+			UpVectorsAttribute.ModifyAt(Index, InputData.UpVector);
+			NormalVectorsAttribute.ModifyAt(Index, InputData.NormalVector);
+			ScalesAttribute.ModifyAt(Index, InputData.Scale);
+			LengthFromRootAttribute.ModifyAt(Index, InputData.LengthFromRoot);
+			LengthFromRootAttribute.ModifyAt(Index, InputData.LengthFromRoot);
+			LengthFromRootAttribute.ModifyAt(Index, InputData.LengthFromRoot);
+			ParentBoneIdsAttribute.ModifyAt(Index, InputData.ParentBoneID);
+		}
+	}
+
+	void FFoliageFacade::SetFoliageIdsArray(const int32 Index, const TArray<int32>& InputIds)
+	{
+		check(!IsConst());
+		if (IsValid() && Index >= 0)
+		{
+			FoliageIdsAttribute.Modify()[Index] = InputIds;
+		}
+	}
+
+	void FFoliageFacade::SetFoliageBranchId(const int32 Index, const int32 InputId)
+	{
+		check(!IsConst());
+		if (IsValid() && BranchIdsAttribute.IsValidIndex(Index))
+		{
+			BranchIdsAttribute.ModifyAt(Index, InputId);
+		}
+	}
+
+	FString FFoliageFacade::GetFoliageName(const int32 Index) const
+	{
+		if (NamesAttribute.IsValid() && NamesAttribute.IsValidIndex(Index))
+		{
+			return NamesAttribute[Index];
+		}
+
+		return FString();
+	}
+
+	const FVector3f& FFoliageFacade::GetPivotPoint(const int32 Index) const
+	{
+		if (PivotPointsAttribute.IsValid() && PivotPointsAttribute.IsValidIndex(Index))
+		{
+			return PivotPointsAttribute[Index];
+		}
+
+		static const FVector3f DefaultPivotPoint = FVector3f::Zero();
+		return DefaultPivotPoint;
+	}
+
+	void FFoliageFacade::SetPivotPoint(const int32 Index, const FVector3f& Input)
+	{
+		check(!IsConst());
+		if (IsValid() && PivotPointsAttribute.IsValidIndex(Index))
+		{
+			PivotPointsAttribute.ModifyAt(Index, Input);
+		}
+	}
+
+	void FFoliageFacade::SetScale(const int32 Index, const float& Input)
+	{
+		check(!IsConst());
+		if (IsValid() && ScalesAttribute.IsValidIndex(Index))
+		{
+			ScalesAttribute.ModifyAt(Index, Input);
+		}
+	}
+	
+	void FFoliageFacade::SetFoliageName(const int32 Index, const FString& Input)
+	{
+		check(!IsConst());
+		if (IsValid() && NamesAttribute.IsValidIndex(Index))
+		{
+			NamesAttribute.ModifyAt(Index, Input);
+		}
+	}
+	
+	void FFoliageFacade::SetFoliageNames(const TArray<FString>& InputNames)
+	{
+		check(!IsConst());
+		Collection->EmptyGroup(GroupNames::FoliageNamesGroup);
+		Collection->AddElements(InputNames.Num(), GroupNames::FoliageNamesGroup);
+		for (int32 Index = 0; Index < InputNames.Num(); ++Index)
+		{
+			NamesAttribute.Modify()[Index] = InputNames[Index];
+		}
+	}
+
+	void FFoliageFacade::SetFoliagePath(const FString InPath)
+	{
+		check(!IsConst());
+
+		int32 NumElements = Collection->NumElements(GroupNames::DetailsGroup);
+
+		if (NumElements == 0)
+		{
+			Collection->AddElements(1, GroupNames::DetailsGroup);
+		}
+
+		FoliagePathAttribute.Modify()[0] = InPath;
+	}
+
+	FString FFoliageFacade::GetFoliagePath() const
+	{
+		if (FoliagePathAttribute.IsValid() && FoliagePathAttribute.IsValidIndex(0))
+		{
+			return FoliagePathAttribute[0];
+		}
+
+		return FString();
+	}
+
+	int32 FFoliageFacade::GetElementCount() const
+	{
+		return NameIdsAttribute.Num();
+	}
+
+	void FFoliageFacade::CopyEntry(const int32 FromIndex, const int32 ToIndex)
+	{
+		if (IsValid() && NameIdsAttribute.IsValidIndex(FromIndex) && NameIdsAttribute.IsValidIndex(ToIndex))
+		{
+			NameIdsAttribute.ModifyAt(ToIndex, NameIdsAttribute[FromIndex]);
+			BranchIdsAttribute.ModifyAt(ToIndex, BranchIdsAttribute[FromIndex]);
+			PivotPointsAttribute.ModifyAt(ToIndex, PivotPointsAttribute[FromIndex]);
+			UpVectorsAttribute.ModifyAt(ToIndex, UpVectorsAttribute[FromIndex]);
+			NormalVectorsAttribute.ModifyAt(ToIndex, NormalVectorsAttribute[FromIndex]);
+			ScalesAttribute.ModifyAt(ToIndex, ScalesAttribute[FromIndex]);
+			LengthFromRootAttribute.ModifyAt(ToIndex, LengthFromRootAttribute[FromIndex]);
+			ParentBoneIdsAttribute.ModifyAt(ToIndex, ParentBoneIdsAttribute[FromIndex]);
+		}
+	}
+
+	void FFoliageFacade::RemoveEntries(const int32 NumEntries, const int32 StartIndex)
+	{
+		if (IsValid() && NameIdsAttribute.IsValidIndex(StartIndex) && StartIndex + NumEntries <= NameIdsAttribute.Num())
+		{
+			NameIdsAttribute.RemoveElements(NumEntries, StartIndex);
+		}
+	}
+
+	FTransform FFoliageFacade::GetFoliageTransform(int32 Id) const
+	{
+		FFoliageEntryData Data = GetFoliageEntry(Id);
+		FString FoliageName = GetFoliageName(Data.NameId);
+
+		const FQuat UpVectorRotationQuat = FQuat::FindBetweenNormals(FVector::UpVector, FVector(Data.UpVector));
+		const FVector RotatedVectorFace = UpVectorRotationQuat.RotateVector(FVector::RightVector);
+		float Dot = FVector::DotProduct(RotatedVectorFace, FVector(Data.NormalVector));
+		float Angle = FMath::Acos(FMath::Clamp(Dot, -1.0f, 1.0f));
+		const FQuat FaceRotationQuat = FQuat(FVector(Data.UpVector), Angle);
+
+		const FTransform FaceRotationTransform(FaceRotationQuat);
+		const FTransform UpVectorRotationTransform(UpVectorRotationQuat);
+		const FTransform TranslationTransform(FRotator::ZeroRotator, FVector(Data.PivotPoint));
+		const FTransform FinalTransform = UpVectorRotationTransform * FaceRotationTransform * TranslationTransform;
+
+		return FinalTransform;
+	}
+}
diff --git a/Engine/Plugins/Experimental/ProceduralVegetationEditor/Source/MegaPlants/Private/Facades/MegaPlantMetaInfoFacade.cpp b/Engine/Plugins/Experimental/ProceduralVegetationEditor/Source/MegaPlants/Private/Facades/MegaPlantMetaInfoFacade.cpp
new file mode 100644
--- /dev/null
+++ b/Engine/Plugins/Experimental/ProceduralVegetationEditor/Source/MegaPlants/Private/Facades/MegaPlantMetaInfoFacade.cpp
@@ -0,0 +1,62 @@
+﻿// Copyright Epic Games, Inc. All Rights Reserved.
+
+#include "Facades/MegaPlantMetaInfoFacade.h"
+#include "Facades/MegaPlantAttributesNames.h"
+
+namespace MegaPlants::Facades
+{
+	FMetaInfoFacade::FMetaInfoFacade(FManagedArrayCollection& InCollection, const int32 InitialSize)
+		: ConstCollection(InCollection)
+		, Collection(&InCollection)
+		, GuidAttribute(InCollection, AttributeNames::Guid, GroupNames::DetailsGroup)
+	{
+		DefineSchema(InitialSize);
+	}
+
+	FMetaInfoFacade::FMetaInfoFacade(const FManagedArrayCollection& InCollection)
+		: ConstCollection(InCollection)
+		, Collection(nullptr)
+		, GuidAttribute(InCollection, AttributeNames::Guid, GroupNames::DetailsGroup)
+	{}
+
+	void FMetaInfoFacade::DefineSchema(const int32 InitialSize)
+	{
+		check(!IsConst());
+
+		if (!Collection->HasGroup(GroupNames::DetailsGroup))
+		{
+			Collection->AddGroup(GroupNames::DetailsGroup);
+		}
+
+		GuidAttribute.Add();
+	}
+
+	bool FMetaInfoFacade::IsValid() const
+	{
+		return GuidAttribute.IsValid();
+	}
+
+	void FMetaInfoFacade::CreateGuid(const FString& InPath)
+	{
+		check(!IsConst());
+
+		int32 NumElements = Collection->NumElements(GroupNames::DetailsGroup);
+
+		if(NumElements == 0)
+		{
+			Collection->AddElements(1, GroupNames::DetailsGroup);
+		}
+		
+		GuidAttribute.Modify()[0] = FGuid::NewDeterministicGuid(InPath);
+	}
+
+	FGuid FMetaInfoFacade::GetGuid() const
+	{
+		if(GuidAttribute.IsValid() && GuidAttribute.IsValidIndex(0))
+		{
+			return GuidAttribute[0];
+		}
+
+		return FGuid::NewGuid();
+	}
+}
diff --git a/Engine/Plugins/Experimental/ProceduralVegetationEditor/Source/MegaPlants/Private/Facades/MegaPlantPointFacade.cpp b/Engine/Plugins/Experimental/ProceduralVegetationEditor/Source/MegaPlants/Private/Facades/MegaPlantPointFacade.cpp
new file mode 100644
--- /dev/null
+++ b/Engine/Plugins/Experimental/ProceduralVegetationEditor/Source/MegaPlants/Private/Facades/MegaPlantPointFacade.cpp
@@ -0,0 +1,401 @@
+﻿// Copyright Epic Games, Inc. All Rights Reserved.
+
+#include "Facades/MegaPlantPointFacade.h"
+#include "Facades/MegaPlantAttributesNames.h"
+
+namespace MegaPlants::Facades
+{
+	FPointFacade::FPointFacade(FManagedArrayCollection& InCollection)
+		: Positions(InCollection, AttributeNames::PointPosition, GroupNames::PointGroup)
+		, LengthFromRoot(InCollection, AttributeNames::LengthFromRoot, GroupNames::PointGroup)
+		, PointScaleGradient(InCollection, AttributeNames::PointScaleGradient, GroupNames::PointGroup)
+		, HullGradient(InCollection, AttributeNames::HullGradient, GroupNames::PointGroup)
+		, MainTrunkGradient(InCollection, AttributeNames::MainTrunkGradient, GroupNames::PointGroup)
+		, GroundGradient(InCollection, AttributeNames::GroundGradient, GroupNames::PointGroup)
+		, PointScale(InCollection, AttributeNames::PointScale, GroupNames::PointGroup)
+		, BudDirections(InCollection, AttributeNames::BudDirection, GroupNames::PointGroup)
+		, BranchGradients(InCollection, AttributeNames::BranchGradient, GroupNames::PointGroup)
+		, BudHormoneLevels(InCollection, AttributeNames::BudHormoneLevels, GroupNames::PointGroup)
+		, PlantGradients(InCollection, AttributeNames::PlantGradient, GroupNames::PointGroup)
+		, LengthFromSeed(InCollection, AttributeNames::LengthFromSeed, GroupNames::PointGroup)
+		, BudNumber(InCollection, AttributeNames::BudNumber, GroupNames::PointGroup)
+		, NjordPixelIndex(InCollection, AttributeNames::NjordPixelIndex, GroupNames::PointGroup)
+		, BudLightDetected(InCollection, AttributeNames::BudLightDetected, GroupNames::PointGroup)
+		, BudDevelopment(InCollection, AttributeNames::BudDevelopment, GroupNames::PointGroup)
+		, TextureCoordV(InCollection, AttributeNames::TextureCoordV, GroupNames::PointGroup)
+		, TextureCoordUOffset(InCollection, AttributeNames::TextureCoordUOffset, GroupNames::PointGroup)
+	{
+	}
+
+	FPointFacade::FPointFacade(const FManagedArrayCollection& InCollection)
+		: Positions(InCollection, AttributeNames::PointPosition, GroupNames::PointGroup)
+		, LengthFromRoot(InCollection, AttributeNames::LengthFromRoot, GroupNames::PointGroup)
+		, PointScaleGradient(InCollection, AttributeNames::PointScaleGradient, GroupNames::PointGroup)
+		, HullGradient(InCollection, AttributeNames::HullGradient, GroupNames::PointGroup)
+		, MainTrunkGradient(InCollection, AttributeNames::MainTrunkGradient, GroupNames::PointGroup)
+		, GroundGradient(InCollection, AttributeNames::GroundGradient, GroupNames::PointGroup)
+		, PointScale(InCollection, AttributeNames::PointScale, GroupNames::PointGroup)
+		, BudDirections(InCollection, AttributeNames::BudDirection, GroupNames::PointGroup)
+		, BranchGradients(InCollection, AttributeNames::BranchGradient, GroupNames::PointGroup)
+		, BudHormoneLevels(InCollection, AttributeNames::BudHormoneLevels, GroupNames::PointGroup)
+		, PlantGradients(InCollection, AttributeNames::PlantGradient, GroupNames::PointGroup)
+		, LengthFromSeed(InCollection, AttributeNames::LengthFromSeed, GroupNames::PointGroup)
+		, BudNumber(InCollection, AttributeNames::BudNumber, GroupNames::PointGroup)
+		, NjordPixelIndex(InCollection, AttributeNames::NjordPixelIndex, GroupNames::PointGroup)
+		, BudLightDetected(InCollection, AttributeNames::BudLightDetected, GroupNames::PointGroup)
+		, BudDevelopment(InCollection, AttributeNames::BudDevelopment, GroupNames::PointGroup)
+		, TextureCoordV(InCollection, AttributeNames::TextureCoordV, GroupNames::PointGroup)
+		, TextureCoordUOffset(InCollection, AttributeNames::TextureCoordUOffset, GroupNames::PointGroup)
+	{
+	}
+
+	bool FPointFacade::IsValid() const
+	{
+		return Positions.IsValid()
+			&& LengthFromRoot.IsValid()
+			&& LengthFromSeed.IsValid()
+			&& PointScaleGradient.IsValid()
+			&& HullGradient.IsValid()
+			&& MainTrunkGradient.IsValid()
+			&& GroundGradient.IsValid()
+			&& PointScale.IsValid()
+			&& BudDirections.IsValid()
+			&& BranchGradients.IsValid()
+			&& BudHormoneLevels.IsValid()
+			&& PlantGradients.IsValid()
+			&& BudNumber.IsValid()
+			&& NjordPixelIndex.IsValid()
+			&& BudLightDetected.IsValid()
+			&& BudDevelopment.IsValid();
+	}
+
+	int32 FPointFacade::GetElementCount() const
+	{
+		return Positions.Num();
+	}
+
+	const FVector3f& FPointFacade::GetPosition(const int32 Index) const
+	{
+		if (Positions.IsValid() && Positions.IsValidIndex(Index))
+		{
+			return Positions[Index];
+		}
+
+		static const FVector3f DefaultPosition = FVector3f::ZeroVector;
+		return DefaultPosition;
+	}
+
+	const float& FPointFacade::GetLengthFromRoot(const int32 Index) const
+	{
+		if (LengthFromRoot.IsValid() && LengthFromRoot.IsValidIndex(Index))
+		{
+			return LengthFromRoot[Index];
+		}
+		static constexpr float DefaultLengthFromRoot = 0.0;
+		return DefaultLengthFromRoot;
+	}
+
+	const float& FPointFacade::GetLengthFromSeed(const int32 Index) const
+	{
+		if (LengthFromSeed.IsValid() && LengthFromSeed.IsValidIndex(Index))
+		{
+			return LengthFromSeed[Index];
+		}
+		static constexpr float DefaultLengthFromSeed = 0.0;
+		return DefaultLengthFromSeed;
+	}
+
+	const float& FPointFacade::GetPointScaleGradient(const int32 Index) const
+	{
+		if (PointScaleGradient.IsValid() && PointScaleGradient.IsValidIndex(Index))
+		{
+			return PointScaleGradient[Index];
+		}
+
+		static constexpr float DefaultPScale = 0.0;
+		return DefaultPScale;
+	}
+
+	const float& FPointFacade::GetHullGradient(const int32 Index) const
+	{
+		if (HullGradient.IsValid() && HullGradient.IsValidIndex(Index))
+		{
+			return HullGradient[Index];
+		}
+
+		static constexpr float DefaultValue = 0.0;
+		return DefaultValue;
+	}
+
+	const float& FPointFacade::GetMainTrunkGradient(const int32 Index) const
+	{
+		if (MainTrunkGradient.IsValid() && MainTrunkGradient.IsValidIndex(Index))
+		{
+			return MainTrunkGradient[Index];
+		}
+
+		static constexpr float DefaultValue = 0.0;
+		return DefaultValue;
+	}
+
+	const float& FPointFacade::GetGroundGradient(const int32 Index) const
+	{
+		if (GroundGradient.IsValid() && GroundGradient.IsValidIndex(Index))
+		{
+			return GroundGradient[Index];
+		}
+
+		static constexpr float DefaultValue = 0.0;
+		return DefaultValue;
+	}
+
+	const float& FPointFacade::GetPointScale(const int32 Index) const
+	{
+		if (PointScale.IsValid() && PointScale.IsValidIndex(Index))
+		{
+			return PointScale[Index];
+		}
+
+		static constexpr float DefaultPScale = 0.0;
+		return DefaultPScale;
+	}
+
+	const int32 FPointFacade::GetBudNumber(const int32 Index) const
+	{
+		if (BudNumber.IsValid() && BudNumber.IsValidIndex(Index))
+		{
+			return BudNumber[Index];
+		}
+
+		return INDEX_NONE;
+	}
+	
+	const TArray<float>& FPointFacade::GetBudLightDetected(const int32 Index) const
+	{
+		if (BudLightDetected.IsValid() && BudLightDetected.IsValidIndex(Index))
+		{
+			return BudLightDetected[Index];
+		}
+
+		static const TArray<float> EmptyArray;
+		return EmptyArray;
+	}
+
+	const TArray<int>& FPointFacade::GetBudDevelopment(const int32 Index) const
+	{
+		if (BudDevelopment.IsValid() && BudDevelopment.IsValidIndex(Index))
+		{
+			return BudDevelopment[Index];
+		}
+
+		static const TArray<int> EmptyArray;
+		return EmptyArray;
+	}
+
+	const TArray<FVector3f>& FPointFacade::GetBudDirection(const int32 Index) const
+	{
+		if (BudDirections.IsValid() && BudDirections.IsValidIndex(Index))
+		{
+			return BudDirections[Index];
+		}
+
+		static const TArray<FVector3f> EmptyArray;
+		return EmptyArray;
+	}
+
+	const TArray<float>& FPointFacade::GetBudHormoneLevels(const int32 Index) const
+	{
+		if (BudHormoneLevels.IsValid() && BudHormoneLevels.IsValidIndex(Index))
+		{
+			return BudHormoneLevels[Index];
+		}
+
+		static const TArray<float> EmptyArray;
+		return EmptyArray;
+	}
+
+	void FPointFacade::SetBudHormoneLevels(const int32 Index, const TArray<float>& InBudHormoneLevels)
+	{
+		if (BudHormoneLevels.IsValid() && BudHormoneLevels.IsValidIndex(Index))
+		{
+			BudHormoneLevels.ModifyAt(Index, InBudHormoneLevels);
+		}
+	}
+
+	const float& FPointFacade::GetBranchGradient(const int32 Index) const
+	{
+		if (BranchGradients.IsValid() && BranchGradients.IsValidIndex(Index))
+		{
+			return BranchGradients[Index];
+		}
+
+		static constexpr float DefaultBranchGradient = 0.0;
+		return DefaultBranchGradient;
+	}
+
+	void FPointFacade::SetBranchGradient(const int32 Index, const float InBranchGradient)
+	{
+		if (BranchGradients.IsValid() && BranchGradients.IsValidIndex(Index))
+		{
+			BranchGradients.ModifyAt(Index, InBranchGradient);
+		}
+	}
+
+	const float& FPointFacade::GetPlantGradient(const int32 Index) const
+	{
+		if (PlantGradients.IsValid() && PlantGradients.IsValidIndex(Index))
+		{
+			return PlantGradients[Index];
+		}
+
+		static constexpr float DefaultPlantGradient = 0.0;
+		return DefaultPlantGradient;
+	}
+
+	const float& FPointFacade::GetNjordPixelIndex(const int32 Index) const
+	{
+		if (NjordPixelIndex.IsValid() && NjordPixelIndex.IsValidIndex(Index))
+		{
+			return NjordPixelIndex[Index];
+		}
+
+		static constexpr float IndexNone = -1.0;
+		return IndexNone;
+	}
+
+	void FPointFacade::SetPlantGradient(const int32 Index, const float InPlantGradient)
+	{
+		if (PlantGradients.IsValid() && PlantGradients.IsValidIndex(Index))
+		{
+			PlantGradients.ModifyAt(Index, InPlantGradient);
+		}
+	}
+
+	const float& FPointFacade::GetTextureCoordV(const int32 Index) const
+	{
+		if (TextureCoordV.IsValid() && TextureCoordV.IsValidIndex(Index))
+		{
+			return TextureCoordV[Index];
+		}
+
+		static constexpr float IndexNone = -1.0;
+		return IndexNone;
+	}
+
+	void FPointFacade::SeTextureCoordV(const int32 Index, const float InTextureCoordV)
+	{
+		if (!TextureCoordV.IsValid())
+		{
+			TextureCoordV.Add();
+		}
+		
+		if (TextureCoordV.IsValid() && TextureCoordV.IsValidIndex(Index))
+		{
+			TextureCoordV.ModifyAt(Index, InTextureCoordV);
+		}
+	}
+
+	const float& FPointFacade::GetTextureCoordUOffset(const int32 Index) const
+	{
+		if (TextureCoordUOffset.IsValid() && TextureCoordUOffset.IsValidIndex(Index))
+		{
+			return TextureCoordUOffset[Index];
+		}
+
+		static constexpr float IndexNone = -1.0;
+		return IndexNone;
+	}
+
+	void FPointFacade::SetTextureCoordUOffset(const int32 Index, const float InTextureCoordUOffset)
+	{
+		if (!TextureCoordUOffset.IsValid())
+		{
+			TextureCoordUOffset.Add();
+		}
+		
+		if (TextureCoordUOffset.IsValid() && TextureCoordUOffset.IsValidIndex(Index))
+		{
+			TextureCoordUOffset.ModifyAt(Index, InTextureCoordUOffset);
+		}
+	}
+
+	TManagedArray<FVector3f>& FPointFacade::ModifyPositions()
+	{
+		return Positions.Modify();
+	}
+
+	TManagedArray<float>& FPointFacade::ModifyPointScales()
+	{
+		return PointScale.Modify();
+	}
+
+	TManagedArray<float>& FPointFacade::ModifyLengthFromRoots()
+	{
+		return LengthFromRoot.Modify();
+	}
+
+	TManagedArray<float>& FPointFacade::ModifyLengthFromSeeds()
+	{
+		return LengthFromSeed.Modify();
+	}
+
+	TManagedArray<float>& FPointFacade::ModifyNjordPixelIDs()
+	{
+		return NjordPixelIndex.Modify();
+	}
+
+	const TManagedArray<FVector3f>& FPointFacade::GetPositions() const
+	{
+		return Positions.Get();
+	}
+
+	const TManagedArray<float>& FPointFacade::GetPointScales() const
+	{
+		return PointScale.Get();
+	}
+
+	const TManagedArray<float>& FPointFacade::GetLengthFromRootsArray() const
+	{
+		return LengthFromRoot.Get();
+	}
+
+	void FPointFacade::CopyEntry(const int32 FromIndex, const int32 ToIndex)
+	{
+		if (IsValid() && Positions.IsValidIndex(FromIndex) && Positions.IsValidIndex(ToIndex))
+		{
+			Positions.ModifyAt(ToIndex, Positions[FromIndex]);
+			LengthFromRoot.ModifyAt(ToIndex, LengthFromRoot[FromIndex]);
+			PointScaleGradient.ModifyAt(ToIndex, PointScaleGradient[FromIndex]);
+			HullGradient.ModifyAt(ToIndex, HullGradient[FromIndex]);
+			MainTrunkGradient.ModifyAt(ToIndex, MainTrunkGradient[FromIndex]);
+			GroundGradient.ModifyAt(ToIndex, GroundGradient[FromIndex]);
+			PointScale.ModifyAt(ToIndex, PointScale[FromIndex]);
+			BudDirections.ModifyAt(ToIndex, BudDirections[FromIndex]);
+			BranchGradients.ModifyAt(ToIndex, BranchGradients[FromIndex]);
+			BudHormoneLevels.ModifyAt(ToIndex, BudHormoneLevels[FromIndex]);
+			PlantGradients.ModifyAt(ToIndex, PlantGradients[FromIndex]);
+			LengthFromSeed.ModifyAt(ToIndex, LengthFromSeed[FromIndex]);
+			BudNumber.ModifyAt(ToIndex, BudNumber[FromIndex]);
+			NjordPixelIndex.ModifyAt(ToIndex, NjordPixelIndex[FromIndex]);
+			BudLightDetected.ModifyAt(ToIndex, BudLightDetected[FromIndex]);
+			BudDevelopment.ModifyAt(ToIndex, BudDevelopment[FromIndex]);
+			
+			if(TextureCoordV.IsValid())
+			{
+				TextureCoordV.ModifyAt(ToIndex, TextureCoordV[FromIndex]);
+			}
+			if(TextureCoordUOffset.IsValid())
+			{
+				TextureCoordUOffset.ModifyAt(ToIndex, TextureCoordUOffset[FromIndex]);
+			}
+		}
+	}
+
+	void FPointFacade::RemoveEntries(const int32 NumEntries, const int32 StartIndex)
+	{
+		if (IsValid() && Positions.IsValidIndex(StartIndex) && StartIndex + NumEntries <= Positions.Num())
+		{
+			Positions.RemoveElements(NumEntries, StartIndex);
+		}
+	}
+}
diff --git a/Engine/Plugins/Experimental/ProceduralVegetationEditor/Source/MegaPlants/Private/Facades/MegaPlantProfileFacade.cpp b/Engine/Plugins/Experimental/ProceduralVegetationEditor/Source/MegaPlants/Private/Facades/MegaPlantProfileFacade.cpp
new file mode 100644
--- /dev/null
+++ b/Engine/Plugins/Experimental/ProceduralVegetationEditor/Source/MegaPlants/Private/Facades/MegaPlantProfileFacade.cpp
@@ -0,0 +1,58 @@
+﻿// Copyright Epic Games, Inc. All Rights Reserved.
+
+#include "Facades/MegaPlantProfileFacade.h"
+#include "Facades/MegaPlantAttributesNames.h"
+
+namespace MegaPlants::Facades
+{
+	FPlantProfileFacade::FPlantProfileFacade(FManagedArrayCollection& InCollection)
+		: ConstCollection(InCollection)
+		, Collection(&InCollection)
+		, ProfilePointsAttribute(InCollection, AttributeNames::ProfilePoints, GroupNames::PlantProfilesGroup)
+	{
+		DefineSchema();
+	}
+
+	FPlantProfileFacade::FPlantProfileFacade(const FManagedArrayCollection& InCollection)
+		: ConstCollection(InCollection)
+		, Collection(nullptr)
+		, ProfilePointsAttribute(InCollection, AttributeNames::ProfilePoints, GroupNames::PlantProfilesGroup)
+	{
+	}
+
+	int32 FPlantProfileFacade::AddProfileEntry(const TArray<float>& InProfilePoints)
+	{
+		const int32 Index = ProfilePointsAttribute.AddElements(1);
+		ProfilePointsAttribute.ModifyAt(Index, InProfilePoints);
+
+		return Index;
+	}
+
+	const TArray<float>& FPlantProfileFacade::GetProfilePoints(const int32 Index) const
+	{
+		if (ProfilePointsAttribute.IsValid() && ProfilePointsAttribute.IsValidIndex(Index))
+		{
+			return ProfilePointsAttribute[Index];
+		}
+
+		static const TArray<float> EmptyArray;
+		return EmptyArray;
+	}
+
+	void FPlantProfileFacade::DefineSchema()
+	{
+		check(!IsConst());
+
+		if (!Collection->HasGroup(GroupNames::PlantProfilesGroup))
+		{
+			Collection->AddGroup(GroupNames::PlantProfilesGroup);
+		}
+
+		ProfilePointsAttribute.Add();
+	}
+
+	bool FPlantProfileFacade::IsValid() const
+	{
+		return ProfilePointsAttribute.IsValid();
+	}
+}
diff --git a/Engine/Plugins/Experimental/ProceduralVegetationEditor/Source/MegaPlants/Private/Facades/MegaPlantRenderingFacade.cpp b/Engine/Plugins/Experimental/ProceduralVegetationEditor/Source/MegaPlants/Private/Facades/MegaPlantRenderingFacade.cpp
new file mode 100644
--- /dev/null
+++ b/Engine/Plugins/Experimental/ProceduralVegetationEditor/Source/MegaPlants/Private/Facades/MegaPlantRenderingFacade.cpp
@@ -0,0 +1,36 @@
+﻿// Copyright Epic Games, Inc. All Rights Reserved.
+
+#include "Facades/MegaPlantRenderingFacade.h"
+
+namespace MegaPlants::Facades
+{
+	FMegaPlantRenderingFacade::FMegaPlantRenderingFacade(FManagedArrayCollection& InSelf)
+		: FRenderingFacade(InSelf)
+		,ConstMegaPlantCollection(InSelf)
+		,MegaPlantCollection(&InSelf)
+		,FoliageFacade(InSelf)
+	{}
+
+	FMegaPlantRenderingFacade::FMegaPlantRenderingFacade(const FManagedArrayCollection& InSelf)
+		: FRenderingFacade(InSelf)
+		,ConstMegaPlantCollection(InSelf)
+		,MegaPlantCollection(nullptr)
+		,FoliageFacade(InSelf)
+	{}
+
+	void FMegaPlantRenderingFacade::DefineSchema()
+	{
+		FRenderingFacade::DefineSchema();
+	}
+
+	const FFoliageFacade& FMegaPlantRenderingFacade::GetFoliageFacade() const
+	{
+		return FoliageFacade;
+	}
+
+	void FMegaPlantRenderingFacade::CopyFrom(const FManagedArrayCollection& Other)
+	{
+		Other.CopyTo(MegaPlantCollection);
+	}
+}
+
diff --git a/Engine/Plugins/Experimental/ProceduralVegetationEditor/Source/MegaPlants/Private/Facades/MegaPlantTreeFacade.cpp b/Engine/Plugins/Experimental/ProceduralVegetationEditor/Source/MegaPlants/Private/Facades/MegaPlantTreeFacade.cpp
new file mode 100644
--- /dev/null
+++ b/Engine/Plugins/Experimental/ProceduralVegetationEditor/Source/MegaPlants/Private/Facades/MegaPlantTreeFacade.cpp
@@ -0,0 +1,153 @@
+﻿// Copyright Epic Games, Inc. All Rights Reserved.
+
+#include "MegaPlantTreeFacade.h"
+#include "Facades/MegaPlantBranchFacade.h"
+#include "Facades/MegaPlantFoliageFacade.h"
+#include "Facades/MegaPlantPointFacade.h"
+#include "GeometryCollection/ManagedArrayAccessor.h"
+
+namespace MegaPlants::Facades
+{
+	FRemoveEntriesResult FTreeFacade::RemoveEntriesAndReIndexAttributes(FManagedArrayCollection& OutCollection, TArray<bool>& PointsToRemove,
+		TArray<bool>& BranchesToRemove, TArray<bool>& FoliageInstancesToRemove)
+	{
+		TRACE_CPUPROFILER_EVENT_SCOPE(MegaPlants::Facades::FTreeFacade::RemoveEntriesAndReIndexAttributes);
+
+		FBranchFacade BranchFacadeOut(OutCollection);
+		FPointFacade PointFacadeOut(OutCollection);
+		FFoliageFacade FoliageFacadeOut(OutCollection);
+
+		TMap<int32, int32> PointsOldIDsToNewIDs = IShrinkable::RemoveEntries(PointFacadeOut, PointsToRemove);
+		TMap<int32, int32> BranchesOldIDsToNewIDs = IShrinkable::RemoveEntries(BranchFacadeOut, BranchesToRemove);
+		TMap<int32, int32> FoliageInstancesOldIDsToNewIDs = IShrinkable::RemoveEntries(FoliageFacadeOut, FoliageInstancesToRemove);
+
+		const int32 NumOfBranches = BranchFacadeOut.GetElementCount();
+		TSet<int32> CurrentBranchNumbers;
+		CurrentBranchNumbers.Reserve(NumOfBranches);
+		for (int32 BranchIndex = 0; BranchIndex < NumOfBranches; ++BranchIndex)
+		{
+			CurrentBranchNumbers.Add(BranchFacadeOut.GetBranchNumber(BranchIndex));
+		}
+
+		for (int32 BranchIndex = 0; BranchIndex < NumOfBranches; ++BranchIndex)
+		{
+			TArray<int32> UpdatedPoints;
+			const TArray<int32>& BranchPoints = BranchFacadeOut.GetPoints(BranchIndex);
+			for (const int32& PointIndex : BranchPoints)
+			{
+				if (PointsOldIDsToNewIDs.Contains(PointIndex))
+				{
+					UpdatedPoints.Add(PointsOldIDsToNewIDs[PointIndex]);
+				}
+				else if (!PointsToRemove[PointIndex])
+				{
+					UpdatedPoints.Add(PointIndex);
+				}
+			}
+			BranchFacadeOut.SetPoints(BranchIndex, UpdatedPoints);
+
+			TArray<int32> UpdatedParents;
+			const TArray<int32>& BranchParents = BranchFacadeOut.GetParents(BranchIndex);
+			for (const int32& BranchNumber : BranchParents)
+			{
+				if (BranchNumber == 0 || CurrentBranchNumbers.Contains(BranchNumber))
+				{
+					UpdatedParents.Add(BranchNumber);
+				}
+			}
+			BranchFacadeOut.SetParents(BranchIndex, UpdatedParents);
+
+			TArray<int32> UpdatedChildren;
+			const TArray<int32>& BranchChildren = BranchFacadeOut.GetChildren(BranchIndex);
+			for (const int32& BranchNumber : BranchChildren)
+			{
+				if (CurrentBranchNumbers.Contains(BranchNumber))
+				{
+					UpdatedChildren.Add(BranchNumber);
+				}
+			}
+			BranchFacadeOut.SetChildren(BranchIndex, UpdatedChildren);
+
+			TArray<int32> UpdatedFoliageInstances;
+			const TArray<int32>& FoliageInstances = FoliageFacadeOut.GetFoliageEntryIdsForBranch(BranchIndex);
+			for (const int32& FoliageInstanceIndex : FoliageInstances)
+			{
+				if (FoliageInstancesOldIDsToNewIDs.Contains(FoliageInstanceIndex))
+				{
+					const int32 NewFoliageInstanceID = FoliageInstancesOldIDsToNewIDs[FoliageInstanceIndex];
+					UpdatedFoliageInstances.Add(NewFoliageInstanceID);
+					FoliageFacadeOut.SetFoliageBranchId(NewFoliageInstanceID, BranchIndex);
+				}
+				else if (!FoliageInstancesToRemove[FoliageInstanceIndex])
+				{
+					UpdatedFoliageInstances.Add(FoliageInstanceIndex);
+				}
+			}
+			FoliageFacadeOut.SetFoliageIdsArray(BranchIndex, UpdatedFoliageInstances);
+		}
+
+		FRemoveEntriesResult Result = {
+			.PointsOldIDsToNewIDs = MoveTemp(PointsOldIDsToNewIDs),
+			.BranchesOldIDsToNewIDs = MoveTemp(BranchesOldIDsToNewIDs),
+			.FoliageInstancesOldIDsToNewIDs = MoveTemp(FoliageInstancesOldIDsToNewIDs)
+		};
+
+		return Result;
+	}
+
+	void FTreeFacade::RemoveBranches(const FBranchFacade& InBranchFacade, TArray<int>& InBranchesToRemove, FManagedArrayCollection& OutCollection)
+	{
+		FBranchFacade BranchFacadeOut(OutCollection);
+		FPointFacade PointFacadeOut(OutCollection);
+		FFoliageFacade FoliageFacadeOut(OutCollection);
+
+		for (int32 i = InBranchesToRemove.Num() - 1; i > 0 ; i--)
+		{
+			GatherChildBranches(InBranchFacade, InBranchesToRemove[i], InBranchesToRemove);
+		}
+			
+		TArray<bool> BranchesToRemove;
+		BranchesToRemove.Init(false, BranchFacadeOut.GetElementCount());
+
+		TArray<bool> PointsToRemove;
+		PointsToRemove.Init(false, PointFacadeOut.GetElementCount());
+	
+		TArray<bool> FoliageInstancesToRemove;
+		FoliageInstancesToRemove.Init(false, FoliageFacadeOut.GetElementCount());
+	
+		for (int32 i = 0; i < InBranchesToRemove.Num(); i++)
+		{
+			BranchesToRemove[InBranchesToRemove[i]] = true;
+		
+			for (const int32& P : BranchFacadeOut.GetPoints(InBranchesToRemove[i]))
+			{
+				PointsToRemove[P] = true;
+			}
+			for (const int32& FId : FoliageFacadeOut.GetFoliageEntryIdsForBranch(InBranchesToRemove[i]))
+			{
+				FoliageInstancesToRemove[FId] = true;
+			}	
+		}
+
+		FRemoveEntriesResult RemoveEntriesResult = RemoveEntriesAndReIndexAttributes(
+			OutCollection, PointsToRemove, BranchesToRemove, FoliageInstancesToRemove);
+	}
+
+	void FTreeFacade::GatherChildBranches(const FBranchFacade& InBranchFacade, const int InParentIndex, TArray<int32>& OutBranchesToRemove)
+	{
+		const TArray<int32>& BranchChildren = InBranchFacade.GetChildren(InParentIndex);
+		const TManagedArray<int32>& BranchNumbers = InBranchFacade.GetBranchNumbers();
+		
+		for (const int32& BranchNumber : BranchChildren)
+		{
+			const int32 BranchIndex = BranchNumbers.Find(BranchNumber);
+			
+			if (BranchIndex != INDEX_NONE && !OutBranchesToRemove.Contains(BranchIndex))
+			{
+				OutBranchesToRemove.Add(BranchIndex);
+			}
+
+			GatherChildBranches(InBranchFacade, BranchIndex, OutBranchesToRemove);
+		}
+	}
+}
diff --git a/Engine/Plugins/Experimental/ProceduralVegetationEditor/Source/MegaPlants/Private/Facades/MegaPlantTreeFacade.h b/Engine/Plugins/Experimental/ProceduralVegetationEditor/Source/MegaPlants/Private/Facades/MegaPlantTreeFacade.h
new file mode 100644
--- /dev/null
+++ b/Engine/Plugins/Experimental/ProceduralVegetationEditor/Source/MegaPlants/Private/Facades/MegaPlantTreeFacade.h
@@ -0,0 +1,37 @@
+// Copyright Epic Games, Inc. All Rights Reserved.
+
+#pragma once
+
+#include "CoreMinimal.h"
+
+struct FManagedArrayCollection;
+
+namespace MegaPlants::Facades
+{
+	struct FRemoveEntriesResult
+	{
+		TMap<int32, int32> PointsOldIDsToNewIDs;
+		TMap<int32, int32> BranchesOldIDsToNewIDs;
+		TMap<int32, int32> FoliageInstancesOldIDsToNewIDs;
+	};
+
+	class FBranchFacade;
+
+	/**
+	 * FTreeFacade is used to remove branches, points, and foliage instances from a tree collection.
+	 * It does not validate whether the input data is consistent i.e. the points actually belong to the
+	 * same branches that are asking to be removed etc.
+	 */
+	class MEGAPLANTS_API FTreeFacade
+	{
+	public:
+		static FRemoveEntriesResult RemoveEntriesAndReIndexAttributes(FManagedArrayCollection& OutCollection, TArray<bool>& PointsToRemove,
+		                                                              TArray<bool>& BranchesToRemove, TArray<bool>& FoliageInstancesToRemove);
+
+		static void RemoveBranches(const FBranchFacade& InBranchFacade, TArray<int>& InBranchesToRemove, FManagedArrayCollection& OutCollection);
+		
+
+	private:
+		static void GatherChildBranches(const FBranchFacade& InBranchFacade, const int InParentIndex, TArray<int32>& OutBranchesToRemove);
+	};
+}
diff --git a/Engine/Plugins/Experimental/ProceduralVegetationEditor/Source/MegaPlants/Private/Helpers/MegaPlantExportHelper.cpp b/Engine/Plugins/Experimental/ProceduralVegetationEditor/Source/MegaPlants/Private/Helpers/MegaPlantExportHelper.cpp
new file mode 100644
--- /dev/null
+++ b/Engine/Plugins/Experimental/ProceduralVegetationEditor/Source/MegaPlants/Private/Helpers/MegaPlantExportHelper.cpp
@@ -0,0 +1,938 @@
+﻿// Copyright Epic Games, Inc. All Rights Reserved.
+
+#include "Helpers/MegaPlantExportHelper.h"
+
+#include "DynamicMeshEditor.h"
+#include "MegaPlantLink.h"
+#include "MegaPlantsModule.h"
+#include "MeshDescriptionToDynamicMesh.h"
+#include "PlanarCut.h"
+#include "RenderUtils.h"
+#include "StaticMeshAttributes.h"
+#include "Animation/Skeleton.h"
+#include "AssetRegistry/IAssetRegistry.h"
+#include "Components/InstancedStaticMeshComponent.h"
+#include "ConversionUtils/SceneComponentToDynamicMesh.h"
+#include "DynamicMesh/DynamicBoneAttribute.h"
+#include "DynamicMesh/DynamicVertexSkinWeightsAttribute.h"
+
+#include "Engine/SkeletalMesh.h"
+#include "Engine/StaticMesh.h"
+#include "Engine/World.h"
+#include "Facades/MegaPlantBranchFacade.h"
+#include "Facades/MegaPlantPointFacade.h"
+#include "GeometryCollection/GeometryCollection.h"
+#include "GeometryScript/GeometryScriptTypes.h"
+#include "GeometryScript/MeshAssetFunctions.h"
+#include "GeometryScript/MeshBoneWeightFunctions.h"
+#include "Helpers/MegaPlantUtilities.h"
+#include "Misc/Paths.h"
+#include "UObject/Package.h"
+#include "MegaPlant.h"
+#include "GeometryCollection/Facades/CollectionConstraintOverrideFacade.h"
+#include "Rendering/RenderCommandPipes.h"
+
+namespace MegaPlants::Export::Internal
+{
+	TObjectPtr<UDynamicMesh> CollectionToDynamicMesh(const FManagedArrayCollection& Collection)
+	{
+		TRACE_CPUPROFILER_EVENT_SCOPE(MegaPlants::Export::CollectionToDynamicMesh);
+		
+		TObjectPtr<UDynamicMesh> NewMesh = NewObject<UDynamicMesh>();
+
+		if (Collection.NumElements(FGeometryCollection::TransformGroup) > 0)
+		{
+			if (const TUniquePtr<FGeometryCollection> GeomCollection = TUniquePtr<FGeometryCollection>(Collection.NewCopy<FGeometryCollection>()))
+			{
+				const TManagedArray<FTransform3f>& BoneTransforms = Collection.GetAttribute<FTransform3f>(
+					"Transform", FGeometryCollection::TransformGroup);
+
+				TArray<int32> TransformIndices;
+				TransformIndices.AddUninitialized(BoneTransforms.Num());
+
+				int32 Idx = 0;
+				for (int32& TransformIdx : TransformIndices)
+				{
+					TransformIdx = Idx++;
+				}
+
+				FMeshDescription MeshDescription;
+				FStaticMeshAttributes Attributes(MeshDescription);
+				Attributes.Register();
+
+				FTransform TransformOut;
+
+				constexpr bool bCenterPivot = false;
+				ConvertToMeshDescription(MeshDescription, TransformOut, bCenterPivot, *GeomCollection, BoneTransforms, TransformIndices);
+
+				NewMesh->Reset();
+
+				UE::Geometry::FDynamicMesh3& DynMesh = NewMesh->GetMeshRef();
+				{
+					FMeshDescriptionToDynamicMesh ConverterToDynamicMesh;
+					ConverterToDynamicMesh.Convert(&MeshDescription, DynMesh);
+				}
+			}
+		}
+
+		return NewMesh;
+	}
+
+	FGeometryScriptCopyMeshToAssetOptions GetCopyMeshToAssetOptions(TArray<TObjectPtr<UMaterialInterface>> InMaterials)
+	{
+		FGeometryScriptCopyMeshToAssetOptions Options;
+		Options.bEnableRecomputeTangents = true;
+
+		Options.bReplaceMaterials = true;
+		Options.NewMaterials = InMaterials;
+
+		return Options;
+	}
+
+#if WITH_EDITORONLY_DATA
+	FGeometryScriptCopyMeshToAssetOptions GetCopyMeshToAssetOptions(TArray<TObjectPtr<UMaterialInterface>> InMaterials, const FMeshNaniteSettings& InNaniteSettings, const ENaniteShapePreservation InShapePreservation)
+	{
+		FGeometryScriptCopyMeshToAssetOptions Options = GetCopyMeshToAssetOptions(InMaterials);
+		Options.NewNaniteSettings = InNaniteSettings;
+		Options.NewNaniteSettings.bEnabled = true;
+		Options.NewNaniteSettings.ShapePreservation = InShapePreservation;
+		Options.bApplyNaniteSettings = true;
+
+		return Options;
+	}
+#endif
+
+	FString GetFoliageName(const Facades::FFoliageFacade& FoliageFacade, int32 Id)
+	{
+		Facades::FFoliageEntryData Data = FoliageFacade.GetFoliageEntry(Id);
+		return FoliageFacade.GetFoliageName(Data.NameId);
+	}
+
+	TArray<TObjectPtr<UMaterialInterface>> GetMaterialsFromCollection(const FManagedArrayCollection& Collection)
+	{
+		TArray<TObjectPtr<UMaterialInterface>> Materials;
+		if (Collection.HasAttribute("MaterialPath", FGeometryCollection::MaterialGroup))
+		{
+			const TManagedArray<FString>& MaterialPaths = Collection.GetAttribute<FString>("MaterialPath", FGeometryCollection::MaterialGroup);
+			for (const FString& Path : MaterialPaths)
+			{
+				Materials.Add(LoadObject<UMaterialInterface>(nullptr, *Path));
+			}
+		}
+		return Materials;
+	}
+
+	static bool SkeletalMeshToDynamicMesh(const USkeletalMesh* SkeletalMesh, const int32 LodIndex, FDynamicMesh3& ToDynamicMesh)
+	{
+#if WITH_EDITOR
+		if (SkeletalMesh->HasMeshDescription(LodIndex))
+		{
+			const FMeshDescription* SourceMesh = SkeletalMesh->GetMeshDescription(LodIndex);
+			if (!SourceMesh)
+			{
+				return false;
+			}
+
+			FMeshDescriptionToDynamicMesh Converter;
+			Converter.Convert(SourceMesh, ToDynamicMesh);
+			return true;
+		}
+#else
+		UE_LOG(LogMegaPlants, Warning, TEXT("SkeletalMeshToDynamicMesh is only supported in editor."));
+#endif
+		return false;
+	}
+
+	void CombineMeshInstancesToDynamicMesh(
+		const FManagedArrayCollection& Collection,
+		UDynamicMesh* DynamicMesh,
+		TArray<TObjectPtr<UMaterialInterface>>& Materials,
+		const FMeshInstanceCombined& OnMeshInstanceCombined
+	)
+	{
+		Facades::FFoliageFacade FoliageFacade(Collection);
+		int32 NumInstances = FoliageFacade.NumFoliageEntries();
+		
+		if(NumInstances > 0)
+		{
+			TMap<FString, FDynamicMesh3> FoliageDynamicMeshes;
+			TMap<FString, TObjectPtr<UStaticMesh>> FoliageStaticMeshes;
+			TMap<FString, TObjectPtr<USkeletalMesh>> FoliageSkeletalMeshes;
+			
+			for(int32 FoliageNameIndex = 0; FoliageNameIndex < FoliageFacade.NumFoliageNames(); ++FoliageNameIndex)
+			{
+				FString FoliageName = FoliageFacade.GetFoliageName(FoliageNameIndex);
+				const TObjectPtr<UStaticMesh> StaticMesh = LoadObject<UStaticMesh>(nullptr, *FoliageName);
+				const TObjectPtr<USkeletalMesh> SkeletalMesh = LoadObject<USkeletalMesh>(nullptr, *FoliageName);
+
+				FDynamicMesh3 FoliageDynamicMesh;
+				FText OutErrorMessage;
+
+				if(StaticMesh)
+				{
+					FoliageStaticMeshes.FindOrAdd(FoliageName) = StaticMesh;
+					
+					for (auto Material : StaticMesh->GetStaticMaterials())
+					{
+						Materials.AddUnique(Material.MaterialInterface);
+					}
+
+					UE::Conversion::FStaticMeshConversionOptions Options;
+					if(StaticMeshToDynamicMesh(StaticMesh, FoliageDynamicMesh, OutErrorMessage, Options))
+					{
+						FoliageDynamicMeshes.Add(FoliageName, MoveTemp(FoliageDynamicMesh));
+					}
+				}
+				else if (SkeletalMesh)
+				{
+					FoliageSkeletalMeshes.FindOrAdd(FoliageName) = SkeletalMesh;
+					
+					for (auto Material : SkeletalMesh->GetMaterials())
+					{
+						Materials.AddUnique(Material.MaterialInterface);
+					}
+
+					if(SkeletalMeshToDynamicMesh(SkeletalMesh, 0, FoliageDynamicMesh))
+					{
+						FoliageDynamicMesh.Attributes()->RemoveAttribute("BoneIndex");
+						FoliageDynamicMesh.Attributes()->RemoveAttribute("BoneParentIndex");
+						FoliageDynamicMeshes.Add(FoliageName, MoveTemp(FoliageDynamicMesh));
+					}
+				}
+			}
+
+			for(int32 Id = 0; Id < NumInstances; Id++)
+			{
+				FString FoliageName = GetFoliageName(FoliageFacade, Id);
+
+				if(FoliageDynamicMeshes.Find(FoliageName))
+				{
+					FDynamicMesh3& FoliageDynamicMesh = FoliageDynamicMeshes[FoliageName];
+
+					const TObjectPtr<UStaticMesh> StaticMesh = FoliageStaticMeshes.Contains(FoliageName) ? FoliageStaticMeshes[FoliageName] : nullptr;
+					const TObjectPtr<USkeletalMesh> SkeletalMesh = FoliageSkeletalMeshes.Contains(FoliageName) ? FoliageSkeletalMeshes[FoliageName] : nullptr;
+
+					if (SkeletalMesh)
+					{
+						int32 FoliageParentBoneID = FoliageFacade.GetParentBoneID(Id);
+
+						//Reparent root bone according to the instanced transform
+						UE::Geometry::FDynamicMeshAttributeSet* MeshAttributes = FoliageDynamicMesh.Attributes();
+						UE::Geometry::FDynamicMeshBoneParentIndexAttribute* BoneParentIndices = MeshAttributes->GetBoneParentIndices();
+						BoneParentIndices->SetValue(0,FoliageParentBoneID);
+					}
+				
+					FTransform Transform = FoliageFacade.GetFoliageTransform(Id);
+					UE::Geometry::FTransformSRT3d GeoTransform(Transform);
+
+					DynamicMesh->EditMesh([&](FDynamicMesh3& Mesh)
+					{
+						int32 VertexStart = Mesh.VertexCount();
+						UE::Geometry::FDynamicMeshEditor Editor = UE::Geometry::FDynamicMeshEditor(&Mesh);
+
+						UE::Geometry::FMeshIndexMappings IndexMap;
+						Editor.AppendMesh(
+							&FoliageDynamicMesh,
+							IndexMap,
+							[&](int, const FVector3d& Position)
+							{
+								return GeoTransform.TransformPosition(Position);
+							},
+							[&](int, const FVector3d& Normal)
+							{
+								return GeoTransform.TransformVectorNoScale(Normal);
+							}
+						);
+
+						UE::Geometry::FIndexMapi& TriangleMap = IndexMap.GetTriangleMap();
+							
+						//Remap Material IDs
+						if (Mesh.Attributes()->HasMaterialID())
+						{
+							UE::Geometry::FDynamicMeshMaterialAttribute* MaterialIDs = Mesh.Attributes()->GetMaterialID();
+
+							for (const auto& [OldTid, NewTid] : TriangleMap.GetForwardMap())
+							{
+								const int32 MaterialID = MaterialIDs->GetValue(OldTid);
+								if (StaticMesh)
+								{
+									UMaterialInterface* Material = StaticMesh->GetMaterial(MaterialID);
+									const int32 NewMaterialID = Materials.Find(Material);
+									MaterialIDs->SetValue(NewTid, NewMaterialID);
+								}
+								if (SkeletalMesh)
+								{
+									UMaterialInterface* Material = SkeletalMesh->GetMaterials()[MaterialID].MaterialInterface;
+									const int32 NewMaterialID = Materials.Find(Material);
+									MaterialIDs->SetValue(NewTid, NewMaterialID);
+								}
+							}
+						}
+
+						OnMeshInstanceCombined(StaticMesh->GetFullName(), Id, VertexStart, Mesh.VertexCount());
+					});
+				}
+			}
+		}
+	}
+
+#if WITH_EDITOR
+	void BuildNaniteAssemblyData(
+		const FManagedArrayCollection& Collection,
+		UStaticMesh* StaticMesh
+	)
+	{
+		//Create NaniteAssemblyDataBuilder and the main static mesh materials to it
+		FNaniteAssemblyDataBuilder AssemblyBuilder;
+		for (const FStaticMaterial& Material : StaticMesh->GetStaticMaterials())
+		{
+			AssemblyBuilder.AddMaterialSlot(Material.MaterialInterface);
+		}
+
+		Facades::FFoliageFacade FoliageFacade(Collection);
+		
+		TMap<FString, int32> MeshNamePartMap;
+		TMap<FString, TArray<TObjectPtr<UMaterialInterface>>> MeshMaterialsMap;
+		
+		for(int32 FoliageNameIndex = 0; FoliageNameIndex < FoliageFacade.NumFoliageNames(); ++FoliageNameIndex)
+		{
+			FString FoliageName = FoliageFacade.GetFoliageName(FoliageNameIndex);
+			TObjectPtr<UStaticMesh> PartStaticMesh = LoadObject<UStaticMesh>(nullptr, *FoliageName);
+			TObjectPtr<USkeletalMesh> PartSkeletalMesh = LoadObject<USkeletalMesh>(nullptr, *FoliageName);
+
+			if(PartStaticMesh)
+			{
+				int32 PartIndex = AssemblyBuilder.FindOrAddPart(PartStaticMesh.GetPath());
+				MeshNamePartMap.FindOrAdd(PartStaticMesh.GetPath()) = PartIndex;
+
+				for (auto Material : PartStaticMesh->GetStaticMaterials())
+				{
+					MeshMaterialsMap.FindOrAdd(FoliageName).AddUnique(Material.MaterialInterface);
+				}
+			}
+			else if(PartSkeletalMesh)
+			{
+				auto StaticMeshName = "SM_" + PartSkeletalMesh.GetName();
+				auto StaticMeshPath = PartSkeletalMesh.GetPath().Replace(*PartSkeletalMesh.GetName(), *StaticMeshName);
+				int32 PartIndex = AssemblyBuilder.FindOrAddPart(StaticMeshPath);
+				MeshNamePartMap.FindOrAdd(FoliageName) = PartIndex;
+		
+				TArray<FAssetData> AssetData;
+				IAssetRegistry::Get()->GetAssetsByPath(*PartSkeletalMesh.GetPackage()->GetName(), AssetData);
+		
+				if(AssetData.Num() <= 0)
+				{
+					FDynamicMesh3 PartDynamicMesh;
+					FText OutErrorMessage;
+
+					TObjectPtr<UDynamicMesh> DynamicMesh = NewObject<UDynamicMesh>();
+					
+					TArray<TObjectPtr<UMaterialInterface>> PartMaterials;
+					for (auto Material : PartSkeletalMesh->GetMaterials())
+					{
+						PartMaterials.Add(Material.MaterialInterface);
+						MeshMaterialsMap.FindOrAdd(FoliageName).AddUnique(Material.MaterialInterface);
+					}
+
+					if (SkeletalMeshToDynamicMesh(PartSkeletalMesh, 0, PartDynamicMesh))
+					{
+						auto MeshPackagePath = PartSkeletalMesh->GetPackage()->GetName().Replace(*PartSkeletalMesh.GetName(), *StaticMeshName);
+						UPackage* MeshPackage = CreatePackage(*MeshPackagePath);
+						
+						// Create a skeletal mesh to save in content browser
+						PartStaticMesh = NewObject<UStaticMesh>(MeshPackage, FName(StaticMeshName), RF_Standalone | RF_Public);
+
+						//Convert the dynamic mesh to static mesh
+						EGeometryScriptOutcomePins Output;
+#if WITH_EDITORONLY_DATA
+						FGeometryScriptCopyMeshToAssetOptions Options = GetCopyMeshToAssetOptions(PartMaterials, PartStaticMesh->GetNaniteSettings(), ENaniteShapePreservation::Voxelize);
+#else
+						FGeometryScriptCopyMeshToAssetOptions Options = GetCopyMeshToAssetOptions(PartMaterials);
+#endif
+						DynamicMesh->SetMesh(MoveTemp(PartDynamicMesh));
+
+						UGeometryScriptLibrary_StaticMeshFunctions::CopyMeshToStaticMesh(DynamicMesh, PartStaticMesh, Options, {}, Output);
+						IAssetRegistry::Get()->AssetCreated(PartStaticMesh);
+					}
+				}
+			}
+		}
+
+		AddNodeToBuilder(AssemblyBuilder, Collection, MeshNamePartMap, MeshMaterialsMap);
+		AssemblyBuilder.ApplyToStaticMesh(*StaticMesh);
+	}
+
+	void BuildNaniteAssemblyData(
+		const FManagedArrayCollection& Collection,
+		USkeletalMesh* SkeletalMesh
+	)
+	{
+		//Create NaniteAssemblyDataBuilder and the main static mesh materials to it
+		FNaniteAssemblyDataBuilder AssemblyBuilder;
+		Facades::FFoliageFacade FoliageFacade(Collection);
+		
+		for (const FSkeletalMaterial& Material : SkeletalMesh->GetMaterials())
+		{
+			AssemblyBuilder.AddMaterialSlot(Material.MaterialInterface);
+		}
+
+		TMap<FString, int32> MeshNamePartMap;
+		TMap<FString, TArray<TObjectPtr<UMaterialInterface>>> MeshMaterialsMap;
+			
+		for(int32 FoliageNameIndex = 0; FoliageNameIndex < FoliageFacade.NumFoliageNames(); ++FoliageNameIndex)
+		{
+			FString FoliageName = FoliageFacade.GetFoliageName(FoliageNameIndex);
+			TObjectPtr<UStaticMesh> PartStaticMesh = LoadObject<UStaticMesh>(nullptr, *FoliageName);
+			TObjectPtr<USkeletalMesh> PartSkeletalMesh = LoadObject<USkeletalMesh>(nullptr, *FoliageName);
+
+			if(PartStaticMesh)
+			{
+				auto SkeletalMeshName = "SKM_" + PartStaticMesh.GetName();
+				auto SkeletalMeshPath = PartStaticMesh.GetPath().Replace(*PartStaticMesh.GetName(), *SkeletalMeshName);
+				int32 PartIndex = AssemblyBuilder.FindOrAddPart(SkeletalMeshPath);
+				MeshNamePartMap.FindOrAdd(FoliageName) = PartIndex;
+		
+				TArray<FAssetData> AssetData;
+				IAssetRegistry::Get()->GetAssetsByPath(*PartStaticMesh.GetPackage()->GetName(), AssetData);
+		
+				if(AssetData.Num() <= 0)
+				{
+					FDynamicMesh3 PartDynamicMesh;
+					FText OutErrorMessage;
+
+					TArray<TObjectPtr<UMaterialInterface>> PartMaterials;
+					for (auto Material : PartStaticMesh->GetStaticMaterials())
+					{
+						PartMaterials.Add(Material.MaterialInterface);
+						MeshMaterialsMap.FindOrAdd(FoliageName).AddUnique(Material.MaterialInterface);
+					}
+
+					UE::Conversion::FStaticMeshConversionOptions Options;
+
+					if (StaticMeshToDynamicMesh(PartStaticMesh, PartDynamicMesh, OutErrorMessage, Options))
+					{
+						auto MeshPackagePath = PartStaticMesh->GetPackage()->GetName().Replace(*PartStaticMesh.GetName(), *SkeletalMeshName);
+						UPackage* MeshPackage = CreatePackage(*MeshPackagePath);
+						auto SkeletonName = "SK_" + PartStaticMesh.GetName();
+						auto SkeletonPath = PartStaticMesh.GetPackage()->GetName().Replace(*PartStaticMesh.GetName(), *SkeletonName);
+						UPackage* SkeletonPackage = CreatePackage(*(SkeletonPath));
+
+						// Create a skeletal mesh to save in content browser
+						PartSkeletalMesh = NewObject<USkeletalMesh>(MeshPackage, FName(SkeletalMeshName), RF_Standalone | RF_Public);
+						const TObjectPtr<USkeleton> PartSkeleton = NewObject<USkeleton>(SkeletonPackage, FName(SkeletonName),
+																						  RF_Standalone | RF_Public);
+						PartSkeletalMesh->SetSkeleton(PartSkeleton);
+
+						ConvertToDefaultSkeletalMesh(PartSkeletalMesh, &PartDynamicMesh, PartMaterials);
+
+						IAssetRegistry::Get()->AssetCreated(PartSkeletalMesh);
+						IAssetRegistry::Get()->AssetCreated(PartSkeleton);
+					}
+				}
+			}
+			else if (PartSkeletalMesh)
+			{
+				int32 PartIndex = AssemblyBuilder.FindOrAddPart(PartSkeletalMesh.GetPath());
+				MeshNamePartMap.FindOrAdd(PartSkeletalMesh.GetPath()) = PartIndex;
+
+				for (auto Material : PartSkeletalMesh->GetMaterials())
+				{
+					MeshMaterialsMap.FindOrAdd(FoliageName).AddUnique(Material.MaterialInterface);
+				}
+			}
+		}
+
+		AddNodeToBuilder(AssemblyBuilder, Collection, MeshNamePartMap, MeshMaterialsMap);
+		AssemblyBuilder.ApplyToSkeletalMesh(*SkeletalMesh);
+	}
+
+	void AddNodeToBuilder(FNaniteAssemblyDataBuilder& AssemblyBuilder, const FManagedArrayCollection& Collection,
+		const TMap<FString, int32>& InMeshNamePartMap, const TMap<FString, TArray<TObjectPtr<UMaterialInterface>>>& InMeshMaterialsMap)
+	{
+		Facades::FFoliageFacade FoliageFacade(Collection);
+		int32 NumInstances = FoliageFacade.NumFoliageEntries();
+			
+		for(int32 Id = 0; Id < NumInstances; Id++)
+		{
+			FString FoliageName = GetFoliageName(FoliageFacade, Id);
+
+			if(InMeshNamePartMap.Find(FoliageName))
+			{
+				int32 PartIndex = InMeshNamePartMap[FoliageName];
+				
+				FTransform Transform = FoliageFacade.GetFoliageTransform(Id);
+				UE::Geometry::FTransformSRT3d GeoTransform(Transform);
+
+				int32 FoliageParentBoneID = FoliageFacade.GetParentBoneID(Id);
+				
+				FNaniteAssemblyBoneInfluence Influence;
+				Influence.BoneIndex = FoliageParentBoneID;
+				//UE_LOG(LogMegaPlants, Warning, TEXT("Adding bone influence FoliageParentBoneID {%i}, Position{%s}, Scale{%s} , Rotation{%s}"),
+				//	FoliageParentBoneID, *Transform.GetLocation().ToString(), *Transform.GetScale3D().ToString(), *Transform.GetRotation().ToString());
+				Influence.BoneWeight = 1.0;
+				
+				AssemblyBuilder.AddNode(PartIndex, FTransform3f(GeoTransform), ENaniteAssemblyNodeTransformSpace::Local, { Influence });
+
+				//Add the materials for part to the assembly builder
+				int32 LocalMaterialIndex = 0;
+				auto FoliageMaterials = InMeshMaterialsMap.Find(FoliageName);
+				for (auto PartMaterial : *FoliageMaterials)
+				{
+					int32 PartMaterialIndex = AssemblyBuilder.GetMaterialSlots().IndexOfByPredicate(
+						[&] (const auto& Slot) { return Slot.Material == PartMaterial; }
+					);
+					if (PartMaterialIndex == INDEX_NONE)
+					{
+						PartMaterialIndex = AssemblyBuilder.AddMaterialSlot(PartMaterial);
+					}
+
+					AssemblyBuilder.RemapPartMaterial(PartIndex, LocalMaterialIndex, PartMaterialIndex);
+					LocalMaterialIndex++;
+				}
+			}
+		}
+	}
+
+	void ConvertToDefaultSkeletalMesh(USkeletalMesh* SkeletalMesh, FDynamicMesh3* Mesh,TArray<TObjectPtr<UMaterialInterface>> Materials)
+	{
+		TObjectPtr<UDynamicMesh> DynamicMesh = NewObject<UDynamicMesh>();
+		DynamicMesh->SetMesh(MoveTemp(*Mesh));
+		
+		//Figure out the correct parent bone index
+		int ParentBoneIndex = INDEX_NONE;
+		UE::Geometry::FDynamicMeshAttributeSet* MeshAttributes = DynamicMesh->GetMeshPtr()->Attributes();
+		MeshAttributes->EnableBones(1);
+
+		UE::Geometry::FDynamicMeshBoneNameAttribute* BoneNames = MeshAttributes->GetBoneNames();
+		UE::Geometry::FDynamicMeshBoneParentIndexAttribute* BoneParentIndices = MeshAttributes->GetBoneParentIndices();
+		UE::Geometry::FDynamicMeshBonePoseAttribute* BonePoses = MeshAttributes->GetBonePoses();
+
+		int BoneIndex = 0;
+		BoneNames->SetValue(BoneIndex, "Root");
+		BonePoses->SetValue(BoneIndex, FTransform::Identity);
+		BoneParentIndices->SetValue(BoneIndex, ParentBoneIndex);
+		
+		//Convert the dynamic mesh to static mesh
+		EGeometryScriptOutcomePins Output;
+#if WITH_EDITORONLY_DATA
+		FGeometryScriptCopyMeshToAssetOptions Options = GetCopyMeshToAssetOptions(Materials, SkeletalMesh->GetNaniteSettings(), ENaniteShapePreservation::Voxelize);
+#else
+		FGeometryScriptCopyMeshToAssetOptions Options = GetCopyMeshToAssetOptions(PartMaterials);
+#endif
+		Options.BoneHierarchyMismatchHandling = EGeometryScriptBoneHierarchyMismatchHandling::CreateNewReferenceSkeleton;
+		
+		UGeometryScriptLibrary_StaticMeshFunctions::CopyMeshToSkeletalMesh(
+			DynamicMesh,
+			SkeletalMesh,
+			Options,
+			{},
+			Output
+		);
+
+		TObjectPtr<USkeleton> Skeleton = SkeletalMesh->GetSkeleton();
+		Skeleton->RecreateBoneTree(SkeletalMesh);
+
+		const FGeometryScriptBoneWeightProfile SkinProfile = FGeometryScriptBoneWeightProfile();
+		FGeometryScriptSmoothBoneWeightsOptions SmoothBoneWeightsOptions;
+		SmoothBoneWeightsOptions.DistanceWeighingType = EGeometryScriptSmoothBoneWeightsType::GeodesicVoxel;
+		SmoothBoneWeightsOptions.MaxInfluences = 2;
+		SmoothBoneWeightsOptions.VoxelResolution = 512;
+		UGeometryScriptLibrary_MeshBoneWeightFunctions::ComputeSmoothBoneWeights(
+			DynamicMesh,
+			Skeleton,
+			SmoothBoneWeightsOptions,
+			SkinProfile
+		);
+
+		const FName ProfileName = SkinProfile.GetProfileName();
+		UGeometryScriptLibrary_StaticMeshFunctions::CopySkinWeightProfileToSkeletalMesh(
+			DynamicMesh,
+			SkeletalMesh,
+			ProfileName,
+			ProfileName,
+			{},
+			{},
+			Output
+		);
+	}
+#endif
+
+	void ExportCollectionToStaticMesh(
+		TObjectPtr<UStaticMesh> ExportMesh,
+		const FManagedArrayCollection& Collection,
+		bool bBuildNaniteAssemblies,
+		ENaniteShapePreservation ShapePreservation
+	)
+	{
+		FManagedArrayCollection ExportedCollection;
+		Collection.CopyTo(&ExportedCollection);
+		
+		TObjectPtr<UDynamicMesh> DynamicMesh = CollectionToDynamicMesh(ExportedCollection);
+
+		//Get the materials from the collection
+		TArray<TObjectPtr<UMaterialInterface>> Materials = GetMaterialsFromCollection(ExportedCollection);
+
+		//Get the foliage data from facade, spawn an actor attach the foliage to the actor
+		Facades::FFoliageFacade FoliageFacade(ExportedCollection);
+		
+		//Combine all the foliage instances into one dynamic mesh if not building nanite assemblies
+		if (!bBuildNaniteAssemblies)
+		{
+			CombineMeshInstancesToDynamicMesh(ExportedCollection, DynamicMesh, Materials,
+				[](FString MeshName, int32 InstanceID, int32 VertexStart, int32 VertexCount){});
+		}
+
+		//Convert the dynamic mesh to static mesh
+		EGeometryScriptOutcomePins Output;
+#if WITH_EDITORONLY_DATA
+		FGeometryScriptCopyMeshToAssetOptions Options = GetCopyMeshToAssetOptions(Materials, ExportMesh->GetNaniteSettings(), ShapePreservation);
+		
+#else
+		FGeometryScriptCopyMeshToAssetOptions Options = GetCopyMeshToAssetOptions(Materials);
+#endif
+		UGeometryScriptLibrary_StaticMeshFunctions::CopyMeshToStaticMesh(DynamicMesh, ExportMesh, Options, {}, Output);
+
+#if WITH_EDITOR
+		if (bBuildNaniteAssemblies)
+		{
+			if (NaniteAssembliesSupported())
+			{
+				//Build the nanite assemblies
+            	BuildNaniteAssemblyData(ExportedCollection, ExportMesh);
+            	ExportMesh->Build(true);
+			}
+			else
+			{
+				UE_LOG(LogMegaPlants, Warning, TEXT("Failed to build Nanite Assemblies, because neither Nanite Foliage nor Nanite Assemblies are enabled for the project."));
+			}
+		}
+#endif
+	}
+
+	void ExportCollectionToSkeletalMesh(
+		TObjectPtr<USkeletalMesh> ExportMesh,
+		const FManagedArrayCollection& Collection,
+		ENaniteShapePreservation ShapePreservation,
+		bool bBuildNaniteAssemblies
+	)
+	{
+		FManagedArrayCollection ExportedCollection;
+		Collection.CopyTo(&ExportedCollection);
+
+		Facades::FFoliageFacade FoliageFacade(ExportedCollection);
+		Facades::FPointFacade PointFacade(ExportedCollection);
+		Facades::FBranchFacade BranchFacade(ExportedCollection);
+		Facades::FBoneFacade BoneFacade = Facades::FBoneFacade(ExportedCollection);
+		
+		TObjectPtr<UDynamicMesh> DynamicMesh = CollectionToDynamicMesh(ExportedCollection);
+		UE::Geometry::FDynamicMeshAttributeSet* MeshAttributes = DynamicMesh->GetMeshPtr()->Attributes();
+		
+		//Get the materials from the collection
+		TArray<TObjectPtr<UMaterialInterface>> Materials = GetMaterialsFromCollection(ExportedCollection);
+
+		TArray<Facades::FBoneNode> BoneNodes;
+		BoneFacade.CreateBoneData(BoneNodes);
+		AssignBoneIDsToFoliage(BoneNodes, ExportedCollection);
+		
+		int32 BoneCount = BoneNodes.Num();
+		MeshAttributes->EnableBones(BoneCount);
+
+		//NjordPixel id is set per vertex
+		TArray<float> NjordPixelIDs;
+		BoneFacade.GetNjordPixelIds(NjordPixelIDs);
+
+		TArray<int32> VertexBoneIDs;
+		//Assign the Bone id to every vertex in the base mesh through NjordPixelId
+		for (const float& NjordPixelId : NjordPixelIDs)
+		{
+			Facades::FBoneNode* BoneNode = BoneNodes.FindByPredicate([&NjordPixelId](const Facades::FBoneNode& Node)
+			{
+				return FMath::FloorToInt(Node.NjordPixelID) == FMath::FloorToInt(NjordPixelId);
+			});
+			VertexBoneIDs.Add(BoneNode->BoneIndex);
+		}
+		int FoliageVertexStart = VertexBoneIDs.Num();
+
+		UE::Geometry::FDynamicMeshBoneNameAttribute* BoneNames = MeshAttributes->GetBoneNames();
+		UE::Geometry::FDynamicMeshBoneParentIndexAttribute* BoneParentIndices = MeshAttributes->GetBoneParentIndices();
+		UE::Geometry::FDynamicMeshBonePoseAttribute* BonePoses = MeshAttributes->GetBonePoses();
+
+		for (int i = 0; i < BoneCount; i++)
+		{
+			BoneNames->SetValue(i, BoneNodes[i].BoneName);
+			BonePoses->SetValue(i, BoneNodes[i].BoneTransform);
+			BoneParentIndices->SetValue(i, BoneNodes[i].ParentBoneIndex);
+		}
+
+		if (!bBuildNaniteAssemblies)
+		{
+			//Combine all the foliage instances into one dynamic mesh if not building nanite assemblies
+			CombineMeshInstancesToDynamicMesh(ExportedCollection, DynamicMesh, Materials,
+			[&](const FString& MeshName, int32 InstanceID, int32 VertexStart, int32 VertexCount)
+			{
+				//Assign parent bone id to foliage vertices
+				int32 FoliageParentBoneID = FoliageFacade.GetParentBoneID(InstanceID);
+					
+				for (int VertexID = VertexStart; VertexID < VertexCount; ++VertexID)
+				{
+					VertexBoneIDs.Add(FoliageParentBoneID);
+					//UE_LOG(LogMegaPlants, Warning, TEXT("Adding VertexBoneIDs, VertexId{%i} BoneID{%i}, VertexStart{%i}, VertexCount{%i}, MeshName{%s}, InstanceID{%i}"),
+					//VertexID, FoliageParentBoneID, VertexStart, VertexCount, *MeshName, InstanceID);
+				}
+			});
+		}
+
+		//Convert the dynamic mesh to static mesh
+		EGeometryScriptOutcomePins Output;
+		
+#if WITH_EDITORONLY_DATA
+		FGeometryScriptCopyMeshToAssetOptions Options = GetCopyMeshToAssetOptions(Materials, ExportMesh->GetNaniteSettings(), ShapePreservation);
+		
+#else
+		FGeometryScriptCopyMeshToAssetOptions Options = GetCopyMeshToAssetOptions(Materials);
+#endif
+		Options.BoneHierarchyMismatchHandling = EGeometryScriptBoneHierarchyMismatchHandling::CreateNewReferenceSkeleton;
+
+		UGeometryScriptLibrary_StaticMeshFunctions::CopyMeshToSkeletalMesh(
+			DynamicMesh,
+			ExportMesh,
+			Options,
+			{},
+			Output
+		);
+
+		TObjectPtr<USkeleton> Skeleton = ExportMesh->GetSkeleton();
+		Skeleton->RecreateBoneTree(ExportMesh);
+
+		const FGeometryScriptBoneWeightProfile SkinProfile = FGeometryScriptBoneWeightProfile();
+		FGeometryScriptSmoothBoneWeightsOptions SmoothBoneWeightsOptions;
+
+		SmoothBoneWeightsOptions.DistanceWeighingType = EGeometryScriptSmoothBoneWeightsType::DirectDistance;
+		SmoothBoneWeightsOptions.MaxInfluences = 3;
+		SmoothBoneWeightsOptions.Stiffness = 0.5;
+		
+		UGeometryScriptLibrary_MeshBoneWeightFunctions::ComputeSmoothBoneWeights(
+			DynamicMesh,
+			Skeleton,
+			SmoothBoneWeightsOptions,
+			SkinProfile
+		);
+
+		const FName ProfileName = SkinProfile.GetProfileName();
+		
+		//Here we make the correction to the weights
+		RemoveUnwantedSkinWeights(DynamicMesh, ProfileName, VertexBoneIDs, NjordPixelIDs, FoliageVertexStart,BoneNodes);
+		
+		UGeometryScriptLibrary_StaticMeshFunctions::CopySkinWeightProfileToSkeletalMesh(
+			DynamicMesh,
+			ExportMesh,
+			ProfileName,
+			ProfileName,
+			{},
+			{},
+			Output
+		);
+
+#if WITH_EDITOR
+		if (bBuildNaniteAssemblies)
+		{
+			if (NaniteAssembliesSupported())
+			{
+				//Build the nanite assemblies
+				BuildNaniteAssemblyData(ExportedCollection, ExportMesh);
+				ExportMesh->Build();
+			}
+			else
+			{
+				UE_LOG(LogMegaPlants, Warning, TEXT("Failed to build Nanite Assemblies, because neither Nanite Foliage nor Nanite Assemblies are enabled for the project."));
+			}
+		}
+#endif
+	}
+
+	void AttachMegaPlantLink(UObject* InExportedMesh, const TObjectPtr<UMegaPlant>& InMegaPlant)
+	{
+		if (IInterface_AssetUserData* IAssetUserData = Cast<IInterface_AssetUserData>(InExportedMesh))
+		{
+			UMegaPlantLink* Data = NewObject<UMegaPlantLink>(InExportedMesh);
+			Data->Source = InMegaPlant;
+		
+			IAssetUserData->AddAssetUserData(Data);	
+		}
+	}
+
+	void AssignBoneIDsToFoliage(const TArray<Facades::FBoneNode>& Bones, FManagedArrayCollection& Collection)
+	{
+		Facades::FPointFacade PointFacade(Collection);
+		Facades::FFoliageFacade FoliageFacade(Collection);
+		
+		int32 NumInstances = FoliageFacade.NumFoliageEntries();
+
+		auto FindBoneById = ( [&](const int32 Id)
+		{
+			const Facades::FBoneNode* BoneNode = Bones.FindByPredicate([Id](const Facades::FBoneNode& Node)
+			{
+				return Node.BoneIndex == Id;
+			});
+			return BoneNode;
+		});
+		
+		for(int32 Id = 0; Id < NumInstances; Id++)
+		{
+			Facades::FFoliageEntryData Data = FoliageFacade.GetFoliageEntry(Id);
+			float FoliageLFR = Data.LengthFromRoot;
+
+			TArray<int32> BoneIDs;
+			for (Facades::FBoneNode BoneNode : Bones)
+			{
+				if (Data.BranchId == BoneNode.BranchIndex)
+				{
+					BoneIDs.Add(BoneNode.BoneIndex);
+				}
+			}
+
+			bool bBoneIDAssigned = false;
+			for (const int32& BoneId : BoneIDs)
+			{
+				const Facades::FBoneNode* BoneNode = FindBoneById(BoneId);
+				check(BoneNode);
+				
+				int PointIndex = BoneNode->PointIndex;
+				const Facades::FBoneNode* BoneParentNode = FindBoneById(BoneNode->ParentBoneIndex);
+				
+				int ParentPointIndex = BoneParentNode ? BoneParentNode->PointIndex : INDEX_NONE;
+				float BoneLFR = PointFacade.GetLengthFromRoot(PointIndex);
+				float BoneParentLFR = PointFacade.GetLengthFromRoot(ParentPointIndex);
+
+				if (FoliageLFR <= BoneLFR && FoliageLFR >= BoneParentLFR)
+				{
+					FoliageFacade.SetParentBoneID(Id, BoneId);
+					bBoneIDAssigned = true;
+					break;
+				}
+			}
+
+			if (!bBoneIDAssigned)
+			{
+				UE_LOG(LogMegaPlants, Warning, TEXT("No bone assigned for foliage Id {%i} , Branch Bone count {%i}"), Id , BoneIDs.Num());
+			}
+		}
+	}
+
+	void RemoveUnwantedSkinWeights(TObjectPtr<UDynamicMesh> DynamicMesh, const FName ProfileName, const TArray<int32>& VertexBoneIDs, const TArray<float>& NjordIDs, int32 FoliageVertexStart, const TArray<Facades::FBoneNode>& BoneNodes)
+	{
+		auto InMesh = DynamicMesh->GetMeshPtr();
+		check(InMesh)
+		const int32 NumVertices = VertexBoneIDs.Num();
+		
+		DynamicMesh->EditMesh([&](FDynamicMesh3& EditMesh)
+		{
+			UE::Geometry::FDynamicMeshVertexSkinWeightsAttribute *SkinWeights = EditMesh.Attributes()->GetSkinWeightsAttribute(ProfileName);
+
+			ParallelFor(NumVertices, [&](const int32 VertexIdx)
+			{
+				int32 BoneID = VertexBoneIDs[VertexIdx];
+				int32 ParentBoneID = BoneNodes[BoneID].ParentBoneIndex;
+			
+				//UE_LOG(LogMegaPlants, Warning, TEXT("Remove Bones, VertexId{%i} BoneID{%i} BoneParentIndex{%i}"), VertexIdx, BoneID, ParentBoneID);
+
+				UE::AnimationCore::FBoneWeights Weights;
+				SkinWeights->GetValue(VertexIdx, Weights);
+
+				auto NextBone = BoneNodes.FindByPredicate([&](const Facades::FBoneNode& Node)
+				{
+					return Node.BoneIndex == BoneID + 1 && Node.BranchIndex == BoneNodes[BoneID].BranchIndex;
+				});
+
+				for (int32 WeightID = Weights.Num() - 1 ; WeightID >= 0; WeightID--)
+				{
+					auto Weight = Weights[WeightID];
+
+					bool bValidBone = (Weight.GetBoneIndex() == ParentBoneID);
+					bValidBone |= (Weight.GetBoneIndex() == BoneID);
+					bValidBone |= NextBone && (Weight.GetBoneIndex() == NextBone->BoneIndex);
+					
+					if (!bValidBone || VertexIdx > FoliageVertexStart || BoneID == 0)
+					{
+						UE::AnimationCore::FBoneWeightsSettings Settings;
+						Settings.SetNormalizeType(UE::AnimationCore::EBoneWeightNormalizeType::Always);
+						Weights.RemoveBoneWeight(Weight.GetBoneIndex(),Settings);
+					}
+				}
+
+				if (VertexIdx > FoliageVertexStart || BoneID == 0 || Weights.Num() == 0)
+				{
+					UE::AnimationCore::FBoneWeightsSettings Settings;
+					Settings.SetNormalizeType(UE::AnimationCore::EBoneWeightNormalizeType::Always);
+					Weights.SetBoneWeight(UE::AnimationCore::FBoneWeight(BoneID, 1.0f), Settings);
+				}
+				
+				SkinWeights->SetNewValue(VertexIdx, Weights);
+			});
+				
+		}, EDynamicMeshChangeType::GeneralEdit, EDynamicMeshAttributeChangeFlags::Unknown, false);
+	}
+}
+
+namespace MegaPlants::Export
+{
+	void ExportCollectionAsMesh(
+		const TObjectPtr<UMegaPlant> InMegaPlant,
+		const FManagedArrayCollection& Collection,
+		const FManagedArrayCollection& FoliageCollection,
+		const FMegaPlantExportSettings& Settings
+	)
+	{
+		// Set up the Package path and name from settings
+		const FString PackageName = Settings.ContentBrowserFolder.Path;
+		const FString MeshName = Settings.MeshName.ToString();
+
+		const FString FullPath = FPaths::Combine(PackageName, MeshName);
+
+		if (Settings.ExportMeshType == EMegaPlantExportMeshType::StaticMesh)
+		{
+			UPackage* MeshPackage = CreatePackage(*FullPath);
+
+			// Create a static mesh to save in content browser
+			const TObjectPtr<UStaticMesh> ExportMesh = NewObject<UStaticMesh>(MeshPackage, FName(MeshName), RF_Standalone | RF_Public);
+
+			Internal::AttachMegaPlantLink(ExportMesh, InMegaPlant);
+			
+			MegaPlants::Export::Internal::ExportCollectionToStaticMesh(
+				ExportMesh,
+				Collection,
+				Settings.bCreateNaniteAssemblies,
+				Settings.NaniteShapePreservation
+			);
+
+			MeshPackage->SetDirtyFlag(true);
+		}
+		else if (Settings.ExportMeshType == EMegaPlantExportMeshType::SkeletalMesh)
+		{
+			UPackage* MeshPackage = CreatePackage(*FullPath);
+			UPackage* SkeletonPackage = CreatePackage(*(FullPath + "_Skeleton"));
+
+			// Create a skeletal mesh to save in content browser
+			const TObjectPtr<USkeletalMesh> ExportMesh = NewObject<USkeletalMesh>(MeshPackage, FName(MeshName), RF_Standalone | RF_Public);
+			const TObjectPtr<USkeleton> ExportSkeleton = NewObject<USkeleton>(SkeletonPackage, FName(MeshName + "_Skeleton"),
+			                                                                  RF_Standalone | RF_Public);
+			Internal::AttachMegaPlantLink(ExportMesh, InMegaPlant);
+			Internal::AttachMegaPlantLink(ExportSkeleton, InMegaPlant);
+			
+			ExportMesh->SetSkeleton(ExportSkeleton);
+
+			MegaPlants::Export::Internal::ExportCollectionToSkeletalMesh(
+				ExportMesh,
+				Collection,
+				Settings.NaniteShapePreservation,
+				Settings.bCreateNaniteAssemblies
+			);
+
+			MeshPackage->SetDirtyFlag(true);
+			SkeletonPackage->SetDirtyFlag(true);
+
+			IAssetRegistry::Get()->AssetCreated(ExportMesh);
+			IAssetRegistry::Get()->AssetCreated(ExportSkeleton);
+		}
+	}
+}
