Create ProceduralVegetationEditor plugin

diff --git a/Engine/Plugins/Experimental/ProceduralVegetationEditor/Source/MegaPlants/Private/Helpers/MegaPlantUtilities.cpp b/Engine/Plugins/Experimental/ProceduralVegetationEditor/Source/MegaPlants/Private/Helpers/MegaPlantUtilities.cpp
new file mode 100644
--- /dev/null
+++ b/Engine/Plugins/Experimental/ProceduralVegetationEditor/Source/MegaPlants/Private/Helpers/MegaPlantUtilities.cpp
@@ -0,0 +1,108 @@
+﻿// Copyright Epic Games, Inc. All Rights Reserved.
+
+#include "Helpers/MegaPlantUtilities.h"
+
+#include "MegaPlantsModule.h"
+#include "Components/InstancedSkinnedMeshComponent.h"
+#include "Engine/StaticMesh.h"
+#include "Engine/SkeletalMesh.h"
+#include "UObject/Package.h"
+#include "Components/InstancedStaticMeshComponent.h"
+#include "Components/MeshComponent.h"
+
+namespace MegaPlants::Utilities
+{
+	void AddFoliageInstancesToActor(const Facades::FFoliageFacade& FoliageFacade, AActor* Actor,
+	                                TMap<FString, TObjectPtr<UMeshComponent>>& OutInstancedComponentMap)
+	{
+		TRACE_CPUPROFILER_EVENT_SCOPE(MegaPlants::Utilities::AddFoliageInstancesToActor());
+
+		auto OnComponentCreated = [&Actor](UMeshComponent* InInstancedComponent)
+			{
+				InInstancedComponent->RegisterComponentWithWorld(Actor->GetWorld());
+				InInstancedComponent->AttachToComponent(Actor->GetRootComponent(), FAttachmentTransformRules::KeepRelativeTransform);
+				Actor->AddInstanceComponent(InInstancedComponent);
+			};
+		
+		AddFoliageInstances(FoliageFacade, Actor, FFoliageComponentCreatedCallback::CreateLambda(OnComponentCreated),OutInstancedComponentMap);
+	}
+
+	void AddFoliageInstances(const Facades::FFoliageFacade& InFoliageFacade, const UObject* InParent, FFoliageComponentCreatedCallback InCallback, TMap<FString, TObjectPtr<UMeshComponent>>& OutInstancedComponentMap)
+	{
+		TRACE_CPUPROFILER_EVENT_SCOPE(MegaPlants::Utilities::AddFoliageInstances());
+
+		int32 NumInstances = InFoliageFacade.NumFoliageEntries();
+
+		OutInstancedComponentMap.Empty();
+
+		if (NumInstances > 0)
+		{
+			for (int32 FoliageNameIndex = 0; FoliageNameIndex < InFoliageFacade.NumFoliageNames(); ++FoliageNameIndex)
+			{
+				FString FoliageName = InFoliageFacade.GetFoliageName(FoliageNameIndex);
+
+				UStaticMesh* StaticFoliageMesh = LoadObject<UStaticMesh>(nullptr, *FoliageName);
+				USkeletalMesh* SkeletalFoliageMesh = LoadObject<USkeletalMesh>(nullptr, *FoliageName);
+
+				if(StaticFoliageMesh)
+				{
+					UInstancedStaticMeshComponent* InstancedComponent = NewObject<UInstancedStaticMeshComponent>(InParent != nullptr ? InParent->GetPackage() : GetTransientPackage(), NAME_None);
+					InstancedComponent->SetStaticMesh(StaticFoliageMesh);
+					OutInstancedComponentMap.Add(FoliageName, InstancedComponent);
+					InCallback.ExecuteIfBound(InstancedComponent);
+				}
+				else if(SkeletalFoliageMesh)
+				{
+#if WITH_EDITORONLY_DATA
+					if (!SkeletalFoliageMesh->NaniteSettings.bEnabled)
+					{
+						SkeletalFoliageMesh->NaniteSettings.bEnabled = true;
+						SkeletalFoliageMesh->Build();
+						SkeletalFoliageMesh->MarkPackageDirty();
+						UE_LOG(LogMegaPlants, Warning, TEXT("Enabling Nanite for the Skeletal Mesh {%s}, its required for InstancedSkinnedMeshComponent"), *FoliageName);
+					}
+#else
+					UE_LOG(LogMegaPlants, Warning, TEXT("Instanced Skinned Mesh component only works with nanite meshes make sure nanite is enabled for mesh {%s}"), *FoliageName);
+#endif
+					
+					UInstancedSkinnedMeshComponent* InstancedComponent = NewObject<UInstancedSkinnedMeshComponent>(InParent != nullptr ? InParent->GetPackage() : GetTransientPackage(), NAME_None);
+					InstancedComponent->SetSkinnedAssetAndUpdate(SkeletalFoliageMesh);
+					OutInstancedComponentMap.Add(FoliageName, InstancedComponent);
+					InCallback.ExecuteIfBound(InstancedComponent);
+				}
+			}
+
+			TMap<FString, TArray<FTransform>> InstanceTransformsMap;
+
+			for (int32 Id = 0; Id < NumInstances; Id++)
+			{
+				Facades::FFoliageEntryData Data = InFoliageFacade.GetFoliageEntry(Id);
+				FString FoliageName = InFoliageFacade.GetFoliageName(Data.NameId);
+
+				const FTransform Transform = InFoliageFacade.GetFoliageTransform(Id);
+
+				InstanceTransformsMap.FindOrAdd(FoliageName).Add(Transform);
+			}
+		
+			for(auto FoliageEntry : InstanceTransformsMap)
+			{
+				if (FoliageEntry.Value.Num() > 0)
+				{
+					if (auto InstancedComponent = OutInstancedComponentMap.Find(FoliageEntry.Key))
+					{
+						if(TObjectPtr<UInstancedStaticMeshComponent> StaticMeshInstancedComponent = Cast<UInstancedStaticMeshComponent>(*InstancedComponent))
+						{
+							StaticMeshInstancedComponent->AddInstances(FoliageEntry.Value, false);	
+						}
+						else if(TObjectPtr<UInstancedSkinnedMeshComponent> SkinnedInstancedComponent = Cast<UInstancedSkinnedMeshComponent>(*InstancedComponent))
+						{
+							TArray<int32> BankIndices;
+							BankIndices.AddUninitialized(FoliageEntry.Value.Num());
+							SkinnedInstancedComponent->AddInstances(FoliageEntry.Value, BankIndices, false);
+						}
+					}
+				}
+			}
+		}
+	}
+}
diff --git a/Engine/Plugins/Experimental/ProceduralVegetationEditor/Source/MegaPlants/Private/Implementations/MegaPlantCarve.cpp b/Engine/Plugins/Experimental/ProceduralVegetationEditor/Source/MegaPlants/Private/Implementations/MegaPlantCarve.cpp
new file mode 100644
--- /dev/null
+++ b/Engine/Plugins/Experimental/ProceduralVegetationEditor/Source/MegaPlants/Private/Implementations/MegaPlantCarve.cpp
@@ -0,0 +1,631 @@
+// Copyright Epic Games, Inc. All Rights Reserved.
+
+#include "MegaPlantCarve.h"
+
+#include "Facades/MegaPlantFoliageFacade.h"
+#include "Facades/MegaPlantPointFacade.h"
+#include "Facades/MegaPlantTreeFacade.h"
+
+
+void FMegaPlantCarve::ApplyCarve(FManagedArrayCollection& OutCollection, const FManagedArrayCollection& SourceCollection,
+                                 const ECarveBasis CarveBasis, const float Carve)
+{
+	if (Carve == 0.0f)
+	{
+		return;
+	}
+
+	if (CarveBasis == ECarveBasis::FromBottom)
+	{
+		CarveFromBottom(OutCollection, SourceCollection, Carve);
+	}
+	else
+	{
+		CarveFromTop(OutCollection, SourceCollection, Carve, CarveBasis);
+	}
+}
+
+void FMegaPlantCarve::UpdatePointScales(MegaPlants::Facades::FPointFacade& PointFacadeOut, const MegaPlants::Facades::FPointFacade& PointFacadeSource,
+                                        const TArray<int>& BranchPoints, const float LastPointScale, const int32 LastPointIndex,
+                                        const float CarveRatio, const int32 EndIndex, const float FirstPointTargetPScale)
+{
+	TRACE_CPUPROFILER_EVENT_SCOPE(MegaPlantCarve::UpdatePointScales);
+
+	const float OldMax = FirstPointTargetPScale;
+	const float NewMax = OldMax * CarveRatio;
+	const float Ratio = NewMax / OldMax;
+	const float OldMin = PointFacadeSource.GetPointScale(LastPointIndex);
+	const float NewMin = LastPointScale * Ratio;
+
+	for (int32 i = 0; i <= EndIndex; ++i)
+	{
+		const float UpdatedScale = FMath::GetMappedRangeValueClamped(
+			FVector2f(OldMin, OldMax),
+			FVector2f(NewMin, NewMax),
+			PointFacadeSource.GetPointScale(BranchPoints[i]));
+		PointFacadeOut.ModifyPointScales()[BranchPoints[i]] = UpdatedScale;
+	}
+}
+
+void FMegaPlantCarve::UpdateBranchLinkedIndices(MegaPlants::Facades::FBranchFacade& BranchFacadeOut,
+                                                MegaPlants::Facades::FFoliageFacade& FoliageFacadeOut, const TMap<int32, int32>& PointsOldIDsToNewIDs,
+                                                const TArray<bool>& PointsRemoved, const TMap<int32, int32>& FoliageInstancesOldIDsToNewIDs,
+                                                const TArray<bool>& FoliageInstancesRemoved)
+{
+	TRACE_CPUPROFILER_EVENT_SCOPE(MegaPlantCarve::UpdateBranchLinkedIndices);
+
+	const int32 NumOfBranches = BranchFacadeOut.GetElementCount();
+	TSet<int32> CurrentBranchNumbers;
+	CurrentBranchNumbers.Reserve(NumOfBranches);
+	for (int32 BranchIndex = 0; BranchIndex < NumOfBranches; ++BranchIndex)
+	{
+		CurrentBranchNumbers.Add(BranchFacadeOut.GetBranchNumber(BranchIndex));
+	}
+
+	for (int32 BranchIndex = 0; BranchIndex < NumOfBranches; ++BranchIndex)
+	{
+		TArray<int32> UpdatedPoints;
+		const TArray<int32>& BranchPoints = BranchFacadeOut.GetPoints(BranchIndex);
+		for (const int32& PointIndex : BranchPoints)
+		{
+			if (PointsOldIDsToNewIDs.Contains(PointIndex))
+			{
+				UpdatedPoints.Add(PointsOldIDsToNewIDs[PointIndex]);
+			}
+			else if (!PointsRemoved[PointIndex])
+			{
+				UpdatedPoints.Add(PointIndex);
+			}
+		}
+		BranchFacadeOut.SetPoints(BranchIndex, UpdatedPoints);
+
+		TArray<int32> UpdatedParents;
+		const TArray<int32>& BranchParents = BranchFacadeOut.GetParents(BranchIndex);
+		for (const int32& BranchNumber : BranchParents)
+		{
+			if (BranchNumber == 0 || CurrentBranchNumbers.Contains(BranchNumber))
+			{
+				UpdatedParents.Add(BranchNumber);
+			}
+		}
+		BranchFacadeOut.SetParents(BranchIndex, UpdatedParents);
+
+		TArray<int32> UpdatedChildren;
+		const TArray<int32>& BranchChildren = BranchFacadeOut.GetChildren(BranchIndex);
+		for (const int32& BranchNumber : BranchChildren)
+		{
+			if (CurrentBranchNumbers.Contains(BranchNumber))
+			{
+				UpdatedChildren.Add(BranchNumber);
+			}
+		}
+		BranchFacadeOut.SetChildren(BranchIndex, UpdatedChildren);
+
+		TArray<int32> UpdatedFoliageInstances;
+		const TArray<int32>& FoliageInstances = FoliageFacadeOut.GetFoliageEntryIdsForBranch(BranchIndex);
+		for (const int32& FoliageInstanceIndex : FoliageInstances)
+		{
+			if (FoliageInstancesOldIDsToNewIDs.Contains(FoliageInstanceIndex))
+			{
+				const int32 NewFoliageInstanceID = FoliageInstancesOldIDsToNewIDs[FoliageInstanceIndex];
+				UpdatedFoliageInstances.Add(NewFoliageInstanceID);
+				FoliageFacadeOut.SetFoliageBranchId(NewFoliageInstanceID, BranchIndex);
+			}
+			else if (!FoliageInstancesRemoved[FoliageInstanceIndex])
+			{
+				UpdatedFoliageInstances.Add(FoliageInstanceIndex);
+			}
+		}
+		FoliageFacadeOut.SetFoliageIdsArray(BranchIndex, UpdatedFoliageInstances);
+	}
+}
+
+void FMegaPlantCarve::RecomputeAttributes(MegaPlants::Facades::FPointFacade& PointFacadeOut,
+                                          const MegaPlants::Facades::FPointFacade& PointFacadeSource,
+                                          const MegaPlants::Facades::FBranchFacade& BranchFacade,
+                                          const MegaPlants::Facades::FBranchFacade& BranchFacadeSource,
+                                          const TMap<int32, int32>& BranchesNewIDsToOldIDs)
+{
+	TRACE_CPUPROFILER_EVENT_SCOPE(MegaPlantCarve::RecomputeAttributes);
+
+	// Remap Ethylene levels while gathering minimum and maximum PScales for each branch and the whole plant
+	float MinimumPScale = TNumericLimits<float>::Max();
+	float MaximumPScale = TNumericLimits<float>::Min();
+	TArray<TTuple<float, float>> BranchMinMaxPScales;
+	BranchMinMaxPScales.Init(TTuple<float, float>(TNumericLimits<float>::Max(), TNumericLimits<float>::Min()), BranchFacade.GetElementCount());
+
+	for (int32 BranchIndex = 0; BranchIndex < BranchFacade.GetElementCount(); ++BranchIndex)
+	{
+		const TArray<int32>& BranchPoints = BranchFacade.GetPoints(BranchIndex);
+		const int32 NumBranchPointsCurrent = BranchPoints.Num();
+		const int32 OldBranchIndex = BranchesNewIDsToOldIDs.Contains(BranchIndex)
+			? BranchesNewIDsToOldIDs[BranchIndex]
+			: BranchIndex;
+		const TArray<int32>& SourceBranchPoints = BranchFacadeSource.GetPoints(OldBranchIndex);
+
+		for (int32 Idx = 0; Idx < NumBranchPointsCurrent; ++Idx)
+		{
+			const float RelativePointIndex = (static_cast<float>(Idx) / NumBranchPointsCurrent) * SourceBranchPoints.Num();
+			const int32 OldCurrentPointIndex = FMath::Clamp(FMath::CeilToInt(RelativePointIndex), 0, SourceBranchPoints.Num() - 1);
+			const int32 OldPreviousPointIndex = OldCurrentPointIndex == 0
+				? 0
+				: OldCurrentPointIndex - 1;
+			float BlendValue = RelativePointIndex - FMath::FloorToInt(RelativePointIndex);
+
+			const TArray<float>& BudHormoneLevelsAtOldPreviousPointIndex = PointFacadeSource.GetBudHormoneLevels(
+				SourceBranchPoints[OldPreviousPointIndex]);
+			check(BudHormoneLevelsAtOldPreviousPointIndex.Num() >= 5);
+			float EthyleneLevelAtOldPreviousPointIndex = BudHormoneLevelsAtOldPreviousPointIndex[4];
+
+			const TArray<float>& BudHormoneLevelsAtOldCurrentPointIndex = PointFacadeSource.GetBudHormoneLevels(
+				SourceBranchPoints[OldCurrentPointIndex]);
+			check(BudHormoneLevelsAtOldCurrentPointIndex.Num() >= 5);
+			float EthyleneLevelAtOldCurrentPointIndex = BudHormoneLevelsAtOldCurrentPointIndex[4];
+
+			const float NewEthyleneLevel = FMath::Lerp(EthyleneLevelAtOldPreviousPointIndex, EthyleneLevelAtOldCurrentPointIndex, BlendValue);
+			TArray<float> CurrentBudHormoneLevels = PointFacadeOut.GetBudHormoneLevels(BranchPoints[Idx]);
+			check(CurrentBudHormoneLevels.Num() >= 4);
+			CurrentBudHormoneLevels[4] = NewEthyleneLevel;
+			PointFacadeOut.SetBudHormoneLevels(BranchPoints[Idx], CurrentBudHormoneLevels);
+
+			const float PointScaleAtIndex = PointFacadeOut.GetPointScale(BranchPoints[Idx]);
+			MinimumPScale = FMath::Min(MinimumPScale, PointScaleAtIndex);
+			MaximumPScale = FMath::Max(MaximumPScale, PointScaleAtIndex);
+
+			BranchMinMaxPScales[BranchIndex].Get<0>() = FMath::Min(BranchMinMaxPScales[BranchIndex].Get<0>(), PointScaleAtIndex);
+			BranchMinMaxPScales[BranchIndex].Get<1>() = FMath::Max(BranchMinMaxPScales[BranchIndex].Get<1>(), PointScaleAtIndex);
+		}
+	}
+
+	// Remap BranchGradient and PlantGradient
+	const FVector2f NormalRange(0.0f, 1.0f);
+	const FVector2f PlantPScaleRange(MinimumPScale, MaximumPScale);
+
+	for (int32 BranchIndex = 0; BranchIndex < BranchFacade.GetElementCount(); ++BranchIndex)
+	{
+		const TArray<int32>& BranchPoints = BranchFacade.GetPoints(BranchIndex);
+		const FVector2f BranchPScaleRange(BranchMinMaxPScales[BranchIndex].Get<0>(), BranchMinMaxPScales[BranchIndex].Get<1>());
+
+		for (const int32& PointIndex : BranchPoints)
+		{
+			const float PlantGradientAtPointIndex = FMath::GetMappedRangeValueClamped(PlantPScaleRange, NormalRange,
+				PointFacadeOut.GetPointScale(PointIndex));
+			PointFacadeOut.SetPlantGradient(PointIndex, PlantGradientAtPointIndex);
+
+			const float BranchGradientAtPointIndex = FMath::GetMappedRangeValueClamped(BranchPScaleRange, NormalRange,
+				PointFacadeOut.GetPointScale(PointIndex));
+			PointFacadeOut.SetBranchGradient(PointIndex, BranchGradientAtPointIndex);
+		}
+	}
+}
+
+void FMegaPlantCarve::CarveFromTop(FManagedArrayCollection& OutCollection, const FManagedArrayCollection& SourceCollection, const float Carve,
+                                   const ECarveBasis CarveBasis)
+{
+	TRACE_CPUPROFILER_EVENT_SCOPE(MegaPlantCarve::CarveFromTop);
+
+	const float CarveDistance = 1.0f - Carve;
+
+	const MegaPlants::Facades::FBranchFacade BranchFacadeSource(SourceCollection);
+	const MegaPlants::Facades::FPointFacade PointFacadeSource(SourceCollection);
+	const MegaPlants::Facades::FFoliageFacade FoliageFacadeSource(SourceCollection);
+
+	MegaPlants::Facades::FBranchFacade BranchFacadeOut(OutCollection);
+	MegaPlants::Facades::FPointFacade PointFacadeOut(OutCollection);
+	MegaPlants::Facades::FFoliageFacade FoliageFacadeOut(OutCollection);
+
+	if (PointFacadeSource.GetElementCount() == 0 || BranchFacadeSource.GetElementCount() == 0)
+	{
+		return;
+	}
+
+	const int32 NumOfBranchesInSource = BranchFacadeSource.GetElementCount();
+	TArray<bool> PointsToRemove;
+	PointsToRemove.Init(false, PointFacadeSource.GetElementCount());
+	TArray<bool> BranchesToRemove;
+	BranchesToRemove.Init(false, NumOfBranchesInSource);
+	TArray<bool> FoliageInstancesToRemove;
+	FoliageInstancesToRemove.Init(false, FoliageFacadeSource.NumFoliageEntries());
+
+	TMap<int32, int32> BranchNumbersToBranchIDs;
+	TMap<int32, float> BranchNumbersToLengthFromRoots;
+	ComputeMetadata(BranchNumbersToBranchIDs, BranchNumbersToLengthFromRoots, BranchFacadeSource, PointFacadeSource);
+
+	// Compute carve values according to carve basis metric
+	TArray<float> CarveValues;
+	CarveValues.Reserve(PointFacadeSource.GetElementCount());
+
+	float MinimumCarveValue = TNumericLimits<float>::Max();
+	float MaximumCarveValue = TNumericLimits<float>::Min();
+
+	for (int32 i = 0; i < PointFacadeSource.GetElementCount(); ++i)
+	{
+		switch (CarveBasis)
+		{
+		case ECarveBasis::LengthFromRoot:
+			CarveValues.Add(PointFacadeSource.GetLengthFromRoot(i));
+			break;
+
+		case ECarveBasis::FromBottom:
+			CarveValues.Add(PointFacadeSource.GetLengthFromRoot(i));
+			break;
+
+		case ECarveBasis::ZPosition:
+			CarveValues.Add(PointFacadeSource.GetPosition(i).Z);
+			break;
+
+		case ECarveBasis::Radius:
+			CarveValues.Add(PointFacadeSource.GetPointScale(i));
+			break;
+		}
+
+		MinimumCarveValue = FMath::Min(MinimumCarveValue, CarveValues[i]);
+		MaximumCarveValue = FMath::Max(MaximumCarveValue, CarveValues[i]);
+	}
+
+	FVector2f MappedRange(0.0f, 1.0f);
+	if (CarveBasis == ECarveBasis::Radius)
+	{
+		MappedRange = FVector2f(1.0f, 0.0f);
+	}
+
+	for (int32 i = 0; i < CarveValues.Num(); ++i)
+	{
+		CarveValues[i] = FMath::GetMappedRangeValueClamped(FVector2f(MinimumCarveValue, MaximumCarveValue), MappedRange, CarveValues[i]);
+	}
+
+	// Compute branches, foliage instances and points to remove
+	const int32 TrunkIndex = BranchFacadeSource.GetTrunkIndex();
+	for (int32 BranchIndex = 0; BranchIndex < NumOfBranchesInSource; ++BranchIndex)
+	{
+		if (BranchesToRemove[BranchIndex])
+		{
+			continue;
+		}
+
+		const TArray<int32>& BranchPoints = BranchFacadeSource.GetPoints(BranchIndex);
+		const int32 NumOfPointsInCurrentBranch = BranchPoints.Num();
+
+		if (NumOfPointsInCurrentBranch == 0)
+		{
+			BranchesToRemove[BranchIndex] = true;
+			continue;
+		}
+
+		// Compute PScale for the point in parent branch that is closest to the first point
+		// (For non-trunk branches)
+		float FirstPointTargetPScale = PointFacadeSource.GetPointScale(BranchPoints[0]);
+		if (BranchIndex != TrunkIndex) [[likely]]
+		{
+			const FVector3f& FirstPointPosition = PointFacadeSource.GetPosition(BranchPoints[0]);
+
+			const int32 ParentBranchNumber = BranchFacadeSource.GetParents(BranchIndex).Last();
+			const int32 ParentBranchIndex = BranchNumbersToBranchIDs[ParentBranchNumber];
+			const TArray<int32>& ParentBranchPoints = BranchFacadeSource.GetPoints(ParentBranchIndex);
+
+			int32 ClosestPointIndex = 0;
+			float MinimumDistance = FLT_MAX;
+			for (int32 PointIndex = 0; PointIndex < ParentBranchPoints.Num(); ++PointIndex)
+			{
+				const FVector3f& ParentPointPosition = PointFacadeSource.GetPosition(ParentBranchPoints[PointIndex]);
+				if (const float Distance = FVector3f::Distance(FirstPointPosition, ParentPointPosition); Distance < MinimumDistance)
+				{
+					MinimumDistance = Distance;
+					ClosestPointIndex = PointIndex;
+				}
+			}
+
+			FirstPointTargetPScale = PointFacadeOut.GetPointScale(ParentBranchPoints[ClosestPointIndex]);
+		}
+
+		int EndIndex = 1;
+		while (EndIndex < NumOfPointsInCurrentBranch && (CarveValues[BranchPoints[EndIndex]] < CarveDistance + 0.0001f))
+		{
+			EndIndex++;
+		}
+		EndIndex--;
+
+		if (EndIndex == NumOfPointsInCurrentBranch - 1)
+		{
+			const int32 LastPointIndex = BranchPoints[EndIndex];
+			const float LastPointScale = PointFacadeSource.GetPointScale(LastPointIndex);
+			UpdatePointScales(PointFacadeOut, PointFacadeSource, BranchPoints, LastPointScale, LastPointIndex, 1.0f, EndIndex,
+				FirstPointTargetPScale);
+			continue;
+		}
+
+		// Remove all branches, child branches and foliage instances beyond last point
+		for (const int32& BranchNumber : BranchFacadeSource.GetChildren(BranchIndex))
+		{
+			if (BranchNumbersToLengthFromRoots.Contains(BranchNumber) && BranchNumbersToLengthFromRoots[BranchNumber] > PointFacadeSource.
+				GetLengthFromRoot(BranchPoints[EndIndex]))
+			{
+				const int32 Index = BranchNumbersToBranchIDs[BranchNumber];
+				BranchesToRemove[Index] = true;
+				for (const int32& P : BranchFacadeSource.GetPoints(Index))
+				{
+					PointsToRemove[P] = true;
+				}
+				for (const int32& FId : FoliageFacadeSource.GetFoliageEntryIdsForBranch(Index))
+				{
+					FoliageInstancesToRemove[FId] = true;
+				}
+
+				const TArray<int32>& BranchChildren = BranchFacadeSource.GetChildren(Index);
+				for (int32 k = 0; k < BranchChildren.Num(); ++k)
+				{
+					if (!BranchNumbersToBranchIDs.Contains(BranchChildren[k]))
+					{
+						continue;
+					}
+
+					const int32 ChildIndex = BranchNumbersToBranchIDs[BranchChildren[k]];
+					BranchesToRemove[ChildIndex] = true;
+					for (const int32& P : BranchFacadeSource.GetPoints(ChildIndex))
+					{
+						PointsToRemove[P] = true;
+					}
+					for (const int32& FId : FoliageFacadeSource.GetFoliageEntryIdsForBranch(ChildIndex))
+					{
+						FoliageInstancesToRemove[FId] = true;
+					}
+				}
+			}
+		}
+
+		for (const int32& FoliageIndex : FoliageFacadeSource.GetFoliageEntryIdsForBranch(BranchIndex))
+		{
+			if (FoliageFacadeSource.GetLengthFromRoot(FoliageIndex) > PointFacadeSource.GetLengthFromRoot(BranchPoints[EndIndex]))
+			{
+				FoliageInstancesToRemove[FoliageIndex] = true;
+			}
+		}
+
+		// Keep next point to adjust its position, remove all others
+		++EndIndex;
+		for (int32 j = EndIndex + 1; j < NumOfPointsInCurrentBranch; ++j)
+		{
+			PointsToRemove[BranchPoints[j]] = true;
+		}
+
+		TArray<int32> UpdatedBranchPoints = TArray<int32>(&BranchPoints[0], EndIndex + 1);
+		BranchFacadeOut.SetPoints(BranchIndex, UpdatedBranchPoints);
+
+		// Adjust last point
+		const int32 LastPointIndex = BranchPoints[EndIndex];
+		const int32 PreviousPointIndex = BranchPoints[EndIndex - 1];
+		const FVector3f& LastPointPosition = PointFacadeSource.GetPosition(LastPointIndex);
+		const FVector3f& PreviousPointPosition = PointFacadeSource.GetPosition(PreviousPointIndex);
+		const float LastPointCarveValue = CarveValues[LastPointIndex];
+		const float PreviousPointCarveValue = CarveValues[PreviousPointIndex];
+		const float BlendValue = FMath::Clamp(
+			FMath::GetMappedRangeValueClamped(
+				FVector2f(PreviousPointCarveValue, LastPointCarveValue),
+				FVector2f(0.0f, 1.0f),
+				CarveDistance),
+			0.1f,
+			1.0f);
+		const FVector3f Position = FMath::Lerp(PreviousPointPosition, LastPointPosition, BlendValue);
+		const float LengthFromRoot = FMath::Lerp(PointFacadeSource.GetLengthFromRoot(PreviousPointIndex),
+			PointFacadeSource.GetLengthFromRoot(LastPointIndex),
+			BlendValue);
+		const float PointScale = FMath::Lerp(PointFacadeSource.GetPointScale(PreviousPointIndex),
+			PointFacadeSource.GetPointScale(LastPointIndex),
+			BlendValue);
+		PointFacadeOut.ModifyPositions()[LastPointIndex] = Position;
+		PointFacadeOut.ModifyLengthFromRoots()[LastPointIndex] = LengthFromRoot;
+		PointFacadeOut.ModifyLengthFromSeeds()[LastPointIndex] = LengthFromRoot;
+		PointFacadeOut.ModifyPointScales()[LastPointIndex] = PointScale;
+
+
+		// Update point scales
+		const float CarveRatio = (EndIndex + BlendValue) / NumOfPointsInCurrentBranch;
+		UpdatePointScales(PointFacadeOut, PointFacadeSource, BranchPoints, PointScale, LastPointIndex, CarveRatio, EndIndex, FirstPointTargetPScale);
+	}
+
+	RemoveEntriesAndRecomputeAttributes(OutCollection, SourceCollection, PointsToRemove, BranchesToRemove, FoliageInstancesToRemove);
+}
+
+void FMegaPlantCarve::CarveFromBottom(FManagedArrayCollection& OutCollection, const FManagedArrayCollection& SourceCollection, const float Carve)
+{
+	TRACE_CPUPROFILER_EVENT_SCOPE(MegaPlantCarve::CarveFromBottom);
+
+	const MegaPlants::Facades::FBranchFacade BranchFacadeSource(SourceCollection);
+	const MegaPlants::Facades::FPointFacade PointFacadeSource(SourceCollection);
+	const MegaPlants::Facades::FFoliageFacade FoliageFacadeSource(SourceCollection);
+	MegaPlants::Facades::FBranchFacade BranchFacadeOut(OutCollection);
+	MegaPlants::Facades::FPointFacade PointFacadeOut(OutCollection);
+	MegaPlants::Facades::FFoliageFacade FoliageFacadeOut(OutCollection);
+
+	if (PointFacadeSource.GetElementCount() == 0 || BranchFacadeSource.GetElementCount() == 0)
+	{
+		return;
+	}
+
+	const int32 NumOfBranchesInSource = BranchFacadeSource.GetElementCount();
+	TArray<bool> PointsToRemove;
+	PointsToRemove.Init(false, PointFacadeSource.GetElementCount());
+	TArray<bool> BranchesToRemove;
+	BranchesToRemove.Init(false, NumOfBranchesInSource);
+	TArray<bool> FoliageInstancesToRemove;
+	FoliageInstancesToRemove.Init(false, FoliageFacadeSource.NumFoliageEntries());
+
+	TMap<int32, int32> BranchNumbersToBranchIDs;
+	TMap<int32, float> BranchNumbersToLengthFromRoots;
+	ComputeMetadata(BranchNumbersToBranchIDs, BranchNumbersToLengthFromRoots, BranchFacadeSource, PointFacadeSource);
+
+	// Carve along trunk based on length from root
+	// Compute branches, foliage instances and points to remove
+	const int32 TrunkIndex = BranchFacadeSource.GetTrunkIndex();
+	const TArray<int32>& TrunkPoints = BranchFacadeSource.GetPoints(TrunkIndex);
+	const int32 NumOfPointsInTrunk = TrunkPoints.Num();
+	if (TrunkPoints.Num() < 2)
+	{
+		return;
+	}
+
+	const float LengthFromRootOfFirstPoint = PointFacadeSource.GetLengthFromRoot(TrunkPoints[0]);
+	const float LengthFromRootOfLastPoint = PointFacadeSource.GetLengthFromRoot(TrunkPoints[NumOfPointsInTrunk - 1]);
+	const float CarveLengthFromRoot = FMath::GetMappedRangeValueClamped(FVector2f(0.0f, 1.0f),
+		FVector2f(LengthFromRootOfFirstPoint, LengthFromRootOfLastPoint), Carve);
+
+	int EndIndex = 0;
+	while (EndIndex < NumOfPointsInTrunk && PointFacadeSource.GetLengthFromRoot(TrunkPoints[EndIndex]) < CarveLengthFromRoot)
+	{
+		EndIndex++;
+	}
+	EndIndex--;
+
+	// Remove all branches, child branches and foliage instances before current point
+	for (const int32& BranchNumber : BranchFacadeSource.GetChildren(TrunkIndex))
+	{
+		if (BranchNumbersToLengthFromRoots.Contains(BranchNumber) && BranchNumbersToLengthFromRoots[BranchNumber] < CarveLengthFromRoot)
+		{
+			const int32 Index = BranchNumbersToBranchIDs[BranchNumber];
+			BranchesToRemove[Index] = true;
+			for (const int32& P : BranchFacadeSource.GetPoints(Index))
+			{
+				PointsToRemove[P] = true;
+			}
+			for (const int32& FId : FoliageFacadeSource.GetFoliageEntryIdsForBranch(Index))
+			{
+				FoliageInstancesToRemove[FId] = true;
+			}
+
+			const TArray<int32>& BranchChildren = BranchFacadeSource.GetChildren(Index);
+			for (int32 k = 0; k < BranchChildren.Num(); ++k)
+			{
+				if (!BranchNumbersToBranchIDs.Contains(BranchChildren[k]))
+				{
+					continue;
+				}
+
+				const int32 ChildIndex = BranchNumbersToBranchIDs[BranchChildren[k]];
+				BranchesToRemove[ChildIndex] = true;
+				for (const int32& P : BranchFacadeSource.GetPoints(ChildIndex))
+				{
+					PointsToRemove[P] = true;
+				}
+				for (const int32& FId : FoliageFacadeSource.GetFoliageEntryIdsForBranch(ChildIndex))
+				{
+					FoliageInstancesToRemove[FId] = true;
+				}
+			}
+		}
+	}
+
+	for (const int32& FoliageIndex : FoliageFacadeSource.GetFoliageEntryIdsForBranch(TrunkIndex))
+	{
+		if (FoliageFacadeSource.GetLengthFromRoot(FoliageIndex) < CarveLengthFromRoot)
+		{
+			FoliageInstancesToRemove[FoliageIndex] = true;
+		}
+	}
+
+	// Keep current point to adjust its position, remove all others
+	for (int32 j = 0; j < EndIndex; ++j)
+	{
+		PointsToRemove[TrunkPoints[j]] = true;
+	}
+
+	TArray<int32> UpdatedTrunkPoints = TArray<int32>(&TrunkPoints[EndIndex], NumOfPointsInTrunk - EndIndex);
+	BranchFacadeOut.SetPoints(TrunkIndex, UpdatedTrunkPoints);
+
+	// Adjust End point
+	const int32 CurrentPointIndex = TrunkPoints[EndIndex];
+	const int32 NextPointIndex = TrunkPoints[EndIndex + 1];
+	const FVector3f& PositionOfCurrentPoint = PointFacadeSource.GetPosition(CurrentPointIndex);
+	const FVector3f& PositionOfNextPoint = PointFacadeSource.GetPosition(NextPointIndex);
+	const float LengthFromRootOfCurrentPoint = PointFacadeSource.GetLengthFromRoot(CurrentPointIndex);
+	const float LengthFromRootOfNextPoint = PointFacadeSource.GetLengthFromRoot(NextPointIndex);
+
+	const float BlendValue = FMath::Clamp(
+		FMath::GetMappedRangeValueClamped(
+			FVector2f(LengthFromRootOfCurrentPoint, LengthFromRootOfNextPoint),
+			FVector2f(0.0f, 1.0f),
+			CarveLengthFromRoot),
+		0.1f,
+		1.0f);
+	const FVector3f Position = FMath::Lerp(PositionOfCurrentPoint, PositionOfNextPoint, BlendValue);
+	const float PointScale = FMath::Lerp(PointFacadeSource.GetPointScale(CurrentPointIndex),
+		PointFacadeSource.GetPointScale(NextPointIndex),
+		BlendValue);
+	PointFacadeOut.ModifyPositions()[CurrentPointIndex] = Position;
+	PointFacadeOut.ModifyLengthFromRoots()[CurrentPointIndex] = CarveLengthFromRoot;
+	PointFacadeOut.ModifyLengthFromSeeds()[CurrentPointIndex] = CarveLengthFromRoot;
+	PointFacadeOut.ModifyPointScales()[CurrentPointIndex] = PointScale;
+
+	RemoveEntriesAndRecomputeAttributes(OutCollection, SourceCollection, PointsToRemove, BranchesToRemove, FoliageInstancesToRemove);
+
+	// Move all points down
+	const FVector3f Delta(0.0f, 0.0f, Position.Z);
+
+	for (int32 PointIndex = 0; PointIndex < PointFacadeOut.GetElementCount(); ++PointIndex)
+	{
+		PointFacadeOut.ModifyPositions()[PointIndex] = PointFacadeOut.GetPosition(PointIndex) - Delta;
+	}
+
+	for (int32 FoliageIndex = 0; FoliageIndex < FoliageFacadeOut.GetElementCount(); ++FoliageIndex)
+	{
+		const FVector3f UpdatedPivotPoint = FoliageFacadeOut.GetPivotPoint(FoliageIndex) - Delta;
+		FoliageFacadeOut.SetPivotPoint(FoliageIndex, UpdatedPivotPoint);
+	}
+
+	// Recompute length from root
+	for (int32 PointIndex = 0; PointIndex < PointFacadeOut.GetElementCount(); ++PointIndex)
+	{
+		PointFacadeOut.ModifyLengthFromRoots()[PointIndex] = FMath::Max(PointFacadeOut.GetLengthFromRoot(PointIndex) - CarveLengthFromRoot, 0.0f);
+		PointFacadeOut.ModifyLengthFromSeeds()[PointIndex] = FMath::Max(PointFacadeOut.GetLengthFromSeed(PointIndex) - CarveLengthFromRoot, 0.0f);
+	}
+
+	for (int32 FoliageIndex = 0; FoliageIndex < FoliageFacadeOut.GetElementCount(); ++FoliageIndex)
+	{
+		const float UpdatedLengthFromRoot = FMath::Max(FoliageFacadeOut.GetLengthFromRoot(FoliageIndex) - CarveLengthFromRoot);
+		FoliageFacadeOut.SetLengthFromRoot(FoliageIndex, UpdatedLengthFromRoot);
+	}
+}
+
+void FMegaPlantCarve::ComputeMetadata(TMap<int32, int32>& OutBranchNumbersToBranchIDs, TMap<int32, float>& OutBranchNumbersToLengthFromRoots,
+                                      const MegaPlants::Facades::FBranchFacade& BranchFacadeSource,
+                                      const MegaPlants::Facades::FPointFacade& PointFacadeSource)
+{
+	TRACE_CPUPROFILER_EVENT_SCOPE(MegaPlantCarve::ComputeMetadata);
+
+	const int32 NumOfBranchesInSource = BranchFacadeSource.GetElementCount();
+	OutBranchNumbersToBranchIDs.Reserve(NumOfBranchesInSource);
+	OutBranchNumbersToLengthFromRoots.Reserve(NumOfBranchesInSource);
+
+	for (int32 BranchIndex = 0; BranchIndex < NumOfBranchesInSource; ++BranchIndex)
+	{
+		const int32 BranchNumber = BranchFacadeSource.GetBranchNumber(BranchIndex);
+		const int32 FirstPointIndex = BranchFacadeSource.GetPoints(BranchIndex)[0];
+		const float BranchLengthFromRoot = PointFacadeSource.GetLengthFromRoot(FirstPointIndex);
+		OutBranchNumbersToBranchIDs.Add(BranchNumber, BranchIndex);
+		OutBranchNumbersToLengthFromRoots.Add(BranchNumber, BranchLengthFromRoot);
+	}
+}
+
+void FMegaPlantCarve::RemoveEntriesAndRecomputeAttributes(FManagedArrayCollection& OutCollection, const FManagedArrayCollection& SourceCollection,
+                                                          TArray<bool>& PointsToRemove, TArray<bool>& BranchesToRemove,
+                                                          TArray<bool>& FoliageInstancesToRemove)
+{
+	TRACE_CPUPROFILER_EVENT_SCOPE(MegaPlantCarve::RemoveEntriesAndRecomputeAttributes);
+
+	const MegaPlants::Facades::FBranchFacade BranchFacadeSource(SourceCollection);
+	const MegaPlants::Facades::FPointFacade PointFacadeSource(SourceCollection);
+
+	MegaPlants::Facades::FBranchFacade BranchFacadeOut(OutCollection);
+	MegaPlants::Facades::FPointFacade PointFacadeOut(OutCollection);
+
+	MegaPlants::Facades::FRemoveEntriesResult RemoveEntriesResult = MegaPlants::Facades::FTreeFacade::RemoveEntriesAndReIndexAttributes(
+		OutCollection, PointsToRemove, BranchesToRemove, FoliageInstancesToRemove);
+
+	TMap<int32, int32> BranchesNewIDsToOldIDs;
+	BranchesNewIDsToOldIDs.Reserve(RemoveEntriesResult.BranchesOldIDsToNewIDs.Num());
+	for (const TPair<int32, int32>& Pair : RemoveEntriesResult.BranchesOldIDsToNewIDs)
+	{
+		BranchesNewIDsToOldIDs.Add(Pair.Value, Pair.Key);
+	}
+
+	RecomputeAttributes(PointFacadeOut, PointFacadeSource, BranchFacadeOut, BranchFacadeSource, BranchesNewIDsToOldIDs);
+}
diff --git a/Engine/Plugins/Experimental/ProceduralVegetationEditor/Source/MegaPlants/Private/Implementations/MegaPlantCarve.h b/Engine/Plugins/Experimental/ProceduralVegetationEditor/Source/MegaPlants/Private/Implementations/MegaPlantCarve.h
new file mode 100644
--- /dev/null
+++ b/Engine/Plugins/Experimental/ProceduralVegetationEditor/Source/MegaPlants/Private/Implementations/MegaPlantCarve.h
@@ -0,0 +1,64 @@
+// Copyright Epic Games, Inc. All Rights Reserved.
+
+#pragma once
+
+#include "CoreMinimal.h"
+#include "Facades/MegaPlantPointFacade.h"
+#include "Facades/MegaPlantBranchFacade.h"
+#include "Facades/MegaPlantFoliageFacade.h"
+#include "GeometryCollection/ManagedArrayCollection.h"
+
+#include "MegaPlantCarve.generated.h"
+
+UENUM(BlueprintType)
+enum class ECarveBasis : uint8
+{
+	LengthFromRoot UMETA(DisplayName = "Length from root"),
+	FromBottom UMETA(DisplayName = "From Bottom"),
+	ZPosition UMETA(DisplayName = "Z Position"),
+	Radius UMETA(DisplayName = "Radius")
+};
+
+USTRUCT()
+struct FMegaPlantCarveSettings
+{
+	GENERATED_BODY()
+
+	UPROPERTY(EditAnywhere, Category="Carve")
+	ECarveBasis CarveBasis = ECarveBasis::LengthFromRoot;
+
+	UPROPERTY(EditAnywhere, Category="Carve", meta=(ClampMin=0.0f, ClampMax=1.0f, UIMin=0.0f, UIMax=1.0f))
+	float Carve = 0.0f;
+};
+
+struct FMegaPlantCarve
+{
+	static void ApplyCarve(FManagedArrayCollection& OutCollection, const FManagedArrayCollection& SourceCollection, const ECarveBasis CarveBasis,
+	                       const float Carve);
+
+private:
+	static void UpdatePointScales(MegaPlants::Facades::FPointFacade& PointFacadeOut, const MegaPlants::Facades::FPointFacade& PointFacadeSource,
+	                              const TArray<int>& BranchPoints, const float LastPointScale, const int32 LastPointIndex, const float CarveRatio,
+	                              const int32 EndIndex, const float FirstPointTargetPScale);
+
+	static void UpdateBranchLinkedIndices(MegaPlants::Facades::FBranchFacade& BranchFacadeOut, MegaPlants::Facades::FFoliageFacade& FoliageFacadeOut,
+	                                      const TMap<int32, int32>& PointsOldIDsToNewIDs, const TArray<bool>& PointsRemoved,
+	                                      const TMap<int32, int32>& FoliageInstancesOldIDsToNewIDs, const TArray<bool>& FoliageInstancesRemoved);
+
+	static void RecomputeAttributes(MegaPlants::Facades::FPointFacade& PointFacadeOut, const MegaPlants::Facades::FPointFacade& PointFacadeSource,
+	                                const MegaPlants::Facades::FBranchFacade& BranchFacade,
+	                                const MegaPlants::Facades::FBranchFacade& BranchFacadeSource, const TMap<int32, int32>& BranchesNewIDsToOldIDs);
+
+	static void CarveFromTop(FManagedArrayCollection& OutCollection, const FManagedArrayCollection& SourceCollection, const float Carve,
+	                         const ECarveBasis CarveBasis);
+
+	static void CarveFromBottom(FManagedArrayCollection& OutCollection, const FManagedArrayCollection& SourceCollection, const float Carve);
+
+	static void ComputeMetadata(TMap<int32, int32>& OutBranchNumbersToBranchIDs, TMap<int32, float>& OutBranchNumbersToLengthFromRoots,
+	                            const MegaPlants::Facades::FBranchFacade& BranchFacadeSource,
+	                            const MegaPlants::Facades::FPointFacade& PointFacadeSource);
+
+	static void RemoveEntriesAndRecomputeAttributes(FManagedArrayCollection& OutCollection, const FManagedArrayCollection& SourceCollection,
+	                                                TArray<bool>& PointsToRemove, TArray<bool>& BranchesToRemove,
+	                                                TArray<bool>& FoliageInstancesToRemove);
+};
diff --git a/Engine/Plugins/Experimental/ProceduralVegetationEditor/Source/MegaPlants/Private/Implementations/MegaPlantFoliage.cpp b/Engine/Plugins/Experimental/ProceduralVegetationEditor/Source/MegaPlants/Private/Implementations/MegaPlantFoliage.cpp
new file mode 100644
--- /dev/null
+++ b/Engine/Plugins/Experimental/ProceduralVegetationEditor/Source/MegaPlants/Private/Implementations/MegaPlantFoliage.cpp
@@ -0,0 +1,419 @@
+// Copyright Epic Games, Inc. All Rights Reserved.
+
+#include "MegaPlantFoliage.h"
+
+#include "MegaPlantsModule.h"
+#include "Facades/MegaPlantAttributesNames.h"
+#include "Facades/MegaPlantBranchFacade.h"
+#include "Facades/MegaPlantFoliageFacade.h"
+#include "Facades/MegaPlantPointFacade.h"
+
+void FMegaPlantFoliage::DistributeFoliage(FManagedArrayCollection& OutCollection, const FManagedArrayCollection& SourceCollection,
+                                          const FDistributionSettings& DistributionSettings, const FScaleSettings& ScaleSettings,
+                                          const FVectorSettings& VectorSettings, const FPhyllotaxySettings& PhyllotaxySettings,
+                                          const FMiscSettings& MiscSettings)
+{
+	MegaPlants::Facades::FFoliageFacade FoliageFacadeOutput(OutCollection);
+
+	if (const int32 NumFoliageMeshes = FoliageFacadeOutput.NumFoliageNames(); NumFoliageMeshes == 0)
+	{
+		UE_LOG(LogMegaPlants, Warning, TEXT("There are no foliage meshes available in input for distribution"));	
+		return;
+	}
+
+	FRandomStream RandomStream;
+	RandomStream.Initialize(MiscSettings.RandomSeed);
+
+	MegaPlants::Facades::FBranchFacade BranchFacadeSource(SourceCollection);
+	MegaPlants::Facades::FPointFacade PointFacadeSource(SourceCollection);
+
+	TManagedArrayAccessor<TArray<FVector3f>> BudDirectionAttribute(OutCollection, MegaPlants::AttributeNames::BudDirection,
+	                                                               MegaPlants::GroupNames::PointGroup);
+	TManagedArrayAccessor<float> BranchGradientsAttribute(OutCollection, MegaPlants::AttributeNames::BranchGradient,
+	                                                      MegaPlants::GroupNames::PointGroup);
+	TManagedArrayAccessor<TArray<float>> LeafPhyllotaxyAttribute(OutCollection, MegaPlants::AttributeNames::LeafPhyllotaxy,
+	                                                             MegaPlants::GroupNames::DetailsGroup);
+	TManagedArrayAccessor<TArray<float>> BudHormoneLevels(OutCollection, MegaPlants::AttributeNames::BudHormoneLevels,
+	                                                      MegaPlants::GroupNames::PointGroup);
+
+	float PhyllotaxyFormationLeaf = LeafPhyllotaxyAttribute[0][1];
+	float MinBudsLeaf = LeafPhyllotaxyAttribute[0][3];
+	float MaxBudsLeaf = LeafPhyllotaxyAttribute[0][4];
+	int32 ResetPhyllotaxyLeaf = static_cast<int>(LeafPhyllotaxyAttribute[0][6]);
+	float PhyllotaxyOffsetLeaf = LeafPhyllotaxyAttribute[0][7];
+
+	if (PhyllotaxySettings.OverridePhyllotaxy)
+	{
+		switch (PhyllotaxySettings.PhyllotaxyType)
+		{
+		case EPhyllotaxyType::Alternate: PhyllotaxyFormationLeaf = 180 + PhyllotaxySettings.PhyllotaxyAdditionalAngle;
+			MinBudsLeaf = MaxBudsLeaf = 1;
+			break;
+		case EPhyllotaxyType::Opposite: PhyllotaxyFormationLeaf = PhyllotaxySettings.PhyllotaxyAdditionalAngle;
+			MinBudsLeaf = MaxBudsLeaf = 2;
+			break;
+		case EPhyllotaxyType::Decussate: PhyllotaxyFormationLeaf = 90 + PhyllotaxySettings.PhyllotaxyAdditionalAngle;
+			MinBudsLeaf = MaxBudsLeaf = 2;
+			break;
+		case EPhyllotaxyType::Whorled: PhyllotaxyFormationLeaf = 90 + PhyllotaxySettings.PhyllotaxyAdditionalAngle;
+			MinBudsLeaf = PhyllotaxySettings.MinimumNodeBuds;
+			MaxBudsLeaf = PhyllotaxySettings.MaximumNodeBuds;
+			break;
+		case EPhyllotaxyType::Spiral:
+			{
+				// These leaf arrangement configurations where Distichous is two leaves, Tristichous is three etc.
+				switch (PhyllotaxySettings.PhyllotaxyFormation)
+				{
+				case EPhyllotaxyFormation::Distichous: PhyllotaxyFormationLeaf = 180 + PhyllotaxySettings.PhyllotaxyAdditionalAngle;
+					break;
+				case EPhyllotaxyFormation::Tristichous: PhyllotaxyFormationLeaf = 120 + PhyllotaxySettings.PhyllotaxyAdditionalAngle;
+					break;
+				case EPhyllotaxyFormation::Pentastichous: PhyllotaxyFormationLeaf = 144 + PhyllotaxySettings.PhyllotaxyAdditionalAngle;
+					break;
+				case EPhyllotaxyFormation::Octastichous: PhyllotaxyFormationLeaf = 135 + PhyllotaxySettings.PhyllotaxyAdditionalAngle;
+					break;
+				case EPhyllotaxyFormation::Parastichous: PhyllotaxyFormationLeaf = PhyllotaxySettings.PhyllotaxyAdditionalAngle;
+					break;
+				}
+
+				MinBudsLeaf = MaxBudsLeaf = 1;
+				break;
+			}
+		}
+	}
+
+	const int32 NumBranches = BranchFacadeSource.GetElementCount();
+
+	if (DistributionSettings.OverrideDistribution)
+	{
+		TRACE_CPUPROFILER_EVENT_SCOPE(MegaPlantFoliageDistributorNode::OverrideDistribution);
+
+		// Compute attachment points in terms of normalized distances
+		OutCollection.EmptyGroup(MegaPlants::GroupNames::FoliageGroup);
+		for (int32 BranchIndex = 1; BranchIndex < NumBranches; ++BranchIndex)
+		{
+			TArray<float> NormalizedAttachmentPoints;
+			TArray<float> NormalizedAttachmentPointsFinal;
+			TArray<FVector3f> AttachmentPointsPositions;
+			TArray<float> AttachmentPointsLengthsFromRoot;
+			TArray<FVector3f> AttachmentPointsUpVectors;
+			TArray<FVector3f> AttachmentPointsNormalVectors;
+			TArray<float> AttachmentPointsScales;
+			TArray<int32> BranchPoints = BranchFacadeSource.GetPoints(BranchIndex);
+			const float FirstPointLengthFromRoot = PointFacadeSource.GetLengthFromRoot(BranchPoints[0]);
+			const float LastPointLengthFromRoot = PointFacadeSource.GetLengthFromRoot(BranchPoints.Last());
+			const float BranchLength = LastPointLengthFromRoot - FirstPointLengthFromRoot;
+			float Increment = DistributionSettings.InstanceSpacing / BranchLength;
+			const float AdjustedMaxPerBranch = DistributionSettings.MaxPerBranch == -1 ? BranchPoints.Num() : DistributionSettings.MaxPerBranch;
+			float LookUp = 0.0f;
+
+			for (int32 i = 0; i < AdjustedMaxPerBranch; ++i)
+			{
+				if (LookUp > 1.0f)
+				{
+					break;
+				}
+
+				float InstanceSpacingRampValue = DistributionSettings.InstanceSpacingRamp
+					                                 ? DistributionSettings.InstanceSpacingRamp->GetRichCurveConst()->Eval(1.0f - LookUp)
+					                                 : Increment;
+				// The increment is nudged slightly before blending with the weighted ramp value, the number is chosen based on observation
+				float AdjustedIncrement = Increment + (InstanceSpacingRampValue * (Increment * 10));
+				AdjustedIncrement = FMath::Lerp(Increment, AdjustedIncrement, DistributionSettings.InstanceSpacingRampEffect);
+
+				NormalizedAttachmentPoints.Add(1.0f - LookUp);
+				LookUp += AdjustedIncrement;
+			}
+
+			Algo::Reverse(NormalizedAttachmentPoints);
+
+			// Compute Position, Angles, Scale and LengthFromRoot for attachment points
+
+			// Compute based on first point attributes
+			const FVector3f InitialApicalDirection(BudDirectionAttribute[BranchPoints[0]][0]);
+			FVector3f InitialAxillaryDirection(BudDirectionAttribute[BranchPoints[0]][1]);
+			const FVector3f InitialUpVector(BudDirectionAttribute[BranchPoints[0]][5]);
+
+			if (ResetPhyllotaxyLeaf == 1)
+			{
+				InitialAxillaryDirection = FVector3f::CrossProduct(InitialApicalDirection, InitialUpVector);
+			}
+
+			if (PhyllotaxyOffsetLeaf > 0.01f)
+			{
+				const FQuat4f RotationQuat(InitialApicalDirection.GetSafeNormal(), FMath::DegreesToRadians(PhyllotaxyOffsetLeaf));
+				InitialAxillaryDirection = RotationQuat.RotateVector(InitialAxillaryDirection);
+			}
+
+			int32 AxillaryRotationIteration = 0;
+			for (int32 Idx = 0; Idx < NormalizedAttachmentPoints.Num(); ++Idx)
+			{
+				for (int32 BranchPointIndex = 0; BranchPointIndex < BranchPoints.Num() - 1; ++BranchPointIndex)
+				{
+					float AxillaryRotationDegree = FMath::DegreesToRadians(PhyllotaxyFormationLeaf * AxillaryRotationIteration);
+
+					const int32 CurrentPointIndex = BranchPoints[BranchPointIndex];
+					const int32 NextPointIndex = BranchPoints[BranchPointIndex + 1];
+					const FVector3f CurrentPointPosition = PointFacadeSource.GetPosition(CurrentPointIndex);
+					const FVector3f NextPointPosition = PointFacadeSource.GetPosition(NextPointIndex);
+					float EthyleneLevelAtCurrentPoint = BudHormoneLevels[CurrentPointIndex][4];
+					float EthyleneLevelAtNextPoint = BudHormoneLevels[NextPointIndex][4];
+					const float CurrentPointLengthFromRoot = PointFacadeSource.GetLengthFromRoot(CurrentPointIndex);
+					const float NextPointLengthFromRoot = PointFacadeSource.GetLengthFromRoot(NextPointIndex);
+					const float CurrentPointLfrNormalized = FMath::GetMappedRangeValueClamped(
+						FVector2f(FirstPointLengthFromRoot, LastPointLengthFromRoot),
+						FVector2f(0.0f, 1.0f), CurrentPointLengthFromRoot);
+					const float NextPointLfrNormalized = FMath::GetMappedRangeValueClamped(
+						FVector2f(FirstPointLengthFromRoot, LastPointLengthFromRoot),
+						FVector2f(0.0f, 1.0f),
+						NextPointLengthFromRoot);
+
+					if (NormalizedAttachmentPoints[Idx] >= CurrentPointLfrNormalized && NormalizedAttachmentPoints[Idx] <=
+						NextPointLfrNormalized)
+					{
+						const float DistanceAlpha = (NormalizedAttachmentPoints[Idx] -
+							CurrentPointLfrNormalized) / (NextPointLfrNormalized - CurrentPointLfrNormalized);
+
+						const float EthyleneLevel = FMath::Lerp(EthyleneLevelAtCurrentPoint, EthyleneLevelAtNextPoint, DistanceAlpha);
+
+						if (EthyleneLevel >= DistributionSettings.EthyleneThreshold)
+						{
+							continue;
+						}
+
+						const FVector3f Position = FMath::Lerp(CurrentPointPosition, NextPointPosition, DistanceAlpha);
+						AttachmentPointsPositions.Add(Position);
+
+						const float AttachmentPointLfr = FMath::GetMappedRangeValueClamped(
+							FVector2f(0.0f, 1.0f),
+							FVector2f(FirstPointLengthFromRoot, LastPointLengthFromRoot),
+							NormalizedAttachmentPoints[Idx]);
+						AttachmentPointsLengthsFromRoot.Add(AttachmentPointLfr);
+
+						// Compute Scale
+						const float CurrentPointBranchGradient = BranchGradientsAttribute[CurrentPointIndex];
+						const float NextPointBranchGradient = BranchGradientsAttribute[NextPointIndex];
+
+						const float BranchGradient = FMath::Clamp(
+							FMath::Lerp(CurrentPointBranchGradient, NextPointBranchGradient, DistanceAlpha), 0.0f, 1.0f);
+						float BranchGradientRampValue = ScaleSettings.ScaleRamp
+							                                ? ScaleSettings.ScaleRamp->GetRichCurveConst()->Eval(BranchGradient)
+							                                : 1.0f;
+
+						const float CurrentPointScale = PointFacadeSource.GetPointScale(CurrentPointIndex);
+						const float NextPointScale = PointFacadeSource.GetPointScale(NextPointIndex);
+						const float BranchScale = FMath::Lerp(CurrentPointScale, NextPointScale, DistanceAlpha);
+
+						const float BranchScaledValue = BranchGradientRampValue * (BranchScale + 0.1f);
+						BranchGradientRampValue = FMath::Lerp(BranchGradientRampValue, BranchScaledValue, ScaleSettings.BranchScaleImpact);
+						BranchGradientRampValue = FMath::Clamp(BranchGradientRampValue, 0.0f, 1.0f);
+
+						float Scale = FMath::GetMappedRangeValueClamped(
+							FVector2f(0.0f, 1.0f),
+							FVector2f(ScaleSettings.MinScale, ScaleSettings.MaxScale),
+							BranchGradientRampValue
+						) * ScaleSettings.BaseScale;
+
+						const float RandomValue = RandomStream.FRandRange(ScaleSettings.RandomScaleMin, ScaleSettings.RandomScaleMax);
+						Scale = Scale * RandomValue;
+
+						AttachmentPointsScales.Add(Scale);
+
+						const FVector3f CurrentPointApicalDirectionVector(BudDirectionAttribute[CurrentPointIndex][0]);
+						const FVector3f NextPointApicalDirectionVector(BudDirectionAttribute[NextPointIndex][0]);
+						const FVector3f BudApicalDirectionVector = FMath::Lerp(CurrentPointApicalDirectionVector,
+						                                                       NextPointApicalDirectionVector, DistanceAlpha);
+
+						int32 BudsToAdd = 1;
+
+						// Compute Up Vector
+						// Last point
+						if (Idx == NormalizedAttachmentPoints.Num() - 1)
+						{
+							AttachmentPointsUpVectors.Add(BudApicalDirectionVector);
+						}
+						else // Points other than last
+						[[likely]]
+						{
+							const FQuat4f RotationQuat(InitialApicalDirection.GetSafeNormal(), AxillaryRotationDegree);
+							FVector3f NewAxillaryDirection = RotationQuat.RotateVector(InitialAxillaryDirection.GetSafeNormal());
+							FQuat4f ApicalCorrection = FQuat4f::FindBetweenNormals(
+								InitialApicalDirection.GetSafeNormal(), BudApicalDirectionVector.GetSafeNormal());
+							NewAxillaryDirection = ApicalCorrection.RotateVector(NewAxillaryDirection).GetSafeNormal();
+
+							int32 NumBuds = RandomStream.RandRange(MinBudsLeaf, MaxBudsLeaf);
+							float RotationPerBud = 360.0f / static_cast<float>(NumBuds);
+
+							for (int32 i = 0; i < NumBuds; ++i)
+							{
+								float RotationAngle = FMath::DegreesToRadians(RotationPerBud * i);
+								const FQuat4f AxillaryRotation(InitialApicalDirection.GetSafeNormal(), RotationAngle);
+								FVector3f AxillaryDirectionModified = AxillaryRotation.RotateVector(NewAxillaryDirection.GetSafeNormal());
+
+								if (VectorSettings.OverrideAxilAngle)
+								{
+									float AxilAngleRampValue = VectorSettings.AxilAngleRamp
+										                           ? FMath::Clamp(
+											                           VectorSettings.AxilAngleRamp->GetRichCurveConst()->Eval(BranchGradient),
+											                           0.0f,
+											                           1.0f)
+										                           : 1.0f;
+
+									AxilAngleRampValue = AxilAngleRampValue * VectorSettings.AxilAngleRampUpperValue;
+									const float AxilAngleBlended = FMath::Lerp(VectorSettings.AxilAngle, AxilAngleRampValue,
+									                                           VectorSettings.AxilAngleRampEffect);
+
+									const FVector3f RotationAxis = FVector3f::CrossProduct(
+										BudApicalDirectionVector.GetSafeNormal(), AxillaryDirectionModified.GetSafeNormal());
+
+									const FQuat4f Rotation(RotationAxis.GetSafeNormal(), FMath::DegreesToRadians(AxilAngleBlended));
+									const FVector3f AxilVector = Rotation.RotateVector(BudApicalDirectionVector.GetSafeNormal());
+
+									AttachmentPointsUpVectors.Add(AxilVector);
+								}
+								else
+								{
+									AttachmentPointsUpVectors.Add(AxillaryDirectionModified);
+								}
+
+								if (i > 0)
+								{
+									FVector3f LastPositionAdded = AttachmentPointsPositions.Last();
+									float LastScaleAdded = AttachmentPointsScales.Last();
+									float LastLengthFromRootAdded = AttachmentPointsLengthsFromRoot.Last();
+									AttachmentPointsPositions.Add(LastPositionAdded);
+									AttachmentPointsScales.Add(LastScaleAdded);
+									AttachmentPointsLengthsFromRoot.Add(LastLengthFromRootAdded);
+								}
+							}
+
+							BudsToAdd = NumBuds;
+							AxillaryRotationIteration++;
+						}
+
+						// Compute Normal vectors based on light optimal direction
+						const FVector3f CurrentPointLightOptimalVector(BudDirectionAttribute[CurrentPointIndex][2]);
+						const FVector3f NextPointLightOptimalVector(BudDirectionAttribute[NextPointIndex][2]);
+						const FVector3f BudLightOptimalVector = FMath::Lerp(CurrentPointLightOptimalVector, NextPointLightOptimalVector,
+						                                                    DistanceAlpha);
+
+						FQuat4f QuatRotation = FQuat4f::FindBetweenNormals(BudApicalDirectionVector.GetSafeNormal(),
+						                                                   BudLightOptimalVector.GetSafeNormal());
+						const FVector3f ApicalDirectionRotated = QuatRotation.RotateVector(BudApicalDirectionVector);
+
+						for (int32 i = 0; i < BudsToAdd; ++i)
+						{
+							AttachmentPointsNormalVectors.Add(ApicalDirectionRotated);
+						}
+					}
+				}
+			}
+
+			// Set foliage entries
+			TManagedArrayAccessor<FVector3f> PivotPointsAttribute(OutCollection, MegaPlants::AttributeNames::FoliagePivotPoint,
+			                                                      MegaPlants::GroupNames::FoliageGroup);
+			const int32 NumFoliageMeshes = FoliageFacadeOutput.NumFoliageNames();
+			int32 CurrentFoliageMeshIndex = BranchIndex % NumFoliageMeshes;
+			for (int32 Idx = 0; Idx < AttachmentPointsLengthsFromRoot.Num(); ++Idx)
+			{
+				const int32 NewIndex = PivotPointsAttribute.AddElements(1);
+				FoliageFacadeOutput.SetFoliageEntry(NewIndex, {
+					                                    .NameId = CurrentFoliageMeshIndex,
+					                                    .BranchId = BranchIndex,
+					                                    .PivotPoint = AttachmentPointsPositions[Idx],
+					                                    .UpVector = AttachmentPointsUpVectors[Idx],
+					                                    .NormalVector = AttachmentPointsNormalVectors[Idx],
+					                                    .Scale = AttachmentPointsScales[Idx],
+					                                    .LengthFromRoot = AttachmentPointsLengthsFromRoot[Idx]
+				                                    });
+				CurrentFoliageMeshIndex = (CurrentFoliageMeshIndex + 1) % NumFoliageMeshes;
+			}
+		}
+
+		// Re-calculate Foliage IDs for each branch
+		TMap<int32, TArray<int32>> BranchIndexToFoliageIDs;
+		for (int32 i = 0; i < FoliageFacadeOutput.NumFoliageEntries(); ++i)
+		{
+			const int32 BranchIndex = FoliageFacadeOutput.GetFoliageEntry(i).BranchId;
+			TArray<int32>& IDs = BranchIndexToFoliageIDs.FindOrAdd(BranchIndex);
+			IDs.Add(i);
+		}
+		for (const TPair<int32, TArray<int32>>& Pair : BranchIndexToFoliageIDs)
+		{
+			int32 BranchId = Pair.Key;
+			const TArray<int32>& FoliageIDs = Pair.Value;
+			FoliageFacadeOutput.SetFoliageIdsArray(BranchId, FoliageIDs);
+		}
+	}
+
+	{
+		TRACE_CPUPROFILER_EVENT_SCOPE(MegaPlantFoliageDistributorNode::EthyleneRemove);
+
+		// Compute Foliage instances to remove based on Ethylene threshold
+		TManagedArrayAccessor<FVector3f> PivotPointsAttribute(OutCollection, MegaPlants::AttributeNames::FoliagePivotPoint,
+		                                                      MegaPlants::GroupNames::FoliageGroup);
+		TSet<int32> FoliageEntriesToRemove;
+
+		for (int32 BranchIndex = 0; BranchIndex < NumBranches; ++BranchIndex)
+		{
+			TArray<int32> BranchPoints = BranchFacadeSource.GetPoints(BranchIndex);
+			TArray<int32> FoliageIDs = FoliageFacadeOutput.GetFoliageEntryIdsForBranch(BranchIndex);
+			TSet<int32> FoliageIDsToRemove;
+
+			for (int32 BranchPointIndex = 0; BranchPointIndex < BranchPoints.Num() - 1; ++BranchPointIndex)
+			{
+				for (int32 FoliageIDsIndex = 0; FoliageIDsIndex < FoliageIDs.Num(); ++FoliageIDsIndex)
+				{
+					const float CurrentPointLengthFromRoot = PointFacadeSource.GetLengthFromRoot(BranchPoints[BranchPointIndex]);
+					const float NextPointLengthFromRoot = PointFacadeSource.GetLengthFromRoot(BranchPoints[BranchPointIndex + 1]);
+					const float FoliageLengthFromRoot = FoliageFacadeOutput.GetFoliageEntry(FoliageIDs[FoliageIDsIndex]).LengthFromRoot;
+
+					if (FoliageLengthFromRoot >= CurrentPointLengthFromRoot && FoliageLengthFromRoot <= NextPointLengthFromRoot)
+					{
+						float EthyleneLevelAtCurrentPoint = BudHormoneLevels[BranchPoints[BranchPointIndex]][4];
+						float EthyleneLevelAtNextPoint = BudHormoneLevels[BranchPoints[BranchPointIndex + 1]][4];
+
+						float EthyleneLevel = FMath::GetMappedRangeValueClamped(
+							FVector2f(CurrentPointLengthFromRoot, NextPointLengthFromRoot),
+							FVector2f(EthyleneLevelAtNextPoint, EthyleneLevelAtCurrentPoint),
+							FoliageLengthFromRoot);
+
+						if (EthyleneLevel >= DistributionSettings.EthyleneThreshold)
+						{
+							FoliageIDsToRemove.Add(FoliageIDs[FoliageIDsIndex]);
+						}
+					}
+				}
+			}
+
+			for (int32 Fid : FoliageIDsToRemove)
+			{
+				FoliageEntriesToRemove.Add(Fid);
+			}
+		}
+
+		// Remove foliage instances
+		for (int32 i = PivotPointsAttribute.Num() - 1; i >= 0; --i)
+		{
+			if (FoliageEntriesToRemove.Contains(i))
+			{
+				PivotPointsAttribute.RemoveElements(1, i);
+			}
+		}
+
+		// Re-calculate Foliage IDs for each branch
+		TMap<int32, TArray<int32>> BranchIndexToFoliageIDs;
+		for (int32 i = 0; i < FoliageFacadeOutput.NumFoliageEntries(); ++i)
+		{
+			const int32 BranchIndex = FoliageFacadeOutput.GetFoliageEntry(i).BranchId;
+			TArray<int32>& IDs = BranchIndexToFoliageIDs.FindOrAdd(BranchIndex);
+			IDs.Add(i);
+		}
+		for (const TPair<int32, TArray<int32>>& Pair : BranchIndexToFoliageIDs)
+		{
+			int32 BranchId = Pair.Key;
+			const TArray<int32>& FoliageIDs = Pair.Value;
+			FoliageFacadeOutput.SetFoliageIdsArray(BranchId, FoliageIDs);
+		}
+	}
+}
diff --git a/Engine/Plugins/Experimental/ProceduralVegetationEditor/Source/MegaPlants/Private/Implementations/MegaPlantFoliage.h b/Engine/Plugins/Experimental/ProceduralVegetationEditor/Source/MegaPlants/Private/Implementations/MegaPlantFoliage.h
new file mode 100644
--- /dev/null
+++ b/Engine/Plugins/Experimental/ProceduralVegetationEditor/Source/MegaPlants/Private/Implementations/MegaPlantFoliage.h
@@ -0,0 +1,83 @@
+// Copyright Epic Games, Inc. All Rights Reserved.
+
+#pragma once
+#include "CoreMinimal.h"
+#include "GeometryCollection/ManagedArrayCollection.h"
+#include "Curves/CurveFloat.h"
+
+#include "MegaPlantFoliage.generated.h"
+
+UENUM(BlueprintType)
+enum class EPhyllotaxyType : uint8
+{
+	Alternate UMETA(DisplayName = "Alternate"),
+	Opposite UMETA(DisplayName = "Opposite"),
+	Decussate UMETA(DisplayName = "Decussate"),
+	Whorled UMETA(DisplayName = "Whorled"),
+	Spiral UMETA(DisplayName = "Spiral")
+};
+
+UENUM(BlueprintType)
+enum class EPhyllotaxyFormation : uint8
+{
+	Distichous UMETA(DisplayName = "Distichous"),
+	Tristichous UMETA(DisplayName = "Tristichous"),
+	Pentastichous UMETA(DisplayName = "Pentastichous"),
+	Octastichous UMETA(DisplayName = "Octastichous"),
+	Parastichous UMETA(DisplayName = "Parastichous")
+};
+
+struct FDistributionSettings
+{
+	const float EthyleneThreshold;
+	const bool OverrideDistribution;
+	const float InstanceSpacing;
+	const FRuntimeFloatCurve* InstanceSpacingRamp;
+	const float InstanceSpacingRampEffect;
+	const int32 MaxPerBranch;
+};
+
+struct FScaleSettings
+{
+	const float BaseScale;
+	const float BranchScaleImpact;
+	const float MinScale;
+	const float MaxScale;
+	const float RandomScaleMin;
+	const float RandomScaleMax;
+	const FRuntimeFloatCurve* ScaleRamp;
+};
+
+struct FVectorSettings
+{
+	const bool OverrideAxilAngle;
+	const float AxilAngle;
+	const FRuntimeFloatCurve* AxilAngleRamp;
+	const float AxilAngleRampUpperValue;
+	const float AxilAngleRampEffect;
+};
+
+struct FPhyllotaxySettings
+{
+	const bool OverridePhyllotaxy;
+	const EPhyllotaxyType PhyllotaxyType;
+	const EPhyllotaxyFormation PhyllotaxyFormation;
+	const int32 MinimumNodeBuds;
+	const int32 MaximumNodeBuds;
+	const bool ResetPhyllotaxy;
+	const float PhyllotaxyAdditionalAngle;
+	const float PhyllotaxyOffset;
+};
+
+struct FMiscSettings
+{
+	const int32 RandomSeed;
+};
+
+struct FMegaPlantFoliage
+{
+	static void DistributeFoliage(FManagedArrayCollection& OutCollection, const FManagedArrayCollection& SourceCollection,
+	                              const FDistributionSettings& DistributionSettings, const FScaleSettings& ScaleSettings,
+	                              const FVectorSettings& VectorSettings, const FPhyllotaxySettings& PhyllotaxySettings,
+	                              const FMiscSettings& MiscSettings);
+};
diff --git a/Engine/Plugins/Experimental/ProceduralVegetationEditor/Source/MegaPlants/Private/Implementations/MegaPlantGravity.cpp b/Engine/Plugins/Experimental/ProceduralVegetationEditor/Source/MegaPlants/Private/Implementations/MegaPlantGravity.cpp
new file mode 100644
--- /dev/null
+++ b/Engine/Plugins/Experimental/ProceduralVegetationEditor/Source/MegaPlants/Private/Implementations/MegaPlantGravity.cpp
@@ -0,0 +1,219 @@
+// Copyright Epic Games, Inc. All Rights Reserved.
+
+#include "MegaPlantGravity.h"
+#include "Facades/MegaPlantBranchFacade.h"
+#include "Facades/MegaPlantBudVectorsFacade.h"
+#include "Facades/MegaPlantFoliageFacade.h"
+#include "Facades/MegaPlantPointFacade.h"
+
+void FMegaPlantGravity::GeneratePhototropicData(const FMegaPlantGravitySettings& InGravitySettings, const MegaPlants::Facades::FBranchFacade& InBranchFacade, const MegaPlants::Facades::FPointFacade& InPointFacade, TArray<FVector3f>& OutPhototropicDirections)
+{
+	FVector3f PrevVec = FVector3f::UpVector;
+		
+	OutPhototropicDirections.Init(PrevVec,InPointFacade.GetElementCount());
+
+	static constexpr int BudDirectionsLightOptimalIndex = 2;
+	static constexpr int BudDirectionsLightSubOptimal = 3;
+		
+	for (int BranchIndex = 0; BranchIndex < InBranchFacade.GetElementCount(); BranchIndex++)
+	{
+		const TArray<int32>& BranchPoints = InBranchFacade.GetPoints(BranchIndex);
+			
+		for (int i = BranchPoints.Num() - 1; i > 0; i--)
+		{
+			const auto CurrentPointIndex = BranchPoints[i];
+			const TArray<FVector3f>& BudDirections = InPointFacade.GetBudDirection(CurrentPointIndex);
+
+			if (BudDirections.Num() > BudDirectionsLightSubOptimal)
+			{
+				const FVector3f& LightOptimal = BudDirections[BudDirectionsLightOptimalIndex];
+				const FVector3f& LightSubOptimal = BudDirections[BudDirectionsLightSubOptimal];
+
+				FVector3f PhototropicDirection = FMath::Lerp(LightOptimal,LightSubOptimal, InGravitySettings.PhototropicBias);
+				PhototropicDirection.Normalize();
+
+				if (PhototropicDirection.Length() < 0.001)
+				{
+					PhototropicDirection = PrevVec;
+				}
+				else
+				{
+					PrevVec = PhototropicDirection;
+				}
+
+				if (OutPhototropicDirections.Num() > CurrentPointIndex)
+				{
+					OutPhototropicDirections[CurrentPointIndex] = PhototropicDirection;
+				}
+			}
+		}	
+	}
+}
+
+void FMegaPlantGravity::ApplyGravity(const FMegaPlantGravitySettings& InGravitySettings, FManagedArrayCollection& OutCollection)
+{
+	MegaPlants::Facades::FBranchFacade BranchFacade = MegaPlants::Facades::FBranchFacade(OutCollection);
+	MegaPlants::Facades::FPointFacade PointFacade = MegaPlants::Facades::FPointFacade(OutCollection);
+		
+	if(!PointFacade.IsValid() || !BranchFacade.IsValid())
+	{
+		return;
+	}
+
+	const TManagedArray<int32>& BranchNumbers = BranchFacade.GetBranchNumbers();
+	const int32 RootBranchIndex = BranchNumbers.Find(1);
+
+	TArray<FVector3f> OutPhototropicDirection;
+	
+	if (InGravitySettings.Mode == EGravityMode::Phototropic)
+	{
+		GeneratePhototropicData(InGravitySettings, BranchFacade, PointFacade, OutPhototropicDirection);
+	}
+
+	ApplyGravity(RootBranchIndex, InGravitySettings, OutPhototropicDirection, OutCollection);
+}
+
+// Please note that LFR (length from root) for both skeleton points and foliage points are not being scaled by
+// 100 at the moment, which is OK for now since they're only being compared relative to each other for now,
+// however if they're used in any computation in the future in any meaningful way where other quantities are
+// involved they will end up giving erroneous results.
+void FMegaPlantGravity::ApplyGravity(
+	const int32 BranchIndex,
+	const FMegaPlantGravitySettings& GravitySettings,
+	const TArray<FVector3f>& PhototropicDirections,
+	FManagedArrayCollection& OutCollection,
+	FQuat4f TotalDownForce,
+	FVector3f PreviousPosition,
+	FVector3f PivotPosition
+)
+{
+	static const auto CalculateDownForce = [](const FVector3f InRelativeDirection, const int InPointIndex, const float InPScaleGradient, const FMegaPlantGravitySettings& InGravitySettings, const TArray<FVector3f>& InPhototropicDirections) -> FQuat4f
+	{
+		const FVector3f PhototropicDirection = InPhototropicDirections.Num() > InPointIndex ? InPhototropicDirections[InPointIndex] : FVector3f::UpVector;
+		const FVector3f GravityDirection = 	InGravitySettings.Mode == EGravityMode::Phototropic ? PhototropicDirection : InGravitySettings.Direction; 
+
+		float YDot = 1.0f - FMath::Abs(FVector3f::DotProduct(InRelativeDirection,InGravitySettings.Direction));
+		YDot = FMath::GetMappedRangeValueClamped(FVector2f(0, 1),
+		FVector2f(FMath::Max(0.1f, 1 - InGravitySettings.AngleCorrection), 1),
+		          YDot);
+		
+		const float GravityImpact = InGravitySettings.Gravity * (1 - InPScaleGradient) * YDot;
+		const FVector3f GravityForce = FMath::Lerp(InRelativeDirection, GravityDirection, GravityImpact).GetUnsafeNormal();
+		return FQuat4f::FindBetween(InRelativeDirection, GravityForce);
+	};
+
+	MegaPlants::Facades::FBranchFacade BranchFacade = MegaPlants::Facades::FBranchFacade(OutCollection);
+	MegaPlants::Facades::FPointFacade PointFacade = MegaPlants::Facades::FPointFacade(OutCollection);
+	MegaPlants::Facades::FBudVectorsFacade BudVectorFacade = MegaPlants::Facades::FBudVectorsFacade(OutCollection);
+	MegaPlants::Facades::FFoliageFacade FoliageFacade = MegaPlants::Facades::FFoliageFacade(OutCollection);
+	
+	TManagedArray<FVector3f>& PointPositions = PointFacade.ModifyPositions();
+	TManagedArray<TArray<FVector3f>>& PointBudDirections = BudVectorFacade.ModifyBudDirections();
+	
+	const TArray<int32>& CurrentBranchPoints = BranchFacade.GetPoints(BranchIndex);
+	const TArray<int32>& CurrentBranchChildren = BranchFacade.GetChildren(BranchIndex);
+	const TManagedArray<int32>& BranchNumbers = BranchFacade.GetBranchNumbers();
+
+	TArray<int32> ChildrenIndices;
+	ChildrenIndices.Reserve(CurrentBranchChildren.Num());
+
+	for (const int32 Child : CurrentBranchChildren)
+	{
+		const int32 ChildIndex = BranchNumbers.Find(Child);
+		const TArray<int32>& Parents = BranchFacade.GetParents(ChildIndex);
+		
+		if (ChildIndex != INDEX_NONE && !Parents.IsEmpty() && Parents[Parents.Num() - 1] == BranchNumbers[BranchIndex])
+		{
+			ChildrenIndices.Add(ChildIndex);
+		}
+	}
+
+	if (TotalDownForce != FQuat4f::Identity)
+	[[likely]]
+	{
+		const int32 PointIndex = CurrentBranchPoints[0];
+		FVector3f& Position = PointPositions[PointIndex];
+		TArray<FVector3f>& CurrentBudDirections = PointBudDirections[PointIndex];
+		const float& PScaleGradient = PointFacade.GetPointScaleGradient(PointIndex);
+
+		const FVector3f RelativeDirection = TotalDownForce * (Position - PreviousPosition);
+		const FQuat4f DownForce = CalculateDownForce(RelativeDirection.GetUnsafeNormal(), PointIndex, PScaleGradient, GravitySettings, PhototropicDirections);
+
+		TotalDownForce = DownForce * TotalDownForce;
+		PreviousPosition = Position;
+		Position = PivotPosition + DownForce * RelativeDirection;
+
+		for (FVector3f& BudDirection : CurrentBudDirections)
+		{
+			BudDirection = TotalDownForce * BudDirection;
+		}
+	}
+	else
+	{
+		PreviousPosition = PointPositions[CurrentBranchPoints[0]];
+	}
+
+	const int32 MaxIndex = CurrentBranchPoints.Num() - 1;
+	for (int32 Index = 1; Index <= MaxIndex; ++Index)
+	{
+		const int32 PointIndex = CurrentBranchPoints[Index];
+		const int32 PreviousPointIndex = CurrentBranchPoints[Index - 1];
+
+		FVector3f& Position = PointPositions[PointIndex];
+		PivotPosition = PointPositions[PreviousPointIndex];
+
+		TArray<FVector3f>& CurrentBudDirections = PointBudDirections[PointIndex];
+
+		const float LengthFromRoot = PointFacade.GetLengthFromRoot(PointIndex);
+		const float PreviousLengthFromRoot = PointFacade.GetLengthFromRoot(PreviousPointIndex);
+
+		const FVector3f RelativeDirection = TotalDownForce * (Position - PreviousPosition);
+		const FVector3f Tangent = RelativeDirection.GetUnsafeNormal();
+
+		const FQuat4f DownForce = CalculateDownForce(Tangent, PointIndex, PointFacade.GetPointScaleGradient(PointIndex), GravitySettings, PhototropicDirections);
+		TotalDownForce = DownForce * TotalDownForce;
+
+		// Update attachment points for all foliage elements for current branch between the current and previous point 
+		TArray<int32> FoliageEntryIds = FoliageFacade.GetFoliageEntryIdsForBranch(BranchIndex);
+		for (const int32 FoliageEntryId : FoliageEntryIds)
+		{
+			MegaPlants::Facades::FFoliageEntryData Entry = FoliageFacade.GetFoliageEntry(FoliageEntryId);
+
+			if (Entry.LengthFromRoot <= LengthFromRoot && Entry.LengthFromRoot > PreviousLengthFromRoot)
+			{
+				// Pivot point is updated based on Down force while Up Vector is updated based on Total Down force
+				// same as point positions and bud directions respectively
+				Entry.PivotPoint = PivotPosition + DownForce * RelativeDirection;
+				Entry.UpVector = TotalDownForce * Entry.UpVector;
+				FoliageFacade.SetFoliageEntry(FoliageEntryId, Entry);
+			}
+		}
+
+		for (const int32 ChildrenIndex : ChildrenIndices)
+		{
+			const int32 FirstPointIndex = BranchFacade.GetPoints(ChildrenIndex)[0];
+			const float ChildLengthFromRoot = PointFacade.GetLengthFromRoot(FirstPointIndex);
+			if (ChildLengthFromRoot <= LengthFromRoot && ChildLengthFromRoot > PreviousLengthFromRoot)
+			[[unlikely]]
+			{
+				ApplyGravity(
+					ChildrenIndex,
+					GravitySettings,
+					PhototropicDirections,
+					OutCollection,
+					TotalDownForce,
+					PreviousPosition,
+					PivotPosition
+				);
+			}
+		}
+
+		PreviousPosition = Position;
+		Position = PivotPosition + DownForce * RelativeDirection;
+
+		for (FVector3f& BudDirection : CurrentBudDirections)
+		{
+			BudDirection = TotalDownForce * BudDirection;
+		}
+	}
+}
diff --git a/Engine/Plugins/Experimental/ProceduralVegetationEditor/Source/MegaPlants/Private/Implementations/MegaPlantGravity.h b/Engine/Plugins/Experimental/ProceduralVegetationEditor/Source/MegaPlants/Private/Implementations/MegaPlantGravity.h
new file mode 100644
--- /dev/null
+++ b/Engine/Plugins/Experimental/ProceduralVegetationEditor/Source/MegaPlants/Private/Implementations/MegaPlantGravity.h
@@ -0,0 +1,53 @@
+// Copyright Epic Games, Inc. All Rights Reserved.
+
+#pragma once
+
+#include "CoreMinimal.h"
+#include "GeometryCollection/ManagedArrayCollection.h"
+#include "MegaPlantGravity.generated.h"
+
+namespace MegaPlants::Facades
+{
+	class FPointFacade;
+	class FBranchFacade;
+}
+
+UENUM()
+enum class EGravityMode : uint8
+{
+	Gravity,
+	Phototropic
+};
+
+USTRUCT()
+struct FMegaPlantGravitySettings
+{
+	GENERATED_BODY()
+	
+	UPROPERTY(EditAnywhere, Category="Gravity")
+	EGravityMode Mode = EGravityMode::Gravity;
+
+	UPROPERTY(EditAnywhere, Category="Gravity", meta=(ClampMin=0.0f, ClampMax=1.0f, UIMin=0.0f, UIMax=1.0f))
+	float Gravity = 0;
+
+	UPROPERTY(EditAnywhere, Category="Gravity", meta=(EditCondition = "Mode == EGravityMode::Gravity"))
+	FVector3f Direction = FVector3f::DownVector;
+
+	UPROPERTY(EditAnywhere, Category="Gravity", meta=(ClampMin=0.0f, ClampMax=1.0f, UIMin=0.0f, UIMax=1.0f))
+	float AngleCorrection = 0;
+
+	UPROPERTY(EditAnywhere, Category="Gravity", meta=(ClampMin=0.0f, ClampMax=1.0f, UIMin=0.0f, UIMax=1.0f, EditCondition = "Mode == EGravityMode::Phototropic"))
+	float PhototropicBias = 0;
+};
+
+struct FMegaPlantGravity
+{
+	static void ApplyGravity(const FMegaPlantGravitySettings& InGravitySettings, FManagedArrayCollection& OutCollection);
+	
+private :
+
+	static void GeneratePhototropicData(const FMegaPlantGravitySettings& InGravitySettings, const MegaPlants::Facades::FBranchFacade& InBranchFacade, const MegaPlants::Facades::FPointFacade& InPointFacade, TArray<FVector3f>& OutPhototropicDirections);
+
+	static void ApplyGravity(const int32 BranchIndex, const FMegaPlantGravitySettings& GravitySettings, const TArray<FVector3f>& PhototropicDirections, FManagedArrayCollection& OutCollection,
+		FQuat4f TotalDownForce = FQuat4f::Identity, FVector3f PreviousPosition = FVector3f::ZeroVector, FVector3f PivotPosition = FVector3f::ZeroVector);
+};
diff --git a/Engine/Plugins/Experimental/ProceduralVegetationEditor/Source/MegaPlants/Private/Implementations/MegaPlantMaterialSettings.cpp b/Engine/Plugins/Experimental/ProceduralVegetationEditor/Source/MegaPlants/Private/Implementations/MegaPlantMaterialSettings.cpp
new file mode 100644
--- /dev/null
+++ b/Engine/Plugins/Experimental/ProceduralVegetationEditor/Source/MegaPlants/Private/Implementations/MegaPlantMaterialSettings.cpp
@@ -0,0 +1,330 @@
+// Copyright Epic Games, Inc. All Rights Reserved.
+
+#include "MegaPlantMaterialSettings.h"
+
+#include "MegaPlantsModule.h"
+#include "Algo/MaxElement.h"
+#include "Facades/MegaPlantBranchFacade.h"
+#include "Facades/MegaPlantPointFacade.h"
+
+struct FBranchUVData
+{
+	float XOffset;
+	float YOffset;
+	float YScale;
+	int	MaterialId;
+};
+
+struct FPointUVData
+{
+	FVector2f BaseUVs;
+};
+
+void FMegaPlantMaterialSettings::ApplyMaterialSettings(FManagedArrayCollection& Collection) const
+{
+	TArray<int32> UsedMaterials;
+	TArray<FBranchUVData> BranchUVData;
+	SetBranchUVMaterial(Collection, UsedMaterials, BranchUVData);
+	SetUVOffsetAndScale(Collection, BranchUVData);
+	TArray<FPointUVData> PointUVData;
+	SetBaseUVs(Collection, PointUVData);
+	FitUVs01(Collection, PointUVData, UsedMaterials, BranchUVData);
+	SetOutputUVs(Collection, PointUVData, BranchUVData);
+}
+
+void FMegaPlantMaterialSettings::SetMaterial(FString TrunkMaterial, int32 Index)
+{
+	FTrunkGenerationMaterialSetup MaterialSetup;
+	MaterialSetup.Material = LoadObject<UMaterialInterface>(nullptr, TrunkMaterial);
+
+	if (!MaterialSetups.IsValidIndex(Index))
+	{
+		MaterialSetups.Add(MaterialSetup);
+	}
+	else
+	{
+		MaterialSetups[Index].Material = MaterialSetup.Material;
+	}
+}
+
+void FMegaPlantMaterialSettings::SetBranchUVMaterial(FManagedArrayCollection& Collection, TArray<int32>& UsedMaterials, TArray<FBranchUVData>& BranchUVData) const
+{
+	MegaPlants::Facades::FPointFacade PointFacade = MegaPlants::Facades::FPointFacade(Collection);
+	MegaPlants::Facades::FBranchFacade BranchFacade = MegaPlants::Facades::FBranchFacade(Collection);
+
+	BranchUVData.Init(FBranchUVData(), BranchFacade.GetElementCount());
+	
+	float MaxValue = 10000;
+	float MinGeneration = MaxValue;
+	float MinAge = MaxValue;
+	float MinPScale = MaxValue;
+	
+	float MaxGeneration = 0;
+	float MaxAge = 0;
+	float MaxPScale = 0;
+
+	for (int32 PointIndex = 0; PointIndex < PointFacade.GetElementCount(); PointIndex++)
+	{
+		TArray<int> BudDevelopment = PointFacade.GetBudDevelopment(PointIndex);
+		float PScale = PointFacade.GetPointScale(PointIndex);
+
+		check(BudDevelopment.Num() > 2);
+		MinGeneration = FMath::Min(MinGeneration, BudDevelopment[0]);
+		MinAge = FMath::Min(MinAge, BudDevelopment[2]);
+		MinPScale = FMath::Min(MinPScale, PScale);
+
+		MaxGeneration = FMath::Max(MaxGeneration, BudDevelopment[0]);
+		MaxAge = FMath::Max(MaxAge, BudDevelopment[2]);
+		MaxPScale = FMath::Max(MaxPScale, PScale);
+	}
+
+	for (int32 BranchIndex = 0; BranchIndex < BranchFacade.GetElementCount(); BranchIndex++)
+	{
+		const auto& BranchPoints = BranchFacade.GetPoints(BranchIndex);
+		check(BranchPoints.Num() > 0);
+		int32 BranchRootPoint = BranchPoints[0];
+
+		float BranchMaterialVal = 0;
+
+		TArray<int> BudDevelopment = PointFacade.GetBudDevelopment(BranchRootPoint);
+		float PScale = PointFacade.GetPointScale(BranchRootPoint);
+		check(BudDevelopment.Num() >= 2);
+		
+		if (MaterialMode == EUVMaterialMode::Generation)
+		{
+			BranchMaterialVal = GetMaterialValue(BudDevelopment[0], MinGeneration, MaxGeneration, MinGeneration, MaxGeneration);
+		}
+		else if (MaterialMode == EUVMaterialMode::Age)
+		{
+			BranchMaterialVal = GetMaterialValue(BudDevelopment[2], MinAge, MaxAge, MaxGeneration, MinGeneration);
+		}
+		else if (MaterialMode == EUVMaterialMode::Radius)
+		{
+			BranchMaterialVal = GetMaterialValue(PScale, MinPScale, MaxPScale, MaxGeneration, MinGeneration);
+		}
+
+		if (DistributionMethod == EMaterialDistributionMethod::Fit)
+		{
+			MaxGeneration = MaterialSetups.Num() + GenerationOffset;
+		}
+
+		if (OffsetMethod == EGenerationOffsetMethod::Clamped)
+		{
+			BranchMaterialVal = FMath::Clamp(BranchMaterialVal + GenerationOffset,MinGeneration + GenerationOffset,MaterialSetups.Num());
+		}
+		else if (OffsetMethod == EGenerationOffsetMethod::Refit)
+		{
+			FVector2f OldRange = FVector2f(MinGeneration, MaxGeneration);
+			FVector2f NewRange = FVector2f(MinGeneration + GenerationOffset, MaterialSetups.Num());
+			BranchMaterialVal = FMath::GetMappedRangeValueClamped(OldRange, NewRange, BranchMaterialVal);
+		}
+
+		int BranchUVMaterial = FMath::RoundToInt(BranchMaterialVal);
+		//Set this as branch atrribute
+		BranchFacade.SetBranchUVMaterial(BranchIndex, BranchUVMaterial - 1);
+		UsedMaterials.AddUnique(BranchUVMaterial);
+		BranchUVData[BranchIndex].MaterialId = BranchUVMaterial - 1;
+	}
+}
+
+float FMegaPlantMaterialSettings::GetMaterialValue(float Value, float MinValue, float MaxValue, float MinGeneration, float MaxGeneration) const
+{
+	//float RampValue = 0;
+	FVector2f OldRange = FVector2f(MinValue, MaxValue);
+	FVector2f NewRange = FVector2f(0.0f, 1.0f);
+	FVector2f GenerationRange = FVector2f(MinGeneration, MaxGeneration);
+
+	float MaterialValue = FMath::GetMappedRangeValueClamped(OldRange, NewRange, Value);
+    MaterialValue = FMath::GetMappedRangeValueClamped(NewRange, GenerationRange, MaterialValue);
+
+	return MaterialValue;
+}
+
+void FMegaPlantMaterialSettings::SetUVOffsetAndScale(FManagedArrayCollection& Collection, TArray<FBranchUVData>& BranchUVData) const
+{
+	const MegaPlants::Facades::FBranchFacade& BranchFacade = MegaPlants::Facades::FBranchFacade(Collection);
+	
+	FRandomStream RandomStream(103);
+				
+	for (int BranchIndex = 0; BranchIndex < BranchFacade.GetElementCount(); BranchIndex++ )
+	{
+		int BranchUVMaterial = BranchUVData[BranchIndex].MaterialId;
+			
+		for (int i=0; i< MaterialSetups.Num(); i++)
+		{
+			float OffsetX = MaterialSetups[i].XOffset;
+			float OffsetXRand = MaterialSetups[i].XOffsetRandom;
+			float OffsetY = MaterialSetups[i].YOffset;
+			float ScaleY = MaterialSetups[i].YScale;
+			float OffsetYRand = MaterialSetups[i].YOffsetRandom;
+	
+			if (BranchUVMaterial == i)
+			{
+				float TextureOffsetX = OffsetX + (RandomStream.FRand() * OffsetXRand);
+				float TextureOffsetY = OffsetY + (RandomStream.FRand() * OffsetYRand);
+				FBranchUVData UVData{ TextureOffsetX, TextureOffsetY, ScaleY, BranchUVMaterial};
+				check(BranchUVData.IsValidIndex(BranchIndex));
+				BranchUVData[BranchIndex] = UVData;
+			}
+		}
+		
+		if (BranchUVMaterial > MaterialSetups.Num())
+		{
+			FBranchUVData UVData{ 0, 0, 1, BranchUVMaterial};
+			check(BranchUVData.IsValidIndex(BranchIndex));
+			BranchUVData[BranchIndex] = UVData;
+		}
+	}
+}
+
+void FMegaPlantMaterialSettings::SetBaseUVs(FManagedArrayCollection& Collection, TArray<FPointUVData>& PointUVData) const
+{
+	MegaPlants::Facades::FPointFacade PointFacade = MegaPlants::Facades::FPointFacade(Collection);
+	MegaPlants::Facades::FBranchFacade BranchFacade = MegaPlants::Facades::FBranchFacade(Collection);
+
+	const TManagedArray<float>& PointScales = PointFacade.GetPointScales();
+	const float MaxPointScale = *Algo::MaxElement(PointScales);
+
+	PointUVData.Init(FPointUVData(), PointScales.Num());
+	
+	for (int BranchIndex = 0; BranchIndex < BranchFacade.GetElementCount(); BranchIndex++ )
+	{
+		auto& BranchPoints = BranchFacade.GetPoints(BranchIndex);
+		TArray<float> SegmentArrayMetric = MakeSegmentArray(PointFacade,BranchPoints, MaxPointScale, 0);
+		TArray<float> SegmentArrayBase = MakeSegmentArray(PointFacade,BranchPoints, MaxPointScale, 1);
+    
+		float UVYPreviousMetric = 0;
+		float UVYPreviousBase = 0;
+
+		for (int j=1; j < BranchPoints.Num(); j++)
+		{
+			int PTCurrent = BranchPoints[j];
+			float UVYMetric = UVYPreviousMetric + SegmentArrayMetric[j-1];
+        
+			// Set the V of the UV using segment length
+			FVector2f UVMetric = FVector2f(0.0f, UVYMetric);
+			float UVYBase = UVYPreviousBase + SegmentArrayBase[j-1];
+			FVector2f UVBase = FVector2f(0.0f, UVYBase);
+			FPointUVData UVData = { UVBase };
+			check(PointUVData.IsValidIndex(PTCurrent));
+			PointUVData[PTCurrent] = UVData;
+			UVYPreviousMetric = UVYMetric;
+			UVYPreviousBase = UVYBase;
+		}
+	}
+}
+
+TArray<float> FMegaPlantMaterialSettings::MakeSegmentArray(const MegaPlants::Facades::FPointFacade& PointFacade,const TArray<int32>& Points,const float MaxPointScale, const int Mode) const
+{
+	TArray<float> SegmentArray;
+	
+	check(MaxPointScale > 0);
+	float MaxPointScaleRatio = 1.0f /(MaxPointScale * UE_TWO_PI);
+
+	for (int PointIndex = 0; PointIndex < Points.Num() - 1; PointIndex++)
+	{
+		float PointScale = PointFacade.GetPointScale(Points[PointIndex]);
+		FVector3f Position = PointFacade.GetPosition(Points[PointIndex]);
+		FVector3f NextPosition = PointFacade.GetPosition(Points[PointIndex + 1]);
+
+		float Dist = FVector3f::Distance(Position,NextPosition);
+		float PScaleRatio = 1.0;
+		if (Mode == 1)
+		{
+			PScaleRatio = MaxPointScale / (PointScale == 0 ? UE_KINDA_SMALL_NUMBER : PointScale);
+		}
+		float SegmentLength = Dist * (MaxPointScaleRatio * PScaleRatio);
+        
+		SegmentArray.Add(SegmentLength);
+	}
+	
+	return SegmentArray;
+}
+
+void FMegaPlantMaterialSettings::FitUVs01(FManagedArrayCollection& Collection, TArray<FPointUVData>& PointUVData, const TArray<int32>& UsedMaterials, TArray<FBranchUVData>& BranchUVData) const
+{
+	MegaPlants::Facades::FPointFacade PointFacade = MegaPlants::Facades::FPointFacade(Collection);
+	MegaPlants::Facades::FBranchFacade BranchFacade = MegaPlants::Facades::FBranchFacade(Collection);
+	
+	//For every material
+	for(int i = 0; i < UsedMaterials.Num(); i++)
+	{
+		//Get the UI y Texture Mode for Material
+		// yTextureMode = 0 Default, 1 O-1 Fit, 2 0-1 Fit Adjusted
+		for (int BranchIndex = 0; BranchIndex < BranchFacade.GetElementCount(); BranchIndex++ )
+		{
+			int BranchUVMaterial = BranchUVData[BranchIndex].MaterialId;
+			bool bFit01 = MaterialSetups.IsValidIndex(i) && MaterialSetups[i].YTextureMode > EYTextureMode::Default;
+			//Get all the branches with the material
+			if (BranchUVMaterial == i && bFit01)
+			{
+				//zeroOneFitUVs(matPRs,minMax_array);
+				auto& BranchPoints = BranchFacade.GetPoints(BranchIndex);
+	
+					//Get teh V min Max for the Branch
+				check(BranchPoints.IsValidIndex(0));
+				check(PointUVData.IsValidIndex(BranchPoints[0]));
+				FVector2f RootPTBaseUV = PointUVData[BranchPoints[0]].BaseUVs;//point(0,"uv_base",rootPT);
+				check(PointUVData.IsValidIndex(BranchPoints[BranchPoints.Num() - 1]));
+				FVector2f LastPTBaseUV = PointUVData[BranchPoints[BranchPoints.Num() - 1]].BaseUVs;//point(0,"uv_base",PRPTs[len(PRPTs)-1]);
+				FVector2f BranchUVMinMax(RootPTBaseUV.Y, LastPTBaseUV.Y);
+				for (int j=0; j < BranchPoints.Num(); j++)
+				{
+					FVector2f& UVBase = PointUVData[BranchPoints[j]].BaseUVs;//point(0,"uv_base",PRPT);
+    
+					//if point V >= Branch min V -> remap the V to 0 and 1
+					if (UVBase.Y >= BranchUVMinMax.X)
+					{
+						UVBase.Y = FMath::GetMappedRangeValueClamped(BranchUVMinMax, FVector2f(0,1), UVBase.Y);
+					}
+					// if point V < Branch min V -> remap the V between 0 and  -1
+					else
+					{
+						FVector2f OldRange(BranchUVMinMax.X, -BranchUVMinMax.Y);
+						UVBase.Y = FMath::GetMappedRangeValueClamped(OldRange, FVector2f(0,-1), UVBase.Y);
+					}
+				}
+			}
+		}
+	}
+}
+
+void FMegaPlantMaterialSettings::SetOutputUVs(FManagedArrayCollection& Collection, TArray<FPointUVData>& PointUVData, TArray<FBranchUVData>& BranchUVData) const
+{
+	MegaPlants::Facades::FBranchFacade BranchFacade = MegaPlants::Facades::FBranchFacade(Collection);
+	MegaPlants::Facades::FPointFacade PointFacade = MegaPlants::Facades::FPointFacade(Collection);
+	
+	for (int BranchIndex = 0; BranchIndex < BranchFacade.GetElementCount(); BranchIndex++ )
+	{
+		int BranchUVMaterial = BranchUVData[BranchIndex].MaterialId;
+		
+		float UVRatio = 1.0;
+		float UVYTextureOffset = 0;
+		float UVXTextureOffset = 0;
+		float UVYTextureScale = 1;
+		//Get all the branches with the material
+		if (MaterialSetups.IsValidIndex(BranchUVMaterial) && MaterialSetups[BranchUVMaterial].YTextureMode == EYTextureMode::Default)
+		{
+			UVYTextureOffset = BranchUVData[BranchIndex].YOffset;
+			UVYTextureScale = BranchUVData[BranchIndex].YScale;
+			UVRatio *= 1;
+		}
+
+		UVXTextureOffset = BranchUVData[BranchIndex].XOffset;
+		
+		auto& BranchPoints = BranchFacade.GetPoints(BranchIndex);
+	
+		for (int j=0; j < BranchPoints.Num(); j++)
+		{
+			FVector2f& UV = PointUVData[BranchPoints[j]].BaseUVs;
+			UV.Y *= -1;
+			UV.Y += UVYTextureOffset;
+			UV.X = -UVXTextureOffset;
+			UV.Y *= UVRatio;
+			UV.Y *= UVYTextureScale;
+
+			PointFacade.SeTextureCoordV(BranchPoints[j], UV.Y);
+			PointFacade.SetTextureCoordUOffset(BranchPoints[j], UV.X);
+		}
+	}
+}
diff --git a/Engine/Plugins/Experimental/ProceduralVegetationEditor/Source/MegaPlants/Private/Implementations/MegaPlantMaterialSettings.h b/Engine/Plugins/Experimental/ProceduralVegetationEditor/Source/MegaPlants/Private/Implementations/MegaPlantMaterialSettings.h
new file mode 100644
--- /dev/null
+++ b/Engine/Plugins/Experimental/ProceduralVegetationEditor/Source/MegaPlants/Private/Implementations/MegaPlantMaterialSettings.h
@@ -0,0 +1,117 @@
+// Copyright Epic Games, Inc. All Rights Reserved.
+
+#pragma once
+
+#include "CoreMinimal.h"
+#include "Facades/MegaPlantPointFacade.h"
+#include "Materials/MaterialInterface.h"
+#include "GeometryCollection/ManagedArrayCollection.h"
+
+#include "MegaPlantMaterialSettings.generated.h"
+
+UENUM()
+enum class EGenerationOffsetMethod : uint8
+{
+	Clamped,
+	Refit,
+};
+
+UENUM()
+enum class EMaterialDistributionMethod : uint8
+{
+	Repeat,
+	Fit,
+};
+
+UENUM()
+enum class EYTextureMode : uint8
+{
+	Default,
+	Fit0_1,
+	//Fit0_1_Adjusted,
+};
+
+UENUM()
+enum class EUVMaterialMode : uint8
+{
+	Generation,
+	Age,
+	Radius,
+};
+
+USTRUCT()
+struct FTrunkGenerationMaterialSetup
+{
+	GENERATED_BODY()
+
+	UPROPERTY(EditAnywhere, Category="Materials")
+	TObjectPtr<UMaterialInterface> Material;
+
+	UPROPERTY(EditAnywhere, Category="Materials")
+	EYTextureMode YTextureMode = EYTextureMode::Default;
+
+	UPROPERTY(EditAnywhere, Category="Materials", meta=(EditCondition = "YTextureMode == EYTextureMode::Default"))
+	float YScale = 1.0f;
+
+	UPROPERTY(EditAnywhere, Category="Materials", meta=(EditCondition = "YTextureMode == EYTextureMode::Default"))
+	float YOffset = 0.0f;
+
+	UPROPERTY(EditAnywhere, Category="Materials", meta=(EditCondition = "YTextureMode == EYTextureMode::Default"))
+	float YOffsetRandom = 0.0f;
+
+	UPROPERTY(EditAnywhere, Category="Materials")
+	float XOffset = 0.0f;
+
+	UPROPERTY(EditAnywhere, Category="Materials")
+	float XOffsetRandom = 0.0f;
+	
+};
+
+struct FBranchUVData;
+struct FPointUVData;
+
+USTRUCT()
+struct FMegaPlantMaterialSettings
+{
+	GENERATED_BODY()
+
+	UPROPERTY(EditAnywhere, Category="Materials")
+	int32 GenerationOffset = 0;
+
+	UPROPERTY(EditAnywhere, Category="Materials")
+	EGenerationOffsetMethod OffsetMethod = EGenerationOffsetMethod::Clamped;
+
+	UPROPERTY(EditAnywhere, Category="Materials")
+	EMaterialDistributionMethod DistributionMethod = EMaterialDistributionMethod::Repeat;
+
+	UPROPERTY(EditAnywhere, Category="Materials" , DisplayName="Generation Materials")
+	TArray<FTrunkGenerationMaterialSetup> MaterialSetups;
+
+	UPROPERTY(EditAnywhere, Category="Materials")
+	EUVMaterialMode MaterialMode = EUVMaterialMode::Generation;
+
+	/*
+	 Function takes the Collection as input and sets the branchUVMaterial for Branch,
+	 calculates the V value of the UV according to the settings and write into Point Group
+	 calculates the U offset and write to the Point Group in collection
+	 */
+	void ApplyMaterialSettings(FManagedArrayCollection& Collection) const;
+
+	void SetMaterial(FString TrunkMaterial, int32 Index);
+
+private:
+
+	void SetBranchUVMaterial(FManagedArrayCollection& Collection, TArray<int32>& UsedMaterials, TArray<FBranchUVData>& BranchUVData) const;
+
+	float GetMaterialValue(float Value, float MinValue, float MaxValue, float MinGeneration, float MaxGeneration) const;
+
+	void SetUVOffsetAndScale(FManagedArrayCollection& Collection, TArray<FBranchUVData>& BranchUVData) const;
+	
+	void SetBaseUVs(FManagedArrayCollection& Collection, TArray<FPointUVData>& PointUVData) const;
+
+	TArray<float> MakeSegmentArray(const MegaPlants::Facades::FPointFacade& PointFacade,const TArray<int32>& Points, const float MaxPointScale, const int Mode) const;
+
+	void FitUVs01(FManagedArrayCollection& Collection, TArray<FPointUVData>& PointUVData, const TArray<int32>& UsedMaterials, TArray<FBranchUVData>& BranchUVData) const;
+
+	void SetOutputUVs(FManagedArrayCollection& Collection, TArray<FPointUVData>& PointUVData, TArray<FBranchUVData>& BranchUVData) const;
+};
