Create ProceduralVegetationEditor plugin

diff --git a/Engine/Plugins/Experimental/ProceduralVegetationEditor/Source/MegaPlants/Private/Implementations/MegaPlantMesher.cpp b/Engine/Plugins/Experimental/ProceduralVegetationEditor/Source/MegaPlants/Private/Implementations/MegaPlantMesher.cpp
new file mode 100644
--- /dev/null
+++ b/Engine/Plugins/Experimental/ProceduralVegetationEditor/Source/MegaPlants/Private/Implementations/MegaPlantMesher.cpp
@@ -0,0 +1,911 @@
+// Copyright Epic Games, Inc. All Rights Reserved.
+
+#include "MegaPlantMesher.h"
+
+#include <Facades/MegaPlantProfileFacade.h>
+
+#include "MegaPlantsModule.h"
+#include "Algo/MaxElement.h"
+#include "Facades/MegaPlantBoneFacade.h"
+#include "Facades/MegaPlantBranchFacade.h"
+#include "Facades/MegaPlantBudVectorsFacade.h"
+#include "Facades/MegaPlantFoliageFacade.h"
+#include "Facades/MegaPlantPointFacade.h"
+
+void FMegaPlantMesher::GenerateGeometryCollection(const FManagedArrayCollection& InSkeletonCollection, const FMegaPlantMesherSettings& MesherSettings,
+                                                  FGeometryCollection& OutGeometryCollection)
+{
+	InSkeletonCollection.CopyTo(&OutGeometryCollection);
+
+	TArray<FLocalDynamicMeshData> MeshesDataCollection;
+
+	MegaPlants::Facades::FBoneFacade::DefineSchema(OutGeometryCollection);
+	MegaPlants::Facades::FBoneFacade BoneFacade = MegaPlants::Facades::FBoneFacade(OutGeometryCollection);
+
+	SetNjordPixelID(OutGeometryCollection);
+
+	MesherSettings.MaterialSettings.ApplyMaterialSettings(OutGeometryCollection);
+	GetLocalDynamicMeshData(OutGeometryCollection, MesherSettings, MeshesDataCollection);
+
+	uint32 VertexCount = 0;
+	uint32 TriangleCount = 0;
+	for (auto [Vertices, Triangles, MaterialID] : MeshesDataCollection)
+	{
+		VertexCount += Vertices.Num();
+		TriangleCount += Triangles.Num();
+	}
+
+	TManagedArray<float>& NjordPixelIds = BoneFacade.ModifyNjordPixelIds();
+
+	OutGeometryCollection.AddElements(1, OutGeometryCollection.TransformGroup);
+	OutGeometryCollection.AddElements(1, OutGeometryCollection.GeometryGroup);
+	OutGeometryCollection.AddElements(TriangleCount, OutGeometryCollection.FacesGroup);
+	OutGeometryCollection.AddElements(VertexCount, OutGeometryCollection.VerticesGroup);
+	OutGeometryCollection.SetNumUVLayers(1);
+
+	TManagedArray<FVector3f>& Vertices = OutGeometryCollection.Vertex;
+	TManagedArray<FVector3f>& Normals = OutGeometryCollection.Normal;
+	TManagedArray<FVector2f>& UVs = *OutGeometryCollection.FindUVLayer(0);
+	TManagedArray<FIntVector>& Indices = OutGeometryCollection.Indices;
+	TManagedArray<int32>& MaterialIndex = OutGeometryCollection.MaterialIndex;
+	TManagedArray<int32>& MaterialIDAttribute = OutGeometryCollection.MaterialID;
+
+	struct FSection
+	{
+		int32 MaterialID;
+		int32 FirstIndex;
+		int32 NumTriangles;
+		int32 MinVertexIndex;
+		int32 MaxVertexIndex;
+	};
+
+	int32 VertexIndex = 0;
+	int32 FaceIndex = 0;
+	int32 PrevMaterialID = -1;
+	TArray<FSection> Sections;
+
+	//Sorting by material id so sections are created according to material ids
+	MeshesDataCollection.Sort([](const FLocalDynamicMeshData& A, const FLocalDynamicMeshData& B)
+		{
+			return A.MaterialID > B.MaterialID;
+		});
+
+	for (auto MeshIndex = MeshesDataCollection.Num() - 1; MeshIndex >= 0; --MeshIndex)
+	{
+		int MaterialID = MeshesDataCollection[MeshIndex].MaterialID;
+		if (PrevMaterialID != MaterialID)
+		{
+			PrevMaterialID = MaterialID;
+			Sections.Add({MaterialID, FaceIndex, 0, VertexIndex, 0});
+		}
+
+		const FLocalDynamicMeshData& MeshesData = MeshesDataCollection[MeshIndex];
+		for (const FIntVector4& Triangle : MeshesData.Triangles)
+		{
+			const FIntVector4 OffsetTriangle = Triangle + FIntVector4(VertexIndex);
+			Indices[FaceIndex] = FIntVector(OffsetTriangle);
+			MaterialIndex[FaceIndex] = OffsetTriangle.W;
+			MaterialIDAttribute[FaceIndex] = MaterialID;
+			++FaceIndex;
+		}
+		for (const FLocalDynamicMeshData::FVertex& Vertex : MeshesData.Vertices)
+		{
+			Vertices[VertexIndex] = static_cast<FVector3f>(Vertex.Position);
+			Normals[VertexIndex] = static_cast<FVector3f>(Vertex.Normal);
+			UVs[VertexIndex] = Vertex.UV;
+			NjordPixelIds[VertexIndex] = Vertex.NjordPixelId;
+			++VertexIndex;
+		}
+
+		int32 CurrentSectionIndex = Sections.Num() - 1;
+		if (Sections.IsValidIndex(CurrentSectionIndex))
+		{
+			Sections[CurrentSectionIndex].MaxVertexIndex = VertexIndex;
+			Sections[CurrentSectionIndex].NumTriangles = FaceIndex - Sections[CurrentSectionIndex].FirstIndex;
+		}
+	}
+
+	OutGeometryCollection.SimulationType[0] = FGeometryCollection::ESimulationTypes::FST_Rigid;
+	OutGeometryCollection.TransformToGeometryIndex[0] = 0;
+	OutGeometryCollection.TransformIndex[0] = 0;
+	OutGeometryCollection.BoneName[0] = "TreeTest";
+	OutGeometryCollection.BoneColor[0] = FLinearColor::White;
+	OutGeometryCollection.FaceStart[0] = 0;
+	OutGeometryCollection.FaceCount[0] = TriangleCount;
+	OutGeometryCollection.VertexStart[0] = 0;
+	OutGeometryCollection.VertexCount[0] = VertexCount;
+
+	OutGeometryCollection.Visible.Fill(true);
+	OutGeometryCollection.BoneMap.Fill(0);
+	OutGeometryCollection.TangentU.Fill(FVector3f::ForwardVector);
+	OutGeometryCollection.TangentV.Fill(FVector3f::RightVector);
+	OutGeometryCollection.Color.Fill(FLinearColor::White);
+
+	auto& MaterialSetups = MesherSettings.MaterialSettings.MaterialSetups;
+	OutGeometryCollection.AddElements(Sections.Num(), FGeometryCollection::MaterialGroup);
+	auto& MaterialArray = OutGeometryCollection.AddAttribute<FString>("MaterialPath", FGeometryCollection::MaterialGroup);
+
+	int32 SectionIndex = 0;
+	for (auto Section : Sections)
+	{
+		TObjectPtr<UMaterialInterface> SectionMaterial;
+		if (MaterialSetups.IsValidIndex(Section.MaterialID))
+		{
+			SectionMaterial = MaterialSetups[Section.MaterialID].Material;
+		}
+		MaterialArray[SectionIndex] = SectionMaterial
+			? SectionMaterial->GetPathName()
+			: "";
+
+		TManagedArray<FGeometryCollectionSection>& GeometrySections = OutGeometryCollection.Sections;
+
+		GeometrySections[SectionIndex].MaterialID = Section.MaterialID;
+		GeometrySections[SectionIndex].FirstIndex = Section.FirstIndex;
+		GeometrySections[SectionIndex].NumTriangles = Section.NumTriangles;
+		GeometrySections[SectionIndex].MinVertexIndex = Section.MinVertexIndex;
+		GeometrySections[SectionIndex].MaxVertexIndex = Section.MaxVertexIndex;
+
+		SectionIndex++;
+	}
+
+	OutGeometryCollection.UpdateBoundingBox();
+}
+
+void FMegaPlantMesher::SetNjordPixelID(FGeometryCollection& OutGeometryCollection)
+{
+	MegaPlants::Facades::FPointFacade PointFacade(OutGeometryCollection);
+	TManagedArray<float>& NjordPixelIds = PointFacade.ModifyNjordPixelIDs();
+
+	TMap<int, TArray<int>> BudNumberPointIndicesMap;
+	for (int PointIndex = 0; PointIndex < PointFacade.GetElementCount(); PointIndex++)
+	{
+		int BudNumber = PointFacade.GetBudNumber(PointIndex);
+		BudNumberPointIndicesMap.FindOrAdd(BudNumber).Add(PointIndex);
+	}
+
+	for (const auto& KeyValuePair : BudNumberPointIndicesMap)
+	{
+		int ValueIndex = 0;
+		for (auto Value : KeyValuePair.Value)
+		{
+			//Modify NjordPixelID
+			float NjordPixelID = KeyValuePair.Key + (float)ValueIndex / KeyValuePair.Value.Num();
+			NjordPixelIds[Value] = NjordPixelID;
+			ValueIndex++;
+		}
+	}
+}
+
+void FMegaPlantMesher::GenerateDynamicMesh(FManagedArrayCollection& Collection, const FMegaPlantMesherSettings& MesherSettings,
+                                           TObjectPtr<UDynamicMesh>& OutMesh)
+{
+	TArray<FLocalDynamicMeshData> MeshesDataCollection;
+	GetLocalDynamicMeshData(Collection, MesherSettings, MeshesDataCollection);
+
+	uint32 VertexCount = 0;
+	uint32 TriangleCount = 0;
+	for (auto [Vertices, Triangles , MaterialID] : MeshesDataCollection)
+	{
+		VertexCount += Vertices.Num();
+		TriangleCount += Triangles.Num();
+	}
+
+	OutMesh->InitializeMesh();
+	OutMesh->EditMesh([&](UE::Geometry::FDynamicMesh3& InternalMesh)
+			{
+				FMegaPlantMeshGenerator Generator(VertexCount, TriangleCount, &MeshesDataCollection);
+				Generator.Generate();
+				InternalMesh.Copy(&Generator);
+			},
+		EDynamicMeshChangeType::MeshChange);
+}
+
+TSet<int32> FMegaPlantMesher::CollectHardPoints(const MegaPlants::Facades::FBranchFacade& BranchFacade,
+                                                const MegaPlants::Facades::FPointFacade& PointFacade)
+{
+	TSet<int32> HardPoints;
+	const int32 MainTrunkIndex = BranchFacade.GetTrunkIndex();
+	const int32 SeedPointIndex = BranchFacade.GetPoints(MainTrunkIndex)[0];
+	TSet<int32> BranchSourceBudNumbers;
+	BranchSourceBudNumbers.Reserve(BranchFacade.GetElementCount());
+
+	for (int32 BranchIndex = 0; BranchIndex < BranchFacade.GetElementCount(); ++BranchIndex)
+	{
+		const int32 SourceBudNumber = BranchFacade.GetBranchSourceBudNumber(BranchIndex);
+		BranchSourceBudNumbers.Add(SourceBudNumber);
+	}
+
+	for (int32 BranchIndex = 0; BranchIndex < BranchFacade.GetElementCount(); ++BranchIndex)
+	{
+		const TArray<int32>& BranchPoints = BranchFacade.GetPoints(BranchIndex);
+		const int32 LastPointIndex = BranchPoints.Last();
+		const int32 FirstPointIndex = BranchPoints[0];
+
+		for (int32 i = 0; i < BranchPoints.Num(); ++i)
+		{
+			const int32 PointIndex = BranchPoints[i];
+			const int32 PointBudNumber = PointFacade.GetBudNumber(PointIndex);
+			const float NjordPixelIndex = PointFacade.GetNjordPixelIndex(PointIndex);
+
+			bool bIsBranchJoint = (NjordPixelIndex == FMath::FloorToFloat(NjordPixelIndex)) && BranchSourceBudNumbers.Contains(PointBudNumber);
+
+			if (PointIndex == SeedPointIndex ||
+				PointIndex == FirstPointIndex ||
+				PointIndex == LastPointIndex ||
+				bIsBranchJoint
+			)
+			{
+				HardPoints.Add(PointIndex);
+			}
+		}
+	}
+
+	return HardPoints;
+}
+
+TSet<int32> FMegaPlantMesher::ComputePointGradients(const MegaPlants::Facades::FPointFacade& PointFacade,
+                                                    const FMegaPlantMesherSettings& MesherSettings, const TSet<int32>& HardPoints,
+                                                    const float MaxPointScale, TArray<float>& OutMeshDivisionsGradients,
+                                                    TArray<float>& OutDeltaModifiers)
+{
+	TSet<int32> PointsToRemove;
+	OutMeshDivisionsGradients.Reserve(PointFacade.GetElementCount());
+	OutDeltaModifiers.Reserve(PointFacade.GetElementCount());
+
+	for (int32 PointIndex = 0; PointIndex < PointFacade.GetElementCount(); ++PointIndex)
+	{
+		const float HullGradient = PointFacade.GetHullGradient(PointIndex);
+		const float MainTrunkGradient = PointFacade.GetMainTrunkGradient(PointIndex);
+		const float GroundGradient = PointFacade.GetGroundGradient(PointIndex);
+		const float ScaleGradient = PointFacade.GetPointScaleGradient(PointIndex);
+
+		// Compute scale removal gradient
+		const float PointHullGradient =
+			MesherSettings.PointHullRetentionGradient.GetRichCurveConst()->Eval(HullGradient) * MesherSettings.PointHullRetention;
+
+		const float PointMainTrunkGradient =
+			MesherSettings.PointMainTrunkRetentionGradient.GetRichCurveConst()->Eval(MainTrunkGradient) * MesherSettings.PointMainTrunkRetention;
+
+		const float PointGroundGradient =
+			MesherSettings.PointGroundRetentionGradient.GetRichCurveConst()->Eval(GroundGradient) * MesherSettings.PointGroundRetention;
+
+		const float ScaleRemovalGradient = (ScaleGradient + PointHullGradient + PointMainTrunkGradient + PointGroundGradient) * MaxPointScale /
+			100.0f;
+
+		// Compute mesh divisions gradient
+		const float DivisionsScaleGradient = MesherSettings.DivisionScaleRetentionGradient.GetRichCurveConst()->Eval(ScaleGradient);
+		const float DivisionsHullGradient = MesherSettings.DivisionHullRetentionGradient.GetRichCurveConst()->Eval(HullGradient);
+		const float DivisionsMainTrunkGradient = MesherSettings.DivisionMainTrunkRetentionGradient.GetRichCurveConst()->Eval(MainTrunkGradient);
+		const float DivisionsGroundGradient = MesherSettings.DivisionGroundRetentionGradient.GetRichCurveConst()->Eval(GroundGradient);
+
+		const float InterpolatedHullGradient = FMath::GetMappedRangeValueClamped(
+			FVector2f(0.0f, 1.0f),
+			FVector2f(DivisionsHullGradient, 1.0f),
+			1.0f - MesherSettings.DivisionHullRetention);
+		const float InterpolatedMainTrunkGradient = FMath::GetMappedRangeValueClamped(
+			FVector2f(0.0f, 1.0f),
+			FVector2f(DivisionsMainTrunkGradient, 1.0f),
+			1.0f - MesherSettings.DivisionMainTrunkRetention);
+		const float InterpolatedGroundGradient = FMath::GetMappedRangeValueClamped(
+			FVector2f(0.0f, 1.0f),
+			FVector2f(DivisionsGroundGradient, 1.0f),
+			1.0f - MesherSettings.DivisionGroundRetention);
+
+		const float MeshDivisionsGradient = DivisionsScaleGradient * InterpolatedHullGradient * InterpolatedMainTrunkGradient *
+			InterpolatedGroundGradient;
+		OutMeshDivisionsGradients.Add(MeshDivisionsGradient);
+
+		// Compute delta modifier gradient
+		const float SegmentScaleGradient = MesherSettings.SegmentScaleRetentionGradient.GetRichCurveConst()->Eval(ScaleGradient);
+		const float SegmentHullGradient = MesherSettings.SegmentHullRetentionGradient.GetRichCurveConst()->Eval(HullGradient);
+		const float SegmentMainTrunkGradient = MesherSettings.SegmentMainTrunkRetentionGradient.GetRichCurveConst()->Eval(MainTrunkGradient);
+		const float SegmentGroundGradient = MesherSettings.SegmentGroundRetentionGradient.GetRichCurveConst()->Eval(GroundGradient);
+
+		const float DeltaModifierGradient =
+			(SegmentScaleGradient * MesherSettings.SegmentScaleRetention) +
+			(SegmentHullGradient * MesherSettings.SegmentHullRetention) +
+			(SegmentMainTrunkGradient * MesherSettings.SegmentMainTrunkRetention) +
+			(SegmentGroundGradient * MesherSettings.SegmentGroundRetention);
+
+		// Should the point be removed per scale gradient
+		if (!HardPoints.Contains(PointIndex) && MesherSettings.PointRemoval > 0.0f && ScaleRemovalGradient < MesherSettings.PointRemoval)
+		{
+			PointsToRemove.Add(PointIndex);
+		}
+
+		// Compute delta modifier
+		float DeltaModifier = 1.0f;
+		if (!HardPoints.Contains(PointIndex))
+		{
+			DeltaModifier = MesherSettings.SegmentRetentionImpact * (MaxPointScale / 100.0f) * DeltaModifierGradient;
+		}
+		OutDeltaModifiers.Add(DeltaModifier);
+	}
+
+	return PointsToRemove;
+}
+
+void FMegaPlantMesher::PerformPathSimplification(const MegaPlants::Facades::FBranchFacade& BranchFacade,
+                                                 const MegaPlants::Facades::FPointFacade& PointFacade, const FMegaPlantMesherSettings& MesherSettings,
+                                                 const TSet<int32>& HardPoints, const TArray<float>& DeltaModifiers, TSet<int32>& InOutPointsToRemove)
+{
+	const int32 Iterations = MesherSettings.Accuracy;
+	// Epsilon is the maximum allowed curvature error
+	const float Epsilon = MesherSettings.SegmentReduction / static_cast<float>(Iterations);
+
+	for (int32 BranchIndex = 0; BranchIndex < BranchFacade.GetElementCount(); ++BranchIndex)
+	{
+		TArray<int32> BranchPoints = BranchFacade.GetPoints(BranchIndex).FilterByPredicate(
+			[&InOutPointsToRemove](const int32 PointIndex)
+				{
+					return !InOutPointsToRemove.Contains(PointIndex);
+				});
+
+		for (int32 j = 0; j < Iterations; ++j)
+		{
+			const float IterationEpsilon = Epsilon * (j + 1);
+			for (int32 i = 1; i < BranchPoints.Num() - 1; ++i)
+			{
+				const int32 Pt0Index = BranchPoints[i - 1],
+				            Pt1Index = BranchPoints[i],
+				            Pt2Index = BranchPoints[i + 1];
+				const FVector Pt0Position = FVector(PointFacade.GetPosition(Pt0Index)),
+				              Pt1Position = FVector(PointFacade.GetPosition(Pt1Index)),
+				              Pt2Position = FVector(PointFacade.GetPosition(Pt2Index));
+				const float DeltaModifier = DeltaModifiers[Pt1Index];
+				const float Delta = FMath::PointDistToLine(Pt1Position, Pt2Position - Pt0Position, Pt0Position) + DeltaModifier;
+
+				if (Delta < IterationEpsilon && !HardPoints.Contains(Pt1Index))
+				{
+					InOutPointsToRemove.Add(Pt1Index);
+					BranchPoints.Remove(Pt1Index);
+					i--;
+				}
+			}
+		}
+
+		// Ensure long chains of removed points do not exceed a maximum threshold
+		BranchPoints = BranchFacade.GetPoints(BranchIndex);
+		FVector3f PreviousPointPosition = PointFacade.GetPosition(BranchPoints[0]);
+		float TravelDistance = 0.0f;
+
+		for (int32 i = 0; i < BranchPoints.Num(); ++i)
+		{
+			const int32 CurrentPointIndex = BranchPoints[i];
+			FVector3f CurrentPointPosition = PointFacade.GetPosition(CurrentPointIndex);
+			const float Distance = FVector3f::Distance(PreviousPointPosition, CurrentPointPosition) / 100.0f;
+
+			if (InOutPointsToRemove.Contains(CurrentPointIndex))
+			{
+				TravelDistance += Distance;
+			}
+			else
+			{
+				TravelDistance = 0.0f;
+			}
+
+			if (TravelDistance > MesherSettings.LongestSegmentLength)
+			{
+				InOutPointsToRemove.Remove(CurrentPointIndex);
+				TravelDistance = 0.0f;
+			}
+
+			PreviousPointPosition = CurrentPointPosition;
+		}
+
+		// Ensure two kept points are never closer than a minimum threshold
+		BranchPoints = BranchFacade.GetPoints(BranchIndex).FilterByPredicate(
+			[&InOutPointsToRemove](const int32 PointIndex)
+				{
+					return !InOutPointsToRemove.Contains(PointIndex);
+				});
+		PreviousPointPosition = PointFacade.GetPosition(BranchPoints[0]);
+
+		for (int32 i = 0; i < BranchPoints.Num(); ++i)
+		{
+			const int32 CurrentPointIndex = BranchPoints[i];
+			FVector3f CurrentPointPosition = PointFacade.GetPosition(CurrentPointIndex);
+			const float Distance = FVector3f::Distance(PreviousPointPosition, CurrentPointPosition) / 100.0f;
+
+			if (Distance < MesherSettings.ShortestSegmentLength && !HardPoints.Contains(CurrentPointIndex))
+			{
+				InOutPointsToRemove.Add(CurrentPointIndex);
+			}
+			else
+			{
+				PreviousPointPosition = CurrentPointPosition;
+			}
+		}
+	}
+}
+
+TArray<int32> FMegaPlantMesher::ComputeMeshDivisions(const MegaPlants::Facades::FBranchFacade& BranchFacade,
+                                                     const FMegaPlantMesherSettings& MesherSettings, const TArray<float>& MeshDivisionsGradients,
+                                                     const int32 PointCount)
+{
+	TArray<int32> TargetMeshDivisions;
+	static constexpr int32 MinimumDivisions = 3;
+	TargetMeshDivisions.Init(MinimumDivisions, PointCount);
+
+	for (int32 BranchIndex = 0; BranchIndex < BranchFacade.GetElementCount(); ++BranchIndex)
+	{
+		const TArray<int32>& BranchPoints = BranchFacade.GetPoints(BranchIndex);
+		const float MinDivisions = static_cast<float>(MesherSettings.MinMeshDivisions);
+		const float MaxDivisions = static_cast<float>(MesherSettings.MaxMeshDivisions);
+		const FVector2f NormalRange(0.0f, 1.0f);
+		const FVector2f DivisionsRange(MinDivisions, MaxDivisions);
+
+		for (int32 i = 0; i < BranchPoints.Num(); ++i)
+		{
+			const int32 CurrentPointIndex = BranchPoints[i];
+			const float MeshDivisionsGradient = MeshDivisionsGradients[CurrentPointIndex];
+
+			const int32 TargetMeshDivisionsForCurrentPoint = FMath::RoundToInt32(
+				FMath::GetMappedRangeValueClamped(
+					NormalRange,
+					DivisionsRange,
+					MeshDivisionsGradient));
+
+			TargetMeshDivisions[CurrentPointIndex] = TargetMeshDivisionsForCurrentPoint;
+		}
+	}
+
+	return TargetMeshDivisions;
+}
+
+void FMegaPlantMesher::TriangulateRings(const TArray<int32>& PreviousIndices, const TArray<int32>& CurrentIndices, int32& InOutPolyGroupIndex,
+                                        FLocalDynamicMeshData& OutMeshData)
+{
+	if (!PreviousIndices.IsEmpty())
+	[[likely]]
+	{
+		const float IndexDelta1 = PreviousIndices.Num() < CurrentIndices.Num()
+			? static_cast<float>(PreviousIndices.Num()) / CurrentIndices.Num()
+			: 1.0f;
+		const float IndexDelta2 = CurrentIndices.Num() < PreviousIndices.Num()
+			? static_cast<float>(CurrentIndices.Num()) / PreviousIndices.Num()
+			: 1.0f;
+
+		for (
+			float Index1 = 0.0f, Index2 = 0.0f;
+			FMath::RoundToInt32(Index1 + IndexDelta1) <= PreviousIndices.Num() && FMath::RoundToInt32(Index2 + IndexDelta2) <=
+			CurrentIndices.Num();
+			Index1 += IndexDelta1, Index2 += IndexDelta2
+		)
+		{
+			int32 VertexIndices[4] = {
+				PreviousIndices[FMath::RoundToInt32(Index1) % PreviousIndices.Num()],
+				CurrentIndices[FMath::RoundToInt32(Index2) % CurrentIndices.Num()],
+				PreviousIndices[FMath::RoundToInt32(Index1 + IndexDelta1) % PreviousIndices.Num()],
+				CurrentIndices[FMath::RoundToInt32(Index2 + IndexDelta2) % CurrentIndices.Num()],
+			};
+
+			if (VertexIndices[0] != VertexIndices[2])
+			{
+				OutMeshData.Triangles.Emplace(VertexIndices[0], VertexIndices[1], VertexIndices[2], InOutPolyGroupIndex);
+			}
+			if (VertexIndices[1] != VertexIndices[3])
+			{
+				OutMeshData.Triangles.Emplace(VertexIndices[1], VertexIndices[3], VertexIndices[2], InOutPolyGroupIndex);
+			}
+
+			++InOutPolyGroupIndex;
+		}
+	}
+}
+
+float FMegaPlantMesher::GetProfileMultiplier(const TArray<float>& InProfilePoints, const float ProfileUV_U)
+{
+	if (InProfilePoints.IsEmpty())
+	{
+		return 1.0f;
+	}
+
+	const int32 MaxPointsIndex = InProfilePoints.Num() - 1;
+	const int32 ProfileIndex0 = FMath::Min(FMath::FloorToInt(ProfileUV_U * 100), MaxPointsIndex);
+	const int32 ProfileIndex1 = FMath::Min(FMath::CeilToInt(ProfileUV_U * 100), MaxPointsIndex);
+
+	const float Scale0 = InProfilePoints[ProfileIndex0];
+	const float Scale1 = InProfilePoints[ProfileIndex1];
+
+	const float BlendValue = FMath::GetMappedRangeValueClamped(
+		FVector2f(static_cast<float>(ProfileIndex0), static_cast<float>(ProfileIndex1)),
+		FVector2f(0.0f, 1.0f),
+		ProfileUV_U * 100);
+
+	const float ProfileMultiplier = FMath::Lerp(Scale0, Scale1, BlendValue);
+
+	return ProfileMultiplier;
+}
+
+TMap<int32, TArray<int32>> FMegaPlantMesher::GetPointsIndicesToFoliageIndicesMap(const FManagedArrayCollection& Collection)
+{
+	TMap<int32, TArray<int32>> Map;
+	const MegaPlants::Facades::FPointFacade PointFacade(Collection);
+	const MegaPlants::Facades::FBranchFacade BranchFacade(Collection);
+	const MegaPlants::Facades::FFoliageFacade FoliageFacade(Collection);
+
+	for (int32 BranchIndex = 0; BranchIndex < BranchFacade.GetElementCount(); ++BranchIndex)
+	{
+		const TArray<int32>& FoliageEntryIds = FoliageFacade.GetFoliageEntryIdsForBranch(BranchIndex);
+		const TArray<int32>& BranchPoints = BranchFacade.GetPoints(BranchIndex);
+
+		if (BranchPoints.Num() == 0)
+		{
+			continue;
+		}
+
+		for (int32 FoliageEntryId : FoliageEntryIds)
+		{
+			// Find the closest point for current entry's pivot position
+			const FVector3f PivotPosition = FoliageFacade.GetPivotPoint(FoliageEntryId);
+			const int32* Closest = Algo::MinElement(BranchPoints, [PivotPosition, PointFacade](const int32& PointIndexA, const int32& PointIndexB)
+				{
+					return FVector3f::Dist(PointFacade.GetPosition(PointIndexA), PivotPosition) <
+						FVector3f::Dist(PointFacade.GetPosition(PointIndexB), PivotPosition);
+				});
+			int32 ClosestPointIndex = Closest
+				? *Closest
+				: BranchPoints[0];
+
+			Map.FindOrAdd(ClosestPointIndex).Add(FoliageEntryId);
+		}
+	}
+
+	return Map;
+}
+
+void FMegaPlantMesher::GenerateBranchMeshData(const bool bPrimitiveIsTrunk, const int32 GenerationNumber, const TArray<int32>& PrimitivePoints,
+                                              const FDisplacementData& DisplacementData, const FMegaPlantMesherSettings& MesherSettings,
+                                              const TArray<int32>& TargetMeshDivisions, const FManagedArrayCollection& Collection,
+                                              FLocalDynamicMeshData& OutLocalMeshData)
+{
+	const MegaPlants::Facades::FPointFacade PointFacade(Collection);
+	const MegaPlants::Facades::FBudVectorsFacade BudVectorFacade(Collection);
+	const MegaPlants::Facades::FPlantProfileFacade PlantProfileFacade(Collection);
+
+	bool bShouldApplyProfile =
+		PlantProfileFacade.NumProfileEntries() > 0
+		&& MesherSettings.PlantProfileOptions.Num() > 0
+		&& (MesherSettings.bApplyProfileToBranches || bPrimitiveIsTrunk);
+
+	TArray<float> ProfilePoints;
+	if (bShouldApplyProfile)
+	{
+		if (const int32 ProfileId = MesherSettings.PlantProfileOptions.Find(MesherSettings.SelectedPlantProfile) - 1;
+			ProfileId > -1)
+		{
+			ProfilePoints = PlantProfileFacade.GetProfilePoints(ProfileId);
+		}
+		else
+		{
+			bShouldApplyProfile = false;
+		}
+	}
+
+	const TManagedArray<FVector3f>& PointPositions = PointFacade.GetPositions();
+	const TManagedArray<TArray<FVector3f>>& PointBudDirections = BudVectorFacade.GetBudDirections();
+	const TManagedArray<float>& PointScales = PointFacade.GetPointScales();
+
+	const float MaxPointScale = *Algo::MaxElement(PointScales);
+	const float MaxScaleRatio = 1.0f / (MaxPointScale * UE_TWO_PI);
+	int32 PolyGroupIndex = 0;
+	TArray<int32> CurrentIndices;
+	TArray<int32> PreviousIndices;
+
+	FVector3f PreviousPosition = PointPositions[PrimitivePoints[0]];
+	FVector3f PreviousUpVector = FVector3f::ForwardVector;
+
+	bool bShouldApplyDisplacement =
+		DisplacementData.Values.Num() > 0
+		&& DisplacementData.TextureWidth > 0
+		&& DisplacementData.TextureHeight > 0
+		&& GenerationNumber <= MesherSettings.DisplacementGenerationUpperLimit;
+
+	const int32 MaxIndex = PrimitivePoints.Num() - 1;
+	for (int32 Index = 0; Index <= MaxIndex; ++Index)
+	{
+		const int32& PointIndex = PrimitivePoints[Index];
+		const float NjordPixelId = PointFacade.GetNjordPixelIndex(PointIndex);
+
+		float TextureCoordV = PointFacade.GetTextureCoordV(PointIndex);
+		float TextureCoordUOffset = PointFacade.GetTextureCoordUOffset(PointIndex);
+
+		const int32 PointDivisions = TargetMeshDivisions[PointIndex];
+		const FVector3f& PointPosition = PointPositions[PointIndex];
+		const float PointScale = PointScales[PointIndex];
+		FVector3f NormalVector = PointBudDirections[PointIndex][0].GetUnsafeNormal();
+		const float PlantGradient = PointFacade.GetPlantGradient(PointIndex);
+		const float ProfileFallOff = MesherSettings.PlantProfileFallOff.GetRichCurveConst()->Eval(PlantGradient);
+
+		const FVector3f PointUpOriginal = PointBudDirections[PointIndex][5].GetUnsafeNormal();
+		FVector3f PointUp = FVector3f::VectorPlaneProject(PointUpOriginal, NormalVector).GetSafeNormal();
+
+		if (Index > 0 && FVector3f::DotProduct(PointUp, PreviousUpVector) < 0.0f)
+		{
+			PointUp = -PointUp;
+		}
+		PreviousUpVector = PointUp;
+
+		if (PointScale == 0.0f)
+		[[unlikely]]
+		{
+			CurrentIndices.Add(OutLocalMeshData.Vertices.Emplace(PointPosition, NormalVector, FVector2f(0.0f, TextureCoordV), NjordPixelId));
+		}
+		else
+		[[likely]]
+		{
+			const float RotationDelta = UE_TWO_PI / PointDivisions;
+			for (int32 Division = 0; Division < PointDivisions; ++Division)
+			{
+				const float UV_U = (1 - (static_cast<float>(Division) / PointDivisions)) + TextureCoordUOffset;
+
+				float ProfileMultiplier = 1.0f;
+				if (bShouldApplyProfile)
+				{
+					ProfileMultiplier = GetProfileMultiplier(ProfilePoints, UV_U);
+					ProfileMultiplier = FMath::Lerp(1.0f, ProfileMultiplier, ProfileFallOff);
+					ProfileMultiplier = ProfileMultiplier * MesherSettings.PlantProfileScale;
+				}
+
+				const FVector3f Direction = FQuat4f(NormalVector, RotationDelta * Division).RotateVector(PointUp);
+				FVector3f RadialPoint = PointPosition + (Direction * PointScale * ProfileMultiplier);
+
+				if (bShouldApplyDisplacement)
+				{
+					const int32 X = FMath::Clamp(
+						FMath::FloorToInt(
+							UV_U * MesherSettings.DisplacementUVScale.X * (DisplacementData.TextureWidth - 1)
+						) % DisplacementData.TextureWidth,
+						0, DisplacementData.TextureWidth - 1
+					);
+					const int32 Y = FMath::Clamp(
+						FMath::FloorToInt(
+							TextureCoordV * MesherSettings.DisplacementUVScale.Y * (DisplacementData.TextureHeight - 1)
+						) % DisplacementData.TextureHeight,
+						0, DisplacementData.TextureHeight - 1
+					);
+
+					const int32 Idx = Y * DisplacementData.TextureWidth + X;
+					float DisplacementMultiplier =
+						MesherSettings.DisplacementStrength * (DisplacementData.Values[Idx] - MesherSettings.DisplacementBias);
+					DisplacementMultiplier = DisplacementMultiplier * PointScale / MaxPointScale;
+
+					RadialPoint = RadialPoint + (Direction * DisplacementMultiplier);
+				}
+
+				const FVector3f Normal = FVector3f((RadialPoint - PointPosition).GetUnsafeNormal());
+
+				CurrentIndices.Add(OutLocalMeshData.Vertices.Emplace(RadialPoint, Normal, FVector2f(UV_U, TextureCoordV), NjordPixelId));
+			}
+		}
+
+		const FLocalDynamicMeshData::FVertex FirstVertex = OutLocalMeshData.Vertices[CurrentIndices[0]];
+		CurrentIndices.Add(OutLocalMeshData.Vertices.Emplace(FirstVertex.Position, FirstVertex.Normal, FVector2f(TextureCoordUOffset, TextureCoordV),
+			NjordPixelId));
+
+		PreviousPosition = PointPosition;
+
+		TriangulateRings(PreviousIndices, CurrentIndices, PolyGroupIndex, OutLocalMeshData);
+
+		PreviousIndices = MoveTemp(CurrentIndices);
+		CurrentIndices.Empty();
+	}
+}
+
+void FMegaPlantMesher::UpdateFoliagePivotPoints(const TSet<int32>& PointsToRemove, FManagedArrayCollection& OutCollection)
+{
+	const MegaPlants::Facades::FBranchFacade BranchFacade(OutCollection);
+	const MegaPlants::Facades::FPointFacade PointFacade(OutCollection);
+	MegaPlants::Facades::FFoliageFacade FoliageFacade(OutCollection);
+
+	TMap<int32, TArray<int32>> PointsIndicesToFoliageIndicesMap = GetPointsIndicesToFoliageIndicesMap(OutCollection);
+
+	for (int32 BranchIndex = 0; BranchIndex < BranchFacade.GetElementCount(); ++BranchIndex)
+	{
+		const TArray<int32>& BranchPoints = BranchFacade.GetPoints(BranchIndex);
+
+		if (BranchPoints.Num() == 0)
+		{
+			continue;
+		}
+
+		TArray<int32> SortedPointIndices = BranchPoints;
+		SortedPointIndices.Sort([PointFacade](const int32& A, const int32& B)
+			{
+				return PointFacade.GetLengthFromRoot(A) < PointFacade.GetLengthFromRoot(B);
+			});
+
+		for (int32 i = 0; i < SortedPointIndices.Num(); ++i)
+		{
+			const int32 CurrentPointIndex = SortedPointIndices[i];
+			if (PointsToRemove.Contains(CurrentPointIndex) &&
+				i > 0 &&
+				i < SortedPointIndices.Num() - 1 &&
+				PointsIndicesToFoliageIndicesMap.Contains(CurrentPointIndex))
+			{
+				int32 j = i - 1;
+				while (j > 0 && PointsToRemove.Contains(SortedPointIndices[j]))
+				{
+					j--;
+				}
+				const int32 PreviousPointIndex = SortedPointIndices[j];
+
+				int32 k = i + 1;
+				while (k < (SortedPointIndices.Num() - 1) && PointsToRemove.Contains(SortedPointIndices[k]))
+				{
+					k++;
+				}
+				const int32 NextPointIndex = SortedPointIndices[k];
+
+				const float PreviousPointLengthFromRoot = PointFacade.GetLengthFromRoot(PreviousPointIndex);
+				const float CurrentPointLengthFromRoot = PointFacade.GetLengthFromRoot(CurrentPointIndex);
+				const float NextPointLengthFromRoot = PointFacade.GetLengthFromRoot(NextPointIndex);
+				const float BlendValue = FMath::GetMappedRangeValueClamped(
+					FVector2f(PreviousPointLengthFromRoot, NextPointLengthFromRoot),
+					FVector2f(0.0f, 1.0f),
+					CurrentPointLengthFromRoot
+				);
+
+				const FVector3f& CurrentPointPosition = PointFacade.GetPosition(CurrentPointIndex);
+				const FVector3f& PreviousPointPosition = PointFacade.GetPosition(PreviousPointIndex);
+				const FVector3f& NextPointPosition = PointFacade.GetPosition(NextPointIndex);
+
+				const FVector3f InterpolatedPointPosition = FMath::Lerp(PreviousPointPosition, NextPointPosition, BlendValue);
+
+				for (const int32 FoliageIndex : PointsIndicesToFoliageIndicesMap[CurrentPointIndex])
+				{
+					const FVector3f& FoliagePosition = FoliageFacade.GetPivotPoint(FoliageIndex);
+					const FVector3f OffsetVector = FoliagePosition - CurrentPointPosition;
+					const FVector3f UpdatedFoliagePosition = InterpolatedPointPosition + OffsetVector;
+
+					FoliageFacade.SetPivotPoint(FoliageIndex, UpdatedFoliagePosition);
+				}
+			}
+		}
+	}
+}
+
+void FMegaPlantMesher::GetLocalDynamicMeshData(FManagedArrayCollection& Collection, const FMegaPlantMesherSettings& MesherSettings,
+                                               TArray<FLocalDynamicMeshData>& MeshesData)
+{
+	const MegaPlants::Facades::FBranchFacade BranchFacade(Collection);
+	const MegaPlants::Facades::FPointFacade PointFacade(Collection);
+
+	if (!BranchFacade.IsValid())
+	{
+		UE_LOG(LogMegaPlants, Warning, TEXT("Fail to evaluate, InCollection dose not match the schema."));
+		return;
+	}
+
+	const TManagedArray<float>& PointScales = PointFacade.GetPointScales();
+	const float MaxPointScale = *Algo::MaxElement(PointScales);
+
+	const TSet<int32> HardPoints = CollectHardPoints(BranchFacade, PointFacade);
+
+	TArray<float> MeshDivisionsGradientsForPoints;
+	TArray<float> DeltaModifiersForPoints;
+	TSet<int32> PointsToRemove = ComputePointGradients(PointFacade, MesherSettings, HardPoints, MaxPointScale, MeshDivisionsGradientsForPoints,
+		DeltaModifiersForPoints);
+
+	PerformPathSimplification(BranchFacade, PointFacade, MesherSettings, HardPoints, DeltaModifiersForPoints, PointsToRemove);
+
+	// Point removal may update curvate of the branch.
+	// Foliage instance pivot/attachment points need to be updated to compensate for that.
+	if (PointsToRemove.Num() > 0)
+	{
+		UpdateFoliagePivotPoints(PointsToRemove, Collection);
+	}
+
+	const TArray<int32> TargetMeshDivisions = ComputeMeshDivisions(BranchFacade, MesherSettings, MeshDivisionsGradientsForPoints,
+		PointFacade.GetElementCount());
+
+	MeshesData.SetNum(BranchFacade.GetElementCount());
+
+	int32 DisplacementWidth = 0;
+	int32 DisplacementHeight = 0;
+	if (MesherSettings.DisplacementTexture
+		&& MesherSettings.DisplacementValues.Num() > 0
+		&& MesherSettings.DisplacementStrength > 0.0f)
+	{
+		DisplacementWidth = MesherSettings.DisplacementTexture->GetSizeX();
+		DisplacementHeight = MesherSettings.DisplacementTexture->GetSizeY();
+
+		if (MesherSettings.DisplacementTexture->SRGB)
+		{
+			UE_LOG(LogMegaPlants, Warning,
+				TEXT("Displacement texture selected has sRGB enabled! This will corrupt displacement data. Please disable it."));
+		}
+	}
+
+	const FDisplacementData DisplacementData(MesherSettings.DisplacementValues, DisplacementWidth, DisplacementHeight);
+
+	ParallelFor(
+		BranchFacade.GetElementCount(),
+		[&](const int32 PrimitiveIndex)
+			{
+				const TArray<int32>& PrimitivePoints = BranchFacade.GetPoints(PrimitiveIndex).FilterByPredicate(
+					[&PointsToRemove](const int32 PointIndex)
+						{
+							return !PointsToRemove.Contains(PointIndex);
+						});
+
+				const bool PrimitiveIsTrunk = BranchFacade.GetTrunkIndex() == PrimitiveIndex;
+				const int32 GenerationNumber = BranchFacade.GetGenerationNumber(PrimitiveIndex);
+
+				MeshesData[PrimitiveIndex].MaterialID = BranchFacade.GetBranchUVMaterial(PrimitiveIndex);
+
+				GenerateBranchMeshData(PrimitiveIsTrunk, GenerationNumber, PrimitivePoints, DisplacementData, MesherSettings, TargetMeshDivisions,
+					Collection, MeshesData[PrimitiveIndex]);
+			});
+}
+
+void FMegaPlantMesher::ExtractDisplacementData(const TObjectPtr<UTexture2D>& Texture, TArray<float>& OutValues)
+{
+	OutValues.Empty();
+
+#if WITH_EDITOR
+	if (!Texture)
+	{
+		UE_LOG(LogMegaPlants, Warning, TEXT("ExtractDisplacementData: Texture is null."));
+		return;
+	}
+
+	FTextureSource& Source = Texture->Source;
+	if (!Source.IsValid())
+	{
+		UE_LOG(LogMegaPlants, Warning, TEXT("ExtractDisplacementData: No source data found on texture %s."), *Texture->GetName());
+		return;
+	}
+
+	const int32 Width = Source.GetSizeX();
+	const int32 Height = Source.GetSizeY();
+	const int32 NumPix = Width * Height;
+	OutValues.SetNumZeroed(NumPix);
+
+	const uint8* RawData = Source.LockMip(0);
+
+	switch (Source.GetFormat())
+	{
+	case TSF_RGBA32F:
+		{
+			const float* FloatData = reinterpret_cast<const float*>(RawData);
+			for (int32 i = 0; i < NumPix; ++i)
+			{
+				OutValues[i] = FloatData[i * 4 + 0];
+			}
+		}
+		break;
+
+	case TSF_RGBA16F:
+		{
+			const FFloat16* HalfData = reinterpret_cast<const FFloat16*>(RawData);
+			for (int32 i = 0; i < NumPix; ++i)
+			{
+				OutValues[i] = HalfData[i * 4 + 0].GetFloat();
+			}
+		}
+		break;
+	case TSF_BGRA8:
+		{
+			const FColor* ColorData = reinterpret_cast<const FColor*>(RawData);
+			constexpr float Inv255 = 1.0f / 255.0f;
+			for (int32 i = 0; i < NumPix; ++i)
+			{
+				OutValues[i] = ColorData[i].R * Inv255;
+			}
+		}
+		break;
+
+	default:
+		UE_LOG(LogMegaPlants, Warning, TEXT("ExtractDisplacementData: Unsupported source format %d on texture %s."),
+			(int32)Source.GetFormat(), *Texture->GetName());
+		break;
+	}
+
+
+	Source.UnlockMip(0);
+#else
+	UE_LOG(LogMegaPlants, Warning, TEXT("ExtractDisplacementData only works in editor builds (WITH_EDITOR)."));
+#endif
+}
diff --git a/Engine/Plugins/Experimental/ProceduralVegetationEditor/Source/MegaPlants/Private/Implementations/MegaPlantMesher.h b/Engine/Plugins/Experimental/ProceduralVegetationEditor/Source/MegaPlants/Private/Implementations/MegaPlantMesher.h
new file mode 100644
--- /dev/null
+++ b/Engine/Plugins/Experimental/ProceduralVegetationEditor/Source/MegaPlants/Private/Implementations/MegaPlantMesher.h
@@ -0,0 +1,331 @@
+// Copyright Epic Games, Inc. All Rights Reserved.
+
+#pragma once
+#include "CoreMinimal.h"
+#include "UDynamicMesh.h"
+#include "Facades/MegaPlantBranchFacade.h"
+#include "Facades/MegaPlantPointFacade.h"
+#include "Generators/MeshShapeGenerator.h"
+#include "GeometryCollection/GeometryCollection.h"
+#include "GeometryCollection/ManagedArrayCollection.h"
+#include "Materials/MaterialInterface.h"
+#include "Materials/MaterialInstance.h"
+#include "MegaPlantMaterialSettings.h"
+#include "Curves/CurveFloat.h"
+#include "Engine/Texture2D.h"
+
+#include "MegaPlantMesher.generated.h"
+
+USTRUCT()
+struct FMegaPlantMesherSettings
+{
+	GENERATED_BODY()
+
+	UPROPERTY(EditAnywhere, Category="Materials", meta=(PinHiddenByDefault, InlineEditConditionToggle))
+	bool bOverrideMaterial = false;
+
+	UPROPERTY(EditAnywhere, Category="Materials", meta = (ShowOnlyInnerProperties))
+	FMegaPlantMaterialSettings MaterialSettings;
+
+	UPROPERTY(EditAnywhere, Category="Mesh Settings", DisplayName="Point Removal",
+		meta=(ClampMin=0.0f, ClampMax=0.1f, UIMin=0.0f, UIMax=0.1f))
+	float PointRemoval = 0.0f;
+
+	UPROPERTY(EditAnywhere, Category="Mesh Settings", DisplayName="Segment Reduction",
+		meta=(ClampMin=0.0f, ClampMax=0.1f, UIMin=0.0f, UIMax=0.1f))
+	float SegmentReduction = 0.0f;
+
+	UPROPERTY(EditAnywhere, Category="Mesh Settings", DisplayName="Segment Retention Impact",
+		meta=(ClampMin=0.0f, ClampMax=2.0f, UIMin=0.0f, UIMax=2.0f))
+	float SegmentRetentionImpact = 1.0f;
+
+	UPROPERTY(EditAnywhere, Category="Mesh Settings", DisplayName="Longest Segment Allowed",
+		meta=(ClampMin=0.0f, ClampMax=1.0f, UIMin=0.0f, UIMax=1.0f))
+	float LongestSegmentLength = 0.1f;
+
+	UPROPERTY(EditAnywhere, Category="Mesh Settings", DisplayName="Shortest Segment Allowed",
+		meta=(ClampMin=0.0f, ClampMax=1.0f, UIMin=0.0f, UIMax=1.0f))
+	float ShortestSegmentLength = 0.0f;
+
+	UPROPERTY(EditAnywhere, Category="Mesh Settings", DisplayName="Reduction Accuracy",
+		meta=(ClampMin=3, ClampMax=15, UIMin=1, UIMax=15))
+	int32 Accuracy = 5;
+
+	UPROPERTY(EditAnywhere, Category="Mesh Settings", DisplayName="Min Divisions",
+		meta=(ClampMin=3, ClampMax=1024, UIMin=3, UIMax=36))
+	int32 MinMeshDivisions = 6;
+
+	UPROPERTY(EditAnywhere, Category="Mesh Settings", DisplayName="Max Divisions",
+		meta=(ClampMin=3, ClampMax=1024, UIMin=3, UIMax=36))
+	int32 MaxMeshDivisions = 12;
+
+	UPROPERTY(EditAnywhere, Category="Point Removal Settings", DisplayName="Hull Retention",
+		meta=(ClampMin=0.0f, ClampMax=1.0f, UIMin=0.0f, UIMax=1.0f))
+	float PointHullRetention = 0.0f;
+
+	UPROPERTY(EditAnywhere, Category="Point Removal Settings", DisplayName="Main Trunk Retention",
+		meta=(ClampMin=0.0f, ClampMax=1.0f, UIMin=0.0f, UIMax=1.0f))
+	float PointMainTrunkRetention = 0.0f;
+
+	UPROPERTY(EditAnywhere, Category="Point Removal Settings", DisplayName="Ground Retention",
+		meta=(ClampMin=0.0f, ClampMax=1.0f, UIMin=0.0f, UIMax=1.0f))
+	float PointGroundRetention = 0.0f;
+
+	UPROPERTY(EditAnywhere, Category="Point Removal Settings", DisplayName="Hull Retention Gradient", meta=(ClampMin=0.0f, ClampMax=1.0f))
+	FRuntimeFloatCurve PointHullRetentionGradient;
+
+	UPROPERTY(EditAnywhere, Category="Point Removal Settings", DisplayName="Main Trunk Retention Gradient", meta=(ClampMin=0.0f, ClampMax=1.0f))
+	FRuntimeFloatCurve PointMainTrunkRetentionGradient;
+
+	UPROPERTY(EditAnywhere, Category="Point Removal Settings", DisplayName="Ground Retention Gradient", meta=(ClampMin=0.0f, ClampMax=1.0f))
+	FRuntimeFloatCurve PointGroundRetentionGradient;
+
+	UPROPERTY(EditAnywhere, Category="Segment Removal Settings", DisplayName="Hull Retention",
+		meta=(ClampMin=0.0f, ClampMax=1.0f, UIMin=0.0f, UIMax=1.0f))
+	float SegmentHullRetention = 0.0f;
+
+	UPROPERTY(EditAnywhere, Category="Segment Removal Settings", DisplayName="Main Trunk Retention",
+		meta=(ClampMin=0.0f, ClampMax=1.0f, UIMin=0.0f, UIMax=1.0f))
+	float SegmentMainTrunkRetention = 0.0f;
+
+	UPROPERTY(EditAnywhere, Category="Segment Removal Settings", DisplayName="Ground Retention",
+		meta=(ClampMin=0.0f, ClampMax=1.0f, UIMin=0.0f, UIMax=1.0f))
+	float SegmentGroundRetention = 0.0f;
+
+	UPROPERTY(EditAnywhere, Category="Segment Removal Settings", DisplayName="Scale Retention",
+		meta=(ClampMin=0.0f, ClampMax=1.0f, UIMin=0.0f, UIMax=1.0f))
+	float SegmentScaleRetention = 0.0f;
+
+	UPROPERTY(EditAnywhere, Category="Segment Removal Settings", DisplayName="Hull Retention Gradient", meta=(ClampMin=0.0f, ClampMax=1.0f))
+	FRuntimeFloatCurve SegmentHullRetentionGradient;
+
+	UPROPERTY(EditAnywhere, Category="Segment Removal Settings", DisplayName="Main Trunk Retention Gradient", meta=(ClampMin=0.0f, ClampMax=1.0f))
+	FRuntimeFloatCurve SegmentMainTrunkRetentionGradient;
+
+	UPROPERTY(EditAnywhere, Category="Segment Removal Settings", DisplayName="Ground Retention Gradient", meta=(ClampMin=0.0f, ClampMax=1.0f))
+	FRuntimeFloatCurve SegmentGroundRetentionGradient;
+
+	UPROPERTY(EditAnywhere, Category="Segment Removal Settings", DisplayName="Scale Retention Gradient", meta=(ClampMin=0.0f, ClampMax=1.0f))
+	FRuntimeFloatCurve SegmentScaleRetentionGradient;
+
+	UPROPERTY(EditAnywhere, Category="Division Removal Settings", DisplayName="Hull Retention",
+		meta=(ClampMin=0.0f, ClampMax=1.0f, UIMin=0.0f, UIMax=1.0f))
+	float DivisionHullRetention = 0.0f;
+
+	UPROPERTY(EditAnywhere, Category="Division Removal Settings", DisplayName="Main Trunk Retention",
+		meta=(ClampMin=0.0f, ClampMax=1.0f, UIMin=0.0f, UIMax=1.0f))
+	float DivisionMainTrunkRetention = 0.0f;
+
+	UPROPERTY(EditAnywhere, Category="Division Removal Settings", DisplayName="Ground Retention",
+		meta=(ClampMin=0.0f, ClampMax=1.0f, UIMin=0.0f, UIMax=1.0f))
+	float DivisionGroundRetention = 0.0f;
+
+	UPROPERTY(EditAnywhere, Category="Division Removal Settings", DisplayName="Hull Retention Gradient", meta=(ClampMin=0.0f, ClampMax=1.0f))
+	FRuntimeFloatCurve DivisionHullRetentionGradient;
+
+	UPROPERTY(EditAnywhere, Category="Division Removal Settings", DisplayName="Main Trunk Retention Gradient", meta=(ClampMin=0.0f, ClampMax=1.0f))
+	FRuntimeFloatCurve DivisionMainTrunkRetentionGradient;
+
+	UPROPERTY(EditAnywhere, Category="Division Removal Settings", DisplayName="Ground Retention Gradient", meta=(ClampMin=0.0f, ClampMax=1.0f))
+	FRuntimeFloatCurve DivisionGroundRetentionGradient;
+
+	UPROPERTY(EditAnywhere, Category="Division Removal Settings", DisplayName="Scale Retention Gradient", meta=(ClampMin=0.0f, ClampMax=1.0f))
+	FRuntimeFloatCurve DivisionScaleRetentionGradient;
+
+	UPROPERTY(EditAnywhere, Category="Plant Profile Settings", DisplayName="Profile",
+		meta=(EditCondition="bIsPlantProfileDropdownEnabled", EditConditionHides, HideEditConditionToggle, GetOptions="GetPlantProfileOptions"))
+	FString SelectedPlantProfile = "None";
+
+	UPROPERTY(Transient)
+	TArray<FString> PlantProfileOptions = {"None"};
+
+	UPROPERTY(Transient)
+	bool bIsPlantProfileDropdownEnabled = false;
+
+	UPROPERTY(EditAnywhere, Category="Plant Profile Settings", DisplayName="Profile Falloff",
+		meta=(ClampMin=0.0f, ClampMax=1.0f, EditCondition="bIsPlantProfileDropdownEnabled", EditConditionHides, HideEditConditionToggle))
+	FRuntimeFloatCurve PlantProfileFallOff;
+
+	UPROPERTY(EditAnywhere, Category="Plant Profile Settings", DisplayName="Profile Scale",
+		meta=(ClampMin=0.0f, ClampMax=3.0f, UIMin=0.0f, UIMax=3.0f, EditCondition="bIsPlantProfileDropdownEnabled", EditConditionHides,
+			HideEditConditionToggle))
+	float PlantProfileScale = 1.0f;
+
+	UPROPERTY(EditAnywhere, Category="Plant Profile Settings",
+		meta=(EditCondition="bIsPlantProfileDropdownEnabled", EditConditionHides, HideEditConditionToggle))
+	bool bApplyProfileToBranches = false;
+
+	UPROPERTY(EditAnywhere, Category = "Displacement Settings")
+	TObjectPtr<UTexture2D> DisplacementTexture = nullptr;
+
+	UPROPERTY(Transient, NonTransactional)
+	TArray<float> DisplacementValues;
+
+	UPROPERTY(EditAnywhere, Category="Displacement Settings", DisplayName="Displacement Scale",
+		meta=(ClampMin=0.0f, ClampMax=10.0f, UIMin=0.0f, UIMax=1.0f))
+	float DisplacementStrength = 0.5f;
+
+	UPROPERTY(EditAnywhere, Category="Displacement Settings", DisplayName="Displacement Bias",
+		meta=(ClampMin=0.0f, ClampMax=10.0f, UIMin=0.0f, UIMax=10.0f))
+	float DisplacementBias = 0.0f;
+
+	UPROPERTY(EditAnywhere, Category="Displacement Settings", DisplayName="Displacement UV Scale")
+	FIntVector2 DisplacementUVScale = FIntVector2(1.0f, 1.0f);
+
+	UPROPERTY(EditAnywhere, Category="Displacement Settings", DisplayName="Apply Displacement up to generation",
+		meta=(ClampMin=1, ClampMax=10, UIMin=1, UIMax=10))
+	int32 DisplacementGenerationUpperLimit = 1;
+
+	FMegaPlantMesherSettings()
+	{
+		InitializeLinearCurve(PointHullRetentionGradient);
+		InitializeLinearCurve(PointMainTrunkRetentionGradient);
+		InitializeLinearCurve(PointGroundRetentionGradient);
+		InitializeLinearCurve(SegmentHullRetentionGradient);
+		InitializeLinearCurve(SegmentMainTrunkRetentionGradient);
+		InitializeLinearCurve(SegmentGroundRetentionGradient);
+		InitializeLinearCurve(SegmentScaleRetentionGradient);
+		InitializeLinearCurve(DivisionHullRetentionGradient);
+		InitializeLinearCurve(DivisionMainTrunkRetentionGradient);
+		InitializeLinearCurve(DivisionGroundRetentionGradient);
+		InitializeLinearCurve(DivisionScaleRetentionGradient);
+		InitializeLinearCurve(PlantProfileFallOff);
+	}
+
+private:
+	static void InitializeLinearCurve(FRuntimeFloatCurve& Curve)
+	{
+		if (FRichCurve* RichCurve = Curve.GetRichCurve())
+		{
+			RichCurve->Reset();
+			const FKeyHandle Handle0 = RichCurve->AddKey(0.0f, 0.0f);
+			const FKeyHandle Handle1 = RichCurve->AddKey(1.0f, 1.0f);
+
+			RichCurve->SetKeyInterpMode(Handle0, RCIM_Linear);
+			RichCurve->SetKeyInterpMode(Handle1, RCIM_Linear);
+		}
+	}
+};
+
+struct FLocalDynamicMeshData
+{
+	struct FVertex
+	{
+		FVector3f Position;
+		FVector3f Normal;
+		FVector2f UV;
+		float NjordPixelId;
+	};
+
+	TArray<FVertex> Vertices;
+	TArray<FIntVector4> Triangles;
+	int MaterialID = 0;
+};
+
+struct FDisplacementData
+{
+	const TArrayView<const float> Values;
+	const int32 TextureWidth;
+	const int32 TextureHeight;
+
+	FDisplacementData(const TArrayView<const float>& InValues, const int32 InWidth, const int32 InHeight)
+		: Values(InValues)
+		, TextureWidth(InWidth)
+		, TextureHeight(InHeight)
+	{
+	}
+};
+
+class FMegaPlantMeshGenerator : public UE::Geometry::FMeshShapeGenerator
+{
+	int32 VerticesCount = 0;
+	int32 TriangleCount = 0;
+
+	TArray<FLocalDynamicMeshData>* MeshesDatas = nullptr;
+
+public:
+	FMegaPlantMeshGenerator(const int32 InVerticesCount, const int32 InTriangleCount, TArray<FLocalDynamicMeshData>* InMeshesDatas)
+		: VerticesCount(InVerticesCount)
+		, TriangleCount(InTriangleCount)
+		, MeshesDatas(InMeshesDatas)
+	{
+	}
+
+	virtual FMeshShapeGenerator& Generate() override
+	{
+		SetBufferSizes(VerticesCount, TriangleCount, VerticesCount, TriangleCount);
+
+		int32 VertexIndex = 0;
+		int32 FaceIndex = 0;
+		for (auto MeshIndex = MeshesDatas->Num() - 1; MeshIndex >= 0; --MeshIndex)
+		{
+			const FLocalDynamicMeshData& MeshesData = (*MeshesDatas)[MeshIndex];
+			for (const FIntVector4& Triangle : MeshesData.Triangles)
+			{
+				const FIntVector4 OffsetTriangle = Triangle + FIntVector4(VertexIndex);
+				const UE::Geometry::FIndex3i Indices = FIntVector(OffsetTriangle);
+				Triangles[FaceIndex] = Indices;
+				TriangleUVs[FaceIndex] = Indices;
+				TriangleNormals[FaceIndex] = Indices;
+				TrianglePolygonIDs[FaceIndex] = OffsetTriangle.W;
+				++FaceIndex;
+			}
+			for (const FLocalDynamicMeshData::FVertex& Vertex : MeshesData.Vertices)
+			{
+				Vertices[VertexIndex] = static_cast<FVector>(Vertex.Position);
+				Normals[VertexIndex] = Vertex.Normal;
+				NormalParentVertex[VertexIndex] = VertexIndex;
+				UVs[VertexIndex] = Vertex.UV;
+				UVParentVertex[VertexIndex] = VertexIndex;
+				++VertexIndex;
+			}
+		}
+
+		return *this;
+	}
+};
+
+struct FMegaPlantMesher
+{
+	static void GenerateGeometryCollection(const FManagedArrayCollection& InSkeletonCollection, const FMegaPlantMesherSettings& MesherSettings,
+	                                       FGeometryCollection& OutGeometryCollection);
+
+	static void SetNjordPixelID(FGeometryCollection& OutGeometryCollection);
+
+	static void GenerateDynamicMesh(FManagedArrayCollection& Collection, const FMegaPlantMesherSettings& MesherSettings,
+	                                TObjectPtr<UDynamicMesh>& OutMesh);
+
+	static void ExtractDisplacementData(const TObjectPtr<UTexture2D>& Texture, TArray<float>& OutValues);
+
+private:
+	static TSet<int32> CollectHardPoints(const MegaPlants::Facades::FBranchFacade& BranchFacade,
+	                                     const MegaPlants::Facades::FPointFacade& PointFacade);
+
+	static TSet<int32> ComputePointGradients(const MegaPlants::Facades::FPointFacade& PointFacade, const FMegaPlantMesherSettings& MesherSettings,
+	                                         const TSet<int32>& HardPoints, const float MaxPointScale, TArray<float>& OutMeshDivisionsGradients,
+	                                         TArray<float>& OutDeltaModifiers);
+
+	static void PerformPathSimplification(const MegaPlants::Facades::FBranchFacade& BranchFacade,
+	                                      const MegaPlants::Facades::FPointFacade& PointFacade, const FMegaPlantMesherSettings& MesherSettings,
+	                                      const TSet<int32>& HardPoints, const TArray<float>& DeltaModifiers, TSet<int32>& InOutPointsToRemove);
+
+	static TArray<int32> ComputeMeshDivisions(const MegaPlants::Facades::FBranchFacade& BranchFacade, const FMegaPlantMesherSettings& MesherSettings,
+	                                          const TArray<float>& MeshDivisionsGradients, const int32 PointCount);
+
+	static void TriangulateRings(const TArray<int32>& PreviousIndices, const TArray<int32>& CurrentIndices, int32& InOutPolyGroupIndex,
+	                             FLocalDynamicMeshData& OutMeshData);
+
+	static float GetProfileMultiplier(const TArray<float>& InProfilePoints, const float ProfileUV_U);
+
+	static TMap<int32, TArray<int32>> GetPointsIndicesToFoliageIndicesMap(const FManagedArrayCollection& Collection);
+
+	static void GenerateBranchMeshData(const bool bPrimitiveIsTrunk, const int32 GenerationNumber, const TArray<int32>& PrimitivePoints,
+	                                   const FDisplacementData& DisplacementData, const FMegaPlantMesherSettings& MesherSettings,
+	                                   const TArray<int32>& TargetMeshDivisions, const FManagedArrayCollection& Collection,
+	                                   FLocalDynamicMeshData& OutLocalMeshData);
+
+	static void UpdateFoliagePivotPoints(const TSet<int32>& PointsToRemove, FManagedArrayCollection& OutCollection);
+
+	static void GetLocalDynamicMeshData(FManagedArrayCollection& Collection, const FMegaPlantMesherSettings& MesherSettings,
+	                                    TArray<FLocalDynamicMeshData>& MeshesData);
+};
diff --git a/Engine/Plugins/Experimental/ProceduralVegetationEditor/Source/MegaPlants/Private/Implementations/MegaPlantRemoveBranches.cpp b/Engine/Plugins/Experimental/ProceduralVegetationEditor/Source/MegaPlants/Private/Implementations/MegaPlantRemoveBranches.cpp
new file mode 100644
--- /dev/null
+++ b/Engine/Plugins/Experimental/ProceduralVegetationEditor/Source/MegaPlants/Private/Implementations/MegaPlantRemoveBranches.cpp
@@ -0,0 +1,126 @@
+// Copyright Epic Games, Inc. All Rights Reserved.
+
+#include "MegaPlantRemoveBranches.h"
+#include "Facades/MegaPlantBranchFacade.h"
+#include "Facades/MegaPlantBudVectorsFacade.h"
+#include "Facades/MegaPlantFoliageFacade.h"
+#include "Facades/MegaPlantPointFacade.h"
+#include "Facades/MegaPlantTreeFacade.h"
+
+namespace MegaPlants::Facades
+{
+	static constexpr int BranchBudLightDetectedIndex = 2;
+	static constexpr int BudDevelopementGenerationIndex = 0;
+	static constexpr int BudDevelopementAgeIndex = 2;
+	
+	void GatherBranchRemovalData(const FBranchFacade& InBranchFacade, const FPointFacade& InPointFacade,const ERemoveBranchesBasis InBasis, TArray<float>& OutRemovalData)
+	{
+		const int32 NumBranches = InBranchFacade.GetElementCount();
+		const int32 RootBranchIndex = InBranchFacade.GetTrunkIndex();
+
+		check(RootBranchIndex !=  INDEX_NONE);
+		
+		if (RootBranchIndex ==  INDEX_NONE)
+			return;
+		
+		const TArray<int32>& RootBranchPoints = InBranchFacade.GetPoints(RootBranchIndex);
+
+		check(RootBranchPoints.Num() > 0);
+		
+		if (RootBranchPoints.Num() == 0)
+			return;
+
+		OutRemovalData.AddUninitialized(NumBranches);
+
+		const TArray<int>& RootBudDevelopment = InPointFacade.GetBudDevelopment(RootBranchPoints[0]);
+
+		float MaxValue = std::numeric_limits<float>::lowest();
+		
+		for (int32 BranchIndex = 0; BranchIndex < NumBranches; ++BranchIndex)
+		{
+			const TArray<int32>& BranchPoints = InBranchFacade.GetPoints(BranchIndex);
+
+			if (InBasis == ERemoveBranchesBasis::Length)
+			{
+				const float BranchRootPointLFR = InPointFacade.GetLengthFromRoot(BranchPoints[0]);
+				const float BranchLastPointLFR = InPointFacade.GetLengthFromRoot(BranchPoints[BranchPoints.Num() - 1]);
+
+				OutRemovalData[BranchIndex] = BranchLastPointLFR - BranchRootPointLFR;
+			}
+			else if (InBasis == ERemoveBranchesBasis::Radius)
+			{
+				OutRemovalData[BranchIndex] = InPointFacade.GetPointScale(BranchPoints[0]);
+			}
+			else if (InBasis == ERemoveBranchesBasis::Light)
+			{
+				const TArray<float>& PointBudLightDetected = InPointFacade.GetBudLightDetected(BranchPoints[BranchPoints.Num() - 2]);
+
+				// TODO: Some entries dont have complete data. Need to update the JSON to fix it. For now, setting -1 value to keep track of inaccurate data. 
+				OutRemovalData[BranchIndex] = PointBudLightDetected.Num() > BranchBudLightDetectedIndex ? PointBudLightDetected[BranchBudLightDetectedIndex] : -1.0f;
+			}
+			else if (InBasis == ERemoveBranchesBasis::Age ||  InBasis == ERemoveBranchesBasis::Generation)
+			{
+				const TArray<int>& PointBudDevelopment = InPointFacade.GetBudDevelopment(BranchPoints[BranchPoints.Num() - 1]);
+				
+				if (InBasis == ERemoveBranchesBasis::Age)
+				{
+					OutRemovalData[BranchIndex] = (float)PointBudDevelopment[BudDevelopementAgeIndex]/RootBudDevelopment[BudDevelopementAgeIndex];
+				}
+				else if (InBasis == ERemoveBranchesBasis::Generation)
+				{
+					OutRemovalData[BranchIndex] = (float)(RootBudDevelopment[BudDevelopementGenerationIndex])/PointBudDevelopment[BudDevelopementGenerationIndex];
+				}
+			}
+
+			MaxValue = FMath::Max(OutRemovalData[BranchIndex], MaxValue);
+		}
+
+		if (MaxValue > 1)
+		{
+			for (int i = 0; i < OutRemovalData.Num(); i++)
+			{
+				OutRemovalData[i] = OutRemovalData[i] / MaxValue;
+			}
+		}
+	}
+
+	void FindBrancehesToRemove(const FBranchFacade& InBranchFacade, const TArray<float>& RemovalData, const float& InThreshold, TArray<int>& OutBranchesToRemove)
+	{
+		const int32 NumBranches = InBranchFacade.GetElementCount();
+		
+		for (int32 BranchIndex = 0; BranchIndex < NumBranches; ++BranchIndex)
+		{
+			float Val = RemovalData[BranchIndex];
+			
+			if (Val != -1.0 && Val < InThreshold)
+			{
+				if (!OutBranchesToRemove.Contains(BranchIndex))
+				{
+					OutBranchesToRemove.Add(BranchIndex);
+				}
+				
+				UE_LOG(LogTemp, Log, TEXT("MegaPlantRemoveBranches :: RemoveBranches : Removing branch index : %i"), BranchIndex);
+			}
+		}
+	}
+}
+
+void FMegaPlantRemoveBranches::ApplyRemoveBranches(const ERemoveBranchesBasis InBasis, const float InThreshold, FManagedArrayCollection& OutCollection)
+{
+	MegaPlants::Facades::FBranchFacade BranchFacade = MegaPlants::Facades::FBranchFacade(OutCollection);
+	MegaPlants::Facades::FPointFacade PointFacade = MegaPlants::Facades::FPointFacade(OutCollection);
+	MegaPlants::Facades::FFoliageFacade FoliageFacade = MegaPlants::Facades::FFoliageFacade(OutCollection);
+		
+	if(!PointFacade.IsValid() || !BranchFacade.IsValid())
+	{
+		return;
+	}
+	
+	TArray<float> OutRemovalData;
+	GatherBranchRemovalData(BranchFacade, PointFacade, InBasis,OutRemovalData);
+
+	TArray<int> OutBranchesToRemove;
+	FindBrancehesToRemove(BranchFacade, OutRemovalData, InThreshold, OutBranchesToRemove);
+
+	MegaPlants::Facades::FTreeFacade::RemoveBranches(BranchFacade, OutBranchesToRemove, OutCollection);
+}
diff --git a/Engine/Plugins/Experimental/ProceduralVegetationEditor/Source/MegaPlants/Private/Implementations/MegaPlantRemoveBranches.h b/Engine/Plugins/Experimental/ProceduralVegetationEditor/Source/MegaPlants/Private/Implementations/MegaPlantRemoveBranches.h
new file mode 100644
--- /dev/null
+++ b/Engine/Plugins/Experimental/ProceduralVegetationEditor/Source/MegaPlants/Private/Implementations/MegaPlantRemoveBranches.h
@@ -0,0 +1,23 @@
+// Copyright Epic Games, Inc. All Rights Reserved.
+
+#pragma once
+
+#include "CoreMinimal.h"
+#include "MegaPlantRemoveBranches.generated.h"
+
+struct FManagedArrayCollection;
+
+UENUM()
+enum class ERemoveBranchesBasis : uint8
+{
+	Length,
+	Radius,
+	Light,
+	Age,	
+	Generation
+};
+
+struct FMegaPlantRemoveBranches
+{
+	static void ApplyRemoveBranches(const ERemoveBranchesBasis InBasis, const float Threshold, FManagedArrayCollection& OutCollection);
+};
diff --git a/Engine/Plugins/Experimental/ProceduralVegetationEditor/Source/MegaPlants/Private/Implementations/MegaPlantScale.cpp b/Engine/Plugins/Experimental/ProceduralVegetationEditor/Source/MegaPlants/Private/Implementations/MegaPlantScale.cpp
new file mode 100644
--- /dev/null
+++ b/Engine/Plugins/Experimental/ProceduralVegetationEditor/Source/MegaPlants/Private/Implementations/MegaPlantScale.cpp
@@ -0,0 +1,79 @@
+// Copyright Epic Games, Inc. All Rights Reserved.
+
+#include "MegaPlantScale.h"
+
+#include "Facades/MegaPlantBranchFacade.h"
+#include "Facades/MegaPlantBudVectorsFacade.h"
+#include "Facades/MegaPlantFoliageFacade.h"
+#include "Facades/MegaPlantPointFacade.h"
+
+void FMegaPlantScale::ApplyScale(const float InManualScale, FManagedArrayCollection& OutCollection)
+{
+	MegaPlants::Facades::FBranchFacade BranchFacade = MegaPlants::Facades::FBranchFacade(OutCollection);
+	MegaPlants::Facades::FPointFacade PointFacade = MegaPlants::Facades::FPointFacade(OutCollection);
+		
+	if(!PointFacade.IsValid() || !BranchFacade.IsValid())
+	{
+		return;
+	}
+
+	const int32 RootBranchIndex = BranchFacade.GetTrunkIndex();
+	const int32 NumBranches = BranchFacade.GetElementCount();
+
+	check(RootBranchIndex !=  INDEX_NONE);
+	
+	if (RootBranchIndex ==  INDEX_NONE)
+		return;
+	
+	const TArray<int32>& RootBranchPoints = BranchFacade.GetPoints(RootBranchIndex);
+
+	check(RootBranchPoints.Num() > 0);
+	
+	if (RootBranchPoints.Num() == 0)
+		return;
+	
+	const int32 RootPointIndex = RootBranchPoints[0];
+	const FVector3f& PivotPos = PointFacade.GetPosition(RootPointIndex);
+	const float RootPointLFR = PointFacade.GetLengthFromRoot(RootPointIndex);
+	
+	for (int32 BranchIndex = 0; BranchIndex < NumBranches; ++BranchIndex)
+	{
+		const TArray<int32>& BranchPoints = BranchFacade.GetPoints(BranchIndex);
+
+		for (const int32& PointIndex : BranchPoints)
+		{
+			FVector3f PointPos = PointFacade.GetPosition(PointIndex);
+			float LFS = PointFacade.GetLengthFromSeed(PointIndex);
+
+			PointPos -= PivotPos;
+			PointPos *= InManualScale;
+			PointPos += PivotPos;
+
+			LFS *= InManualScale;
+
+			float LFR = LFS + RootPointLFR;
+
+			PointFacade.ModifyPositions()[PointIndex] = PointPos;
+			PointFacade.ModifyLengthFromSeeds()[PointIndex] = LFS;
+			PointFacade.ModifyLengthFromRoots()[PointIndex] = LFR;
+			PointFacade.ModifyPointScales()[PointIndex] *= InManualScale;
+		}
+	}
+
+	MegaPlants::Facades::FFoliageFacade FoliageFacade(OutCollection);
+	int32 NumInstances = FoliageFacade.NumFoliageEntries();
+
+	for(int32 Id = 0; Id < NumInstances; Id++)
+	{
+		MegaPlants::Facades::FFoliageEntryData Data = FoliageFacade.GetFoliageEntry(Id);
+
+		Data.PivotPoint -= PivotPos;
+		Data.PivotPoint *= InManualScale;
+		Data.PivotPoint += PivotPos;
+		
+		Data.Scale *= InManualScale;
+		Data.LengthFromRoot = Data.LengthFromRoot * InManualScale;
+
+		FoliageFacade.SetFoliageEntry(Id, Data);
+	}
+}
diff --git a/Engine/Plugins/Experimental/ProceduralVegetationEditor/Source/MegaPlants/Private/Implementations/MegaPlantScale.h b/Engine/Plugins/Experimental/ProceduralVegetationEditor/Source/MegaPlants/Private/Implementations/MegaPlantScale.h
new file mode 100644
--- /dev/null
+++ b/Engine/Plugins/Experimental/ProceduralVegetationEditor/Source/MegaPlants/Private/Implementations/MegaPlantScale.h
@@ -0,0 +1,12 @@
+// Copyright Epic Games, Inc. All Rights Reserved.
+
+#pragma once
+
+#include "CoreMinimal.h"
+#include "GeometryCollection/ManagedArrayCollection.h"
+
+struct FMegaPlantScale
+{
+public:
+	static void ApplyScale(const float InManualScale, FManagedArrayCollection& OutCollection);
+};
diff --git a/Engine/Plugins/Experimental/ProceduralVegetationEditor/Source/MegaPlants/Private/MegaPlant.cpp b/Engine/Plugins/Experimental/ProceduralVegetationEditor/Source/MegaPlants/Private/MegaPlant.cpp
new file mode 100644
--- /dev/null
+++ b/Engine/Plugins/Experimental/ProceduralVegetationEditor/Source/MegaPlants/Private/MegaPlant.cpp
@@ -0,0 +1,82 @@
+// Copyright Epic Games, Inc. All Rights Reserved.
+
+#include "MegaPlant.h"
+#include "Nodes/PCGMegaPlantFoliage.h"
+#include "Nodes/PCGMegaPlantGravity.h"
+#include "Nodes/PCGMegaPlantMesher.h"
+#include "Nodes/PCGMegaPlantRecipeLoad.h"
+#include "Nodes/PCGMegaPlantOutput.h"
+
+void SanitizeFolderName(FString& FolderName)
+{
+	const TCHAR* InvalidChar = INVALID_OBJECTPATH_CHARACTERS INVALID_LONGPACKAGE_CHARACTERS;
+
+	while (*InvalidChar)
+	{
+		FolderName.ReplaceCharInline(*InvalidChar, TCHAR('_'), ESearchCase::CaseSensitive);
+		++InvalidChar;
+	}
+
+	FolderName.ReplaceCharInline('/', TCHAR('_'), ESearchCase::CaseSensitive);
+}
+
+void UMegaPlant::CreateGraph(const UMegaPlantGraph* InGraph /*= nullptr*/)
+{
+	if (InGraph)
+	{
+		Graph = DuplicateObject(InGraph, this);
+	}
+	else
+	{
+		Graph = NewObject<UMegaPlantGraph>(this, UMegaPlantGraph::StaticClass());
+	}
+}
+
+#if WITH_EDITOR
+void UMegaPlant::CreateGraphFromRecipe(const TObjectPtr<UMegaPlantRecipeDataAsset> Recipe)
+{
+	Graph = NewObject<UMegaPlantGraph>(this, UMegaPlantGraph::StaticClass());
+
+	UPCGSettings* DefaultNodeSettings = nullptr;
+
+	UPCGMegaPlantRecipeLoadSettings* RecipeLoadNodeSettings = NewObject<UPCGMegaPlantRecipeLoadSettings>(Graph);
+	RecipeLoadNodeSettings->RecipeDataAsset = Recipe;
+	UPCGNode* RecipeLoadNode = Graph->AddNodeCopy(RecipeLoadNodeSettings, DefaultNodeSettings);
+
+	FInt32Vector2 Position(0, 300);
+
+	RecipeLoadNode->SetNodePosition(Position.X, Position.Y);
+	Position.X += 100;
+
+	for (const auto& [VariantName, VariantData] : Recipe->Variants)
+	{
+		const UPCGMegaPlantGravitySettings* GravityNodeSettings = NewObject<UPCGMegaPlantGravitySettings>(Graph);
+		const UPCGMegaPlantMesherSettings* MesherNodeSettings = NewObject<UPCGMegaPlantMesherSettings>(Graph);
+		const UPCGMegaPlantFoliageSettings* FoliageNodeSettings = NewObject<UPCGMegaPlantFoliageSettings>(Graph);
+		UPCGMegaPlantOutputSettings* OutputNodeSettings = NewObject<UPCGMegaPlantOutputSettings>(Graph);
+		FString MeshName = *VariantName;
+		SanitizeFolderName(MeshName);
+		OutputNodeSettings->ExportSettings.MeshName = *MeshName;
+
+		UPCGNode* GravityNode = Graph->AddNodeCopy(GravityNodeSettings, DefaultNodeSettings);
+		UPCGNode* MesherNode = Graph->AddNodeCopy(MesherNodeSettings, DefaultNodeSettings);
+		UPCGNode* FoliageNode = Graph->AddNodeCopy(FoliageNodeSettings, DefaultNodeSettings);
+		UPCGNode* OutputNode = Graph->AddNodeCopy(OutputNodeSettings, DefaultNodeSettings);
+
+		OutputNode->NodeTitle = *VariantName;
+
+		RecipeLoadNode
+			->AddEdgeTo(*VariantName, GravityNode, PCGPinConstants::DefaultInputLabel)
+			->AddEdgeTo(PCGPinConstants::DefaultOutputLabel, MesherNode, PCGPinConstants::DefaultInputLabel)
+			->AddEdgeTo(PCGPinConstants::DefaultOutputLabel, FoliageNode, PCGPinConstants::DefaultInputLabel)
+			->AddEdgeTo(PCGPinConstants::DefaultOutputLabel, OutputNode, PCGPinConstants::DefaultInputLabel);
+
+		GravityNode->SetNodePosition(Position.X + 200, Position.Y);
+		MesherNode->SetNodePosition(Position.X + 400, Position.Y);
+		FoliageNode->SetNodePosition(Position.X + 600, Position.Y);
+		OutputNode->SetNodePosition(Position.X + 800, Position.Y);
+
+		Position.Y += 150;
+	}
+}
+#endif
diff --git a/Engine/Plugins/Experimental/ProceduralVegetationEditor/Source/MegaPlants/Private/MegaPlantRecipeDataAsset.cpp b/Engine/Plugins/Experimental/ProceduralVegetationEditor/Source/MegaPlants/Private/MegaPlantRecipeDataAsset.cpp
new file mode 100644
--- /dev/null
+++ b/Engine/Plugins/Experimental/ProceduralVegetationEditor/Source/MegaPlants/Private/MegaPlantRecipeDataAsset.cpp
@@ -0,0 +1,19 @@
+﻿// Copyright Epic Games, Inc. All Rights Reserved.
+
+#include "MegaPlantRecipeDataAsset.h"
+
+#include "Helpers/MegaPlantJSONHelper.h"
+
+class UPCGMegaPlantFoliageSettings;
+
+void UMegaPlantRecipeDataAsset::LoadFromVariantFiles(const TMap<FString, FString>& InVariantFiles)
+{
+	for (const auto& [VariantName, VariantFile] : InVariantFiles)
+	{
+		FManagedArrayCollection Collection;
+		if (FString OutError; MegaPlants::LoadMegaPlantJsonToCollection(Collection, VariantFile, OutError))
+		{
+			Variants.Add(VariantName, Collection);
+		}
+	}
+}
diff --git a/Engine/Plugins/Experimental/ProceduralVegetationEditor/Source/MegaPlants/Private/MegaPlantsModule.cpp b/Engine/Plugins/Experimental/ProceduralVegetationEditor/Source/MegaPlants/Private/MegaPlantsModule.cpp
new file mode 100644
--- /dev/null
+++ b/Engine/Plugins/Experimental/ProceduralVegetationEditor/Source/MegaPlants/Private/MegaPlantsModule.cpp
@@ -0,0 +1,14 @@
+// Copyright Epic Games, Inc. All Rights Reserved.
+#include "MegaPlantsModule.h"
+
+DEFINE_LOG_CATEGORY(LogMegaPlants);
+
+void FMegaPlantsModule::StartupModule()
+{
+}
+
+void FMegaPlantsModule::ShutdownModule()
+{
+}
+
+IMPLEMENT_MODULE(FMegaPlantsModule, MegaPlants)
\ No newline at end of file
diff --git a/Engine/Plugins/Experimental/ProceduralVegetationEditor/Source/MegaPlants/Private/Nodes/PCGMegaPlantBaseSettings.cpp b/Engine/Plugins/Experimental/ProceduralVegetationEditor/Source/MegaPlants/Private/Nodes/PCGMegaPlantBaseSettings.cpp
new file mode 100644
--- /dev/null
+++ b/Engine/Plugins/Experimental/ProceduralVegetationEditor/Source/MegaPlants/Private/Nodes/PCGMegaPlantBaseSettings.cpp
@@ -0,0 +1,61 @@
+﻿// Copyright Epic Games, Inc. All Rights Reserved.
+
+#include "Nodes/PCGMegaPlantBaseSettings.h"
+#include "PCGContext.h"
+#include "Helpers/PCGHelpers.h"
+#include "PCGMegaPlantData.h"
+#include "PCGData.h"
+
+TArray<FPCGPinProperties> UPCGMegaPlantBaseSettings::InputPinProperties() const
+{
+	TArray<FPCGPinProperties> Properties;
+
+	FPCGPinProperties& Pin = Properties.Emplace_GetRef(PCGPinConstants::DefaultInputLabel, EPCGDataType::Any);
+	Pin.SetRequiredPin();
+	Pin.SetAllowMultipleConnections(false);
+	Pin.bAllowMultipleData = false;
+
+	return Properties;
+}
+
+TArray<FPCGPinProperties> UPCGMegaPlantBaseSettings::OutputPinProperties() const
+{
+	TArray<FPCGPinProperties> Properties;
+
+	FPCGPinProperties& Pin = Properties.Emplace_GetRef(PCGPinConstants::DefaultOutputLabel, EPCGDataType::Any);
+	Pin.SetAllowMultipleConnections(true);
+	Pin.bAllowMultipleData = false;
+
+	return Properties;
+}
+
+FPCGElementPtr UPCGMegaPlantBaseSettings::CreateElement() const
+{
+	return MakeShared<FPCGMegaPlantBaseElement>();
+}
+
+void FPCGMegaPlantBaseElement::PostExecuteInternal(FPCGContext* InContext) const
+{
+	IPCGElement::PostExecuteInternal(InContext);
+
+#if WITH_EDITOR
+	const UPCGMegaPlantBaseSettings* Settings = InContext->GetInputSettings<UPCGMegaPlantBaseSettings>();
+	check(Settings);
+
+	for (FPCGTaggedData& OutputData : InContext->OutputData.TaggedData)
+	{
+		if (OutputData.Data && OutputData.Data.IsA<UPCGMegaPlantData>())
+		{
+			if (const UPCGMegaPlantData* MegaPlantData = Cast<UPCGMegaPlantData>(OutputData.Data))
+			{
+				MegaPlantData->SetRenderType(Settings->GetRenderType());
+			}
+		}
+	}
+#endif
+}
+
+bool FPCGMegaPlantBaseElement::ExecuteInternal(FPCGContext* Context) const
+{
+	return true;
+}
\ No newline at end of file
diff --git a/Engine/Plugins/Experimental/ProceduralVegetationEditor/Source/MegaPlants/Private/Nodes/PCGMegaPlantCarve.cpp b/Engine/Plugins/Experimental/ProceduralVegetationEditor/Source/MegaPlants/Private/Nodes/PCGMegaPlantCarve.cpp
new file mode 100644
--- /dev/null
+++ b/Engine/Plugins/Experimental/ProceduralVegetationEditor/Source/MegaPlants/Private/Nodes/PCGMegaPlantCarve.cpp
@@ -0,0 +1,48 @@
+﻿// Copyright Epic Games, Inc. All Rights Reserved.
+
+#include "PCGMegaPlantCarve.h"
+#include "PCGMegaPlantData.h"
+#include "Data/PCGBasePointData.h"
+#include "GeometryCollection/ManagedArrayCollection.h"
+#include "Helpers/PCGSettingsHelpers.h"
+
+FPCGElementPtr UPCGMegaPlantCarveSettings::CreateElement() const
+{
+	return MakeShared<FPCGMegaPlantCarveElement>();
+}
+
+bool FPCGMegaPlantCarveElement::ExecuteInternal(FPCGContext* InContext) const
+{
+	TRACE_CPUPROFILER_EVENT_SCOPE(FPCGMegaPlantCarveElement::ExecuteInternal);
+
+	check(InContext);
+
+	const UPCGMegaPlantCarveSettings* InputSettings = InContext->GetInputSettings<UPCGMegaPlantCarveSettings>();
+	check(InputSettings);
+
+	for (const FPCGTaggedData& Input : InContext->InputData.GetInputsByPin(PCGPinConstants::DefaultInputLabel))
+	{
+		if (const UPCGMegaPlantData* InputData = Cast<UPCGMegaPlantData>(Input.Data))
+		{
+			FManagedArrayCollection SourceCollection = InputData->GetCollection();
+
+			FManagedArrayCollection OutCollection;
+			SourceCollection.CopyTo(&OutCollection);
+
+			UPCGMegaPlantData* OutManagedArrayCollectionData = FPCGContext::NewObject_AnyThread<UPCGMegaPlantData>(InContext);
+
+			auto [CarveBasis, Carve] = InputSettings->CarveSettings;
+			FMegaPlantCarve::ApplyCarve(OutCollection, SourceCollection, CarveBasis, Carve);
+
+			OutManagedArrayCollectionData->Initialize(MoveTemp(OutCollection));
+			InContext->OutputData.TaggedData.Emplace(OutManagedArrayCollectionData);
+		}
+		else
+		{
+			PCGLog::InputOutput::LogInvalidInputDataError(InContext);
+			return true;
+		}
+	}
+
+	return true;
+}
diff --git a/Engine/Plugins/Experimental/ProceduralVegetationEditor/Source/MegaPlants/Private/Nodes/PCGMegaPlantCarve.h b/Engine/Plugins/Experimental/ProceduralVegetationEditor/Source/MegaPlants/Private/Nodes/PCGMegaPlantCarve.h
new file mode 100644
--- /dev/null
+++ b/Engine/Plugins/Experimental/ProceduralVegetationEditor/Source/MegaPlants/Private/Nodes/PCGMegaPlantCarve.h
@@ -0,0 +1,39 @@
+﻿// Copyright Epic Games, Inc. All Rights Reserved.
+#pragma once
+
+#include "CoreMinimal.h"
+#include "Nodes/PCGMegaPlantBaseSettings.h"
+#include "Implementations/MegaPlantCarve.h"
+#include "PCGMegaPlantData.h"
+#include "PCGMegaPlantCarve.generated.h"
+
+UCLASS(BlueprintType, ClassGroup = (Procedural))
+class UPCGMegaPlantCarveSettings : public UPCGMegaPlantBaseSettings
+{
+	GENERATED_BODY()
+
+	friend class FPCGMegaPlantCarveElement;
+
+public:
+#if WITH_EDITOR
+	virtual FName GetDefaultNodeName() const override { return FName(TEXT("MegaPlantCarve")); }
+	virtual FText GetDefaultNodeTitle() const override { return NSLOCTEXT("PCGMegaPlantCarve", "NodeTitle", "Carve"); }
+	virtual FText GetNodeTooltipText() const override { return NSLOCTEXT("PCGMegaPlantCarve", "NodeTooltip", "Carve branches or trunk"); }
+	virtual FLinearColor GetNodeTitleColor() const override { return FLinearColor::Yellow; }
+
+	virtual TArray<EMegaPlantRenderType> GetRenderType() const override { return TArray{EMegaPlantRenderType::Skeleton}; }
+#endif
+
+protected:
+	virtual FPCGElementPtr CreateElement() const override;
+
+private:
+	UPROPERTY(EditAnywhere, Category = "Carve", meta = (ShowOnlyInnerProperties))
+	FMegaPlantCarveSettings CarveSettings;
+};
+
+class FPCGMegaPlantCarveElement : public FPCGMegaPlantBaseElement
+{
+protected:
+	virtual bool ExecuteInternal(FPCGContext* InContext) const override;
+};
diff --git a/Engine/Plugins/Experimental/ProceduralVegetationEditor/Source/MegaPlants/Private/Nodes/PCGMegaPlantFoliage.cpp b/Engine/Plugins/Experimental/ProceduralVegetationEditor/Source/MegaPlants/Private/Nodes/PCGMegaPlantFoliage.cpp
new file mode 100644
--- /dev/null
+++ b/Engine/Plugins/Experimental/ProceduralVegetationEditor/Source/MegaPlants/Private/Nodes/PCGMegaPlantFoliage.cpp
@@ -0,0 +1,108 @@
+﻿// Copyright Epic Games, Inc. All Rights Reserved.
+
+#include "PCGMegaPlantFoliage.h"
+
+#include "MegaPlantsModule.h"
+#include "PCGMegaPlantData.h"
+#include "Data/PCGBasePointData.h"
+#include "Data/PCGStaticMeshResourceData.h"
+#include "Facades/MegaPlantFoliageFacade.h"
+#include "GeometryCollection/ManagedArrayCollection.h"
+#include "Helpers/PCGHelpers.h"
+#include "Helpers/PCGSettingsHelpers.h"
+
+#if WITH_EDITOR
+void UPCGMegaPlantFoliageSettings::PostEditChangeProperty(FPropertyChangedEvent& PropertyChangedEvent)
+{
+	Super::PostEditChangeProperty(PropertyChangedEvent);
+
+	if (PropertyChangedEvent.Property
+		&& PropertyChangedEvent.Property->GetFName() == GET_MEMBER_NAME_CHECKED(UPCGMegaPlantFoliageSettings, bOverrideFoliage))
+	{
+		DirtyCache();
+	}
+}
+#endif
+
+FPCGElementPtr UPCGMegaPlantFoliageSettings::CreateElement() const
+{
+	return MakeShared<FPCGMegaPlantFoliageElement>();
+}
+
+bool FPCGMegaPlantFoliageElement::ExecuteInternal(FPCGContext* InContext) const
+{
+	TRACE_CPUPROFILER_EVENT_SCOPE(FPCGMegaPlantFoliageElement::Execute);
+
+	check(InContext);
+
+	const UPCGMegaPlantFoliageSettings* Settings = InContext->GetInputSettings<UPCGMegaPlantFoliageSettings>();
+	check(Settings);
+
+	UPCGMegaPlantFoliageSettings* const MutableSettings = const_cast<UPCGMegaPlantFoliageSettings*>(Settings);
+	check(MutableSettings);
+
+	const TArray<FPCGTaggedData>& MegaPlantInputs = InContext->InputData.GetInputsByPin(PCGPinConstants::DefaultInputLabel);
+	if (!MegaPlantInputs.IsEmpty())
+	{
+		if (const UPCGMegaPlantData* InputData = Cast<UPCGMegaPlantData>(MegaPlantInputs[0].Data))
+		{
+			FManagedArrayCollection Collection = InputData->GetCollection();
+
+			UPCGMegaPlantData* OutManagedArrayCollectionData = FPCGContext::NewObject_AnyThread<UPCGMegaPlantData>(InContext);
+
+			MegaPlants::Facades::FFoliageFacade Facade(Collection);
+
+			if (!Settings->bOverrideFoliage)
+			{
+				MutableSettings->FoliageMeshes.Empty();
+				MutableSettings->bOverrideFoliage = false;
+				for (int32 FoliageIndex = 0; FoliageIndex < Facade.NumFoliageNames(); ++FoliageIndex)
+				{
+					MutableSettings->FoliageMeshes.Emplace(FSoftObjectPath(Facade.GetFoliageName(FoliageIndex)));
+				}
+#if WITH_EDITOR
+				MutableSettings->PostEditChange();
+#endif
+			}
+
+			const int32 IterationLimit = FMath::Min(Settings->FoliageMeshes.Num(), Facade.NumFoliageNames());
+			TArray<FString> FoliageNames;
+			for (int32 i = 0; i < IterationLimit; ++i)
+			{
+				FoliageNames.Emplace(Settings->FoliageMeshes[i].ToString());
+			}
+
+			if (Settings->FoliageMeshes.Num() > Facade.NumFoliageNames())
+			{
+				UE_LOG(LogMegaPlants, Warning,
+				       TEXT("Num of Foliage meshes specified is greater than the number of available entries, ignoring the last %d entries"),
+				       Settings->FoliageMeshes.Num() - Facade.NumFoliageNames());
+			}
+			else if (Settings->FoliageMeshes.Num() < Facade.NumFoliageNames())
+			{
+				if (Settings->FoliageMeshes.Num() > 0)
+				{
+					UE_LOG(LogMegaPlants, Warning,
+					       TEXT("Num of Foliage meshes specified is less than the number of available entries, repeating the last entry"));
+
+					for (int32 i = Settings->FoliageMeshes.Num(); i < Facade.NumFoliageNames(); ++i)
+					{
+						FoliageNames.Emplace(Settings->FoliageMeshes.Last().ToString());
+					}
+				}
+			}
+
+			Facade.SetFoliageNames(FoliageNames);
+
+			OutManagedArrayCollectionData->Initialize(MoveTemp(Collection));
+			InContext->OutputData.TaggedData.Emplace(OutManagedArrayCollectionData);
+		}
+		else
+		{
+			PCGLog::InputOutput::LogInvalidInputDataError(InContext);
+			return true;
+		}
+	}
+
+	return true;
+}
diff --git a/Engine/Plugins/Experimental/ProceduralVegetationEditor/Source/MegaPlants/Private/Nodes/PCGMegaPlantFoliage.h b/Engine/Plugins/Experimental/ProceduralVegetationEditor/Source/MegaPlants/Private/Nodes/PCGMegaPlantFoliage.h
new file mode 100644
--- /dev/null
+++ b/Engine/Plugins/Experimental/ProceduralVegetationEditor/Source/MegaPlants/Private/Nodes/PCGMegaPlantFoliage.h
@@ -0,0 +1,47 @@
+﻿// Copyright Epic Games, Inc. All Rights Reserved.
+#pragma once
+
+#include "CoreMinimal.h"
+#include "Nodes/PCGMegaPlantBaseSettings.h"
+#include "PCGMegaPlantData.h"
+#include "PCGMegaPlantFoliage.generated.h"
+
+UCLASS(BlueprintType, ClassGroup = (Procedural))
+class UPCGMegaPlantFoliageSettings : public UPCGMegaPlantBaseSettings
+{
+	GENERATED_BODY()
+
+public:
+#if WITH_EDITOR
+	virtual FName GetDefaultNodeName() const override { return FName(TEXT("MegaPlantFoliage")); }
+	virtual FText GetDefaultNodeTitle() const override { return NSLOCTEXT("PCGMegaPlantFoliage", "NodeTitle", "Foliage Palette"); }
+	virtual FText GetNodeTooltipText() const override { return NSLOCTEXT("PCGMegaPlantFoliage", "NodeTooltip", "Foliage setup for the tree"); }
+	virtual FLinearColor GetNodeTitleColor() const override { return FLinearColor::Green; }
+
+	virtual TArray<EMegaPlantRenderType> GetRenderType() const override { return TArray{EMegaPlantRenderType::Mesh, EMegaPlantRenderType::Foliage}; }
+
+	virtual void PostEditChangeProperty(FPropertyChangedEvent& PropertyChangedEvent) override;
+#endif
+
+protected:
+	UPROPERTY()
+	FGuid ReferencedMegaPlantData;
+
+	virtual FPCGElementPtr CreateElement() const override;
+
+public:
+	UPROPERTY(EditAnywhere, Category="Meshes", meta=(PinHiddenByDefault, InlineEditConditionToggle))
+	bool bOverrideFoliage = false;
+
+	UPROPERTY(EditAnywhere, Category = "Meshes", Meta = (EditCondition="bOverrideFoliage", AllowedClasses = "/Script/Engine.StaticMesh,/Script/Engine.SkeletalMesh"))
+	TArray<TSoftObjectPtr<UObject>> FoliageMeshes;
+};
+
+class FPCGMegaPlantFoliageElement : public FPCGMegaPlantBaseElement
+{
+protected:
+	virtual bool ExecuteInternal(FPCGContext* InContext) const override;
+
+public:
+	virtual bool CanExecuteOnlyOnMainThread(FPCGContext* Context) const override { return true; }
+};
diff --git a/Engine/Plugins/Experimental/ProceduralVegetationEditor/Source/MegaPlants/Private/Nodes/PCGMegaPlantFoliageDistributor.cpp b/Engine/Plugins/Experimental/ProceduralVegetationEditor/Source/MegaPlants/Private/Nodes/PCGMegaPlantFoliageDistributor.cpp
new file mode 100644
--- /dev/null
+++ b/Engine/Plugins/Experimental/ProceduralVegetationEditor/Source/MegaPlants/Private/Nodes/PCGMegaPlantFoliageDistributor.cpp
@@ -0,0 +1,115 @@
+﻿// Copyright Epic Games, Inc. All Rights Reserved.
+
+#include "PCGMegaPlantFoliageDistributor.h"
+#include "PCGMegaPlantData.h"
+#include "Data/PCGBasePointData.h"
+#include "GeometryCollection/ManagedArrayCollection.h"
+#include "Helpers/PCGHelpers.h"
+#include "Helpers/PCGSettingsHelpers.h"
+
+UPCGMegaPlantFoliageDistributorSettings::UPCGMegaPlantFoliageDistributorSettings()
+{
+	if (PCGHelpers::IsNewObjectAndNotDefault(this))
+	{
+		FRichCurve* RichCurve = InstanceSpacingRamp.GetRichCurve();
+		RichCurve->Reset();
+		const FKeyHandle Handle0 = RichCurve->AddKey(0.0f, 0.0f);
+		const FKeyHandle Handle1 = RichCurve->AddKey(1.0f, 1.0f);
+		RichCurve->SetKeyInterpMode(Handle0, RCIM_Linear);
+		RichCurve->SetKeyInterpMode(Handle1, RCIM_Linear);
+
+		RichCurve = ScaleRamp.GetRichCurve();
+		RichCurve->Reset();
+		const FKeyHandle Handle2 = RichCurve->AddKey(0.0f, 0.0f);
+		const FKeyHandle Handle3 = RichCurve->AddKey(1.0f, 1.0f);
+		RichCurve->SetKeyInterpMode(Handle2, RCIM_Linear);
+		RichCurve->SetKeyInterpMode(Handle3, RCIM_Linear);
+
+		RichCurve = AxilAngleRamp.GetRichCurve();
+		RichCurve->Reset();
+		const FKeyHandle Handle4 = RichCurve->AddKey(0.0f, 1.0f);
+		const FKeyHandle Handle5 = RichCurve->AddKey(1.0f, 0.5f);
+		RichCurve->SetKeyInterpMode(Handle4, RCIM_Linear);
+		RichCurve->SetKeyInterpMode(Handle5, RCIM_Linear);
+	}
+}
+
+FPCGElementPtr UPCGMegaPlantFoliageDistributorSettings::CreateElement() const
+{
+	return MakeShared<FPCGMegaPlantFoliageDistributorElement>();
+}
+
+bool FPCGMegaPlantFoliageDistributorElement::ExecuteInternal(FPCGContext* InContext) const
+{
+	TRACE_CPUPROFILER_EVENT_SCOPE(FPCGMegaPlantFoliageDistributerElement::Execute);
+
+	check(InContext);
+
+	const UPCGMegaPlantFoliageDistributorSettings* Settings = InContext->GetInputSettings<UPCGMegaPlantFoliageDistributorSettings>();
+	check(Settings);
+
+	for (const FPCGTaggedData& Input : InContext->InputData.GetInputsByPin(PCGPinConstants::DefaultInputLabel))
+	{
+		if(const UPCGMegaPlantData* InputData = Cast<UPCGMegaPlantData>(Input.Data))
+		{
+			FManagedArrayCollection Collection = InputData->GetCollection();
+			
+			UPCGMegaPlantData* OutManagedArrayCollectionData = FPCGContext::NewObject_AnyThread<UPCGMegaPlantData>(InContext);
+
+			const FDistributionSettings DistributionSettings = {
+				.EthyleneThreshold = Settings->EthyleneThreshold,
+				.OverrideDistribution = Settings->OverrideDistribution,
+				.InstanceSpacing = Settings->InstanceSpacing,
+				.InstanceSpacingRamp = &Settings->InstanceSpacingRamp,
+				.InstanceSpacingRampEffect = Settings->InstanceSpacingRampEffect,
+				.MaxPerBranch = Settings->MaxPerBranch
+			};
+
+			const FScaleSettings ScaleSettings = {
+				.BaseScale = Settings->BaseScale,
+				.BranchScaleImpact = Settings->BranchScaleImpact,
+				.MinScale = Settings->MinScale,
+				.MaxScale = Settings->MaxScale,
+				.RandomScaleMin = Settings->RandomScaleMin,
+				.RandomScaleMax = Settings->RandomScaleMax,
+				.ScaleRamp = &Settings->ScaleRamp
+			};
+
+			const FVectorSettings VectorSettings = {
+				.OverrideAxilAngle = Settings->OverrideAxilAngle,
+				.AxilAngle = Settings->AxilAngle,
+				.AxilAngleRamp = &Settings->AxilAngleRamp,
+				.AxilAngleRampUpperValue = Settings->AxilAngleRampUpperValue,
+				.AxilAngleRampEffect = Settings->AxilAngleRampEffect
+			};
+
+			const FPhyllotaxySettings PhyllotaxySettings = {
+				.OverridePhyllotaxy = Settings->OverridePhyllotaxy,
+				.PhyllotaxyType = Settings->PhyllotaxyType,
+				.PhyllotaxyFormation = Settings->PhyllotaxyFormation,
+				.MinimumNodeBuds = Settings->MinimumNodeBuds,
+				.MaximumNodeBuds = Settings->MaximumNodeBuds,
+				.ResetPhyllotaxy = Settings->ResetPhyllotaxy,
+				.PhyllotaxyAdditionalAngle = Settings->PhyllotaxyAdditionalAngle,
+				.PhyllotaxyOffset = Settings->PhyllotaxyOffset
+			};
+
+			const FMiscSettings MiscSettings = {
+				.RandomSeed = Settings->RandomSeed
+			};
+
+			FMegaPlantFoliage::DistributeFoliage(Collection, Collection, DistributionSettings, ScaleSettings, VectorSettings,
+												 PhyllotaxySettings, MiscSettings);
+
+			OutManagedArrayCollectionData->Initialize(MoveTemp(Collection));
+			InContext->OutputData.TaggedData.Emplace(OutManagedArrayCollectionData);
+		}
+		else
+		{
+			PCGLog::InputOutput::LogInvalidInputDataError(InContext);
+			return true;
+		}
+	}
+	
+	return true;
+}
\ No newline at end of file
diff --git a/Engine/Plugins/Experimental/ProceduralVegetationEditor/Source/MegaPlants/Private/Nodes/PCGMegaPlantFoliageDistributor.h b/Engine/Plugins/Experimental/ProceduralVegetationEditor/Source/MegaPlants/Private/Nodes/PCGMegaPlantFoliageDistributor.h
new file mode 100644
--- /dev/null
+++ b/Engine/Plugins/Experimental/ProceduralVegetationEditor/Source/MegaPlants/Private/Nodes/PCGMegaPlantFoliageDistributor.h
@@ -0,0 +1,122 @@
+﻿// Copyright Epic Games, Inc. All Rights Reserved.
+#pragma once
+
+#include "CoreMinimal.h"
+#include "Nodes/PCGMegaPlantBaseSettings.h"
+#include "Implementations/MegaPlantFoliage.h"
+#include "PCGMegaPlantData.h"
+#include "PCGMegaPlantFoliageDistributor.generated.h"
+
+
+UCLASS(BlueprintType, ClassGroup = (Procedural))
+class UPCGMegaPlantFoliageDistributorSettings : public UPCGMegaPlantBaseSettings
+{
+	GENERATED_BODY()
+
+	friend class FPCGMegaPlantFoliageDistributorElement;
+	
+public:
+	UPCGMegaPlantFoliageDistributorSettings();
+
+	#if WITH_EDITOR
+	virtual FName GetDefaultNodeName() const override { return FName(TEXT("MegaPlantFoliageDistributor")); }
+	virtual FText GetDefaultNodeTitle() const override { return NSLOCTEXT("PCGMegaPlantFoliageDistributor", "NodeTitle", "Foliage Distributor"); }
+	virtual FText GetNodeTooltipText() const override { return NSLOCTEXT("PCGMegaPlantFoliageDistributor", "NodeTooltip", "Distributes the Foliage"); }
+	virtual FLinearColor GetNodeTitleColor() const override { return FLinearColor::Blue; }
+
+	virtual TArray<EMegaPlantRenderType> GetRenderType() const override { return TArray{ EMegaPlantRenderType::Mesh, EMegaPlantRenderType::Foliage }; }
+	#endif
+
+protected:
+
+	virtual FPCGElementPtr CreateElement() const override;
+
+private:
+
+	UPROPERTY(EditAnywhere, Category="Distribution Settings", meta=(ClampMin=0.0f, ClampMax=1.0f, UIMin=0.0f, UIMax=1.0f))
+	float EthyleneThreshold = 0.85f;
+
+	UPROPERTY(EditAnywhere, Category="Distribution Settings")
+	bool OverrideDistribution = true;
+
+	UPROPERTY(EditAnywhere, Category="Distribution Settings", meta=(ClampMin=0.0f, ClampMax=1.0f, UIMin=0.0f, UIMax=1.0f))
+	float InstanceSpacing = 0.3f;
+
+	UPROPERTY(EditAnywhere, Category="Distribution Settings", meta=(ClampMin=0.0f, ClampMax=1.0f))
+	FRuntimeFloatCurve InstanceSpacingRamp;
+
+	UPROPERTY(EditAnywhere, Category="Distribution Settings", meta=(ClampMin=0.0f, ClampMax=1.0f, UIMin=0.0f, UIMax=1.0f))
+	float InstanceSpacingRampEffect = 0.0f;
+
+	UPROPERTY(EditAnywhere, Category="Distribution Settings", meta=(ClampMin=-1))
+	int32 MaxPerBranch = -1;
+
+	UPROPERTY(EditAnywhere, Category="Scale Settings", meta=(ClampMin=0.0f, ClampMax=2.0f, UIMin=0.0f, UIMax=2.0f))
+	float BaseScale = 1.0f;
+
+	UPROPERTY(EditAnywhere, Category="Scale Settings", meta=(ClampMin=0.0f, ClampMax=1.0f, UIMin=0.0f, UIMax=1.0f))
+	float BranchScaleImpact = 0.0f;
+
+	UPROPERTY(EditAnywhere, Category="Scale Settings", meta=(ClampMin=0.0f, ClampMax=1.0f, UIMin=0.0f, UIMax=1.0f))
+	float MinScale = 0.5f;
+
+	UPROPERTY(EditAnywhere, Category="Scale Settings", meta=(ClampMin=0.0f, ClampMax=1.0f, UIMin=0.0f, UIMax=1.0f))
+	float MaxScale = 1.0f;
+
+	UPROPERTY(EditAnywhere, Category="Scale Settings", meta=(ClampMin=0.0f, ClampMax=1.0f, UIMin=0.0f, UIMax=1.0f))
+	float RandomScaleMin = 1.0f;
+
+	UPROPERTY(EditAnywhere, Category="Scale Settings", meta=(ClampMin=1.0f, ClampMax=2.0f, UIMin=1.0f, UIMax=2.0f))
+	float RandomScaleMax = 1.0f;
+
+	UPROPERTY(EditAnywhere, Category="Scale Settings", meta=(ClampMin=0.0f, ClampMax=1.0f))
+	FRuntimeFloatCurve ScaleRamp;
+
+	UPROPERTY(EditAnywhere, Category="Vector Settings")
+	bool OverrideAxilAngle = true;
+
+	UPROPERTY(EditAnywhere, Category="Vector Settings", meta=(ClampMin=0.0f, ClampMax=90.0f, UIMin=0.0f, UIMax=90.0f))
+	float AxilAngle = 45.0f;
+
+	UPROPERTY(EditAnywhere, Category="Vector Settings", meta=(ClampMin=0.0f, ClampMax=1.0f))
+	FRuntimeFloatCurve AxilAngleRamp;
+
+	UPROPERTY(EditAnywhere, Category="Vector Settings", meta=(ClampMin=0.0f, ClampMax=90.0f, UIMin=0.0f, UIMax=90.0f))
+	float AxilAngleRampUpperValue = 45.0f;
+
+	UPROPERTY(EditAnywhere, Category="Vector Settings", meta=(ClampMin=0.0f, ClampMax=1.0f, UIMin=0.0f, UIMax=1.0f))
+	float AxilAngleRampEffect = 0.0f;
+
+	UPROPERTY(EditAnywhere, Category="Phyllotaxy Settings")
+	bool OverridePhyllotaxy = false;
+
+	UPROPERTY(EditAnywhere, Category="Phyllotaxy Settings")
+	EPhyllotaxyType PhyllotaxyType = EPhyllotaxyType::Alternate;
+
+	UPROPERTY(EditAnywhere, Category="Phyllotaxy Settings", meta=(EditCondition="PhyllotaxyType == EPhyllotaxyType::Spiral", EditConditionHides))
+	EPhyllotaxyFormation PhyllotaxyFormation = EPhyllotaxyFormation::Distichous;
+
+	UPROPERTY(EditAnywhere, Category="Phyllotaxy Settings", meta=(ClampMin=1, ClampMax=10, UIMin=1, UIMax=10))
+	int32 MinimumNodeBuds = 2;
+
+	UPROPERTY(EditAnywhere, Category="Phyllotaxy Settings", meta=(ClampMin=1, ClampMax=10, UIMin=1, UIMax=10))
+	int32 MaximumNodeBuds = 3;
+
+	UPROPERTY(EditAnywhere, Category="Phyllotaxy Settings")
+	bool ResetPhyllotaxy = false;
+
+	UPROPERTY(EditAnywhere, Category="Phyllotaxy Settings", meta=(ClampMin=0.0f, ClampMax=360.0f, UIMin=0.0f, UIMax=360.0f))
+	float PhyllotaxyAdditionalAngle = 0.0f;
+
+	UPROPERTY(EditAnywhere, Category="Phyllotaxy Settings", meta=(ClampMin=0.0f, ClampMax=360.0f, UIMin=0.0f, UIMax=360.0f))
+	float PhyllotaxyOffset = 0.0f;
+
+	UPROPERTY(EditAnywhere, Category="Misc Settings")
+	int32 RandomSeed = 123456;
+};
+
+class FPCGMegaPlantFoliageDistributorElement : public FPCGMegaPlantBaseElement
+{
+protected:
+	virtual bool ExecuteInternal(FPCGContext* InContext) const override;
+};
diff --git a/Engine/Plugins/Experimental/ProceduralVegetationEditor/Source/MegaPlants/Private/Nodes/PCGMegaPlantGravity.cpp b/Engine/Plugins/Experimental/ProceduralVegetationEditor/Source/MegaPlants/Private/Nodes/PCGMegaPlantGravity.cpp
new file mode 100644
--- /dev/null
+++ b/Engine/Plugins/Experimental/ProceduralVegetationEditor/Source/MegaPlants/Private/Nodes/PCGMegaPlantGravity.cpp
@@ -0,0 +1,48 @@
+﻿// Copyright Epic Games, Inc. All Rights Reserved.
+
+#include "PCGMegaPlantGravity.h"
+#include "PCGMegaPlantData.h"
+#include "Data/PCGBasePointData.h"
+#include "GeometryCollection/ManagedArrayCollection.h"
+#include "Helpers/PCGSettingsHelpers.h"
+#include "Implementations/MegaPlantGravity.h"
+
+FPCGElementPtr UPCGMegaPlantGravitySettings::CreateElement() const
+{
+	return MakeShared<FPCGMegaPlantGravityElement>();
+}
+
+bool FPCGMegaPlantGravityElement::ExecuteInternal(FPCGContext* InContext) const
+{
+	TRACE_CPUPROFILER_EVENT_SCOPE(FPCGMegaPlantGravityElement::Execute);
+
+	check(InContext);
+
+	const UPCGMegaPlantGravitySettings* Settings = InContext->GetInputSettings<UPCGMegaPlantGravitySettings>();
+	check(Settings);
+	
+	for (const FPCGTaggedData& Input : InContext->InputData.GetInputsByPin(PCGPinConstants::DefaultInputLabel))
+	{
+		if(const UPCGMegaPlantData* InputData = Cast<UPCGMegaPlantData>(Input.Data))
+		{
+			FManagedArrayCollection OutCollection = InputData->GetCollection();
+		
+			UPCGMegaPlantData* OutManagedArrayCollectionData = FPCGContext::NewObject_AnyThread<UPCGMegaPlantData>(InContext);
+			
+			if (Settings->GravitySettings.Gravity != 0.0f)
+			{
+				FMegaPlantGravity::ApplyGravity(Settings->GravitySettings, OutCollection);
+			}
+
+			OutManagedArrayCollectionData->Initialize(MoveTemp(OutCollection));
+			InContext->OutputData.TaggedData.Emplace(OutManagedArrayCollectionData);
+		}
+		else
+		{
+			PCGLog::InputOutput::LogInvalidInputDataError(InContext);
+			return true;
+		}
+	}
+	
+	return true;
+}
diff --git a/Engine/Plugins/Experimental/ProceduralVegetationEditor/Source/MegaPlants/Private/Nodes/PCGMegaPlantGravity.h b/Engine/Plugins/Experimental/ProceduralVegetationEditor/Source/MegaPlants/Private/Nodes/PCGMegaPlantGravity.h
new file mode 100644
--- /dev/null
+++ b/Engine/Plugins/Experimental/ProceduralVegetationEditor/Source/MegaPlants/Private/Nodes/PCGMegaPlantGravity.h
@@ -0,0 +1,35 @@
+﻿// Copyright Epic Games, Inc. All Rights Reserved.
+#pragma once
+
+#include "CoreMinimal.h"
+#include "Nodes/PCGMegaPlantBaseSettings.h"
+#include "Implementations/MegaPlantGravity.h"
+#include "PCGMegaPlantGravity.generated.h"
+
+UCLASS(BlueprintType, ClassGroup = (Procedural))
+class UPCGMegaPlantGravitySettings : public UPCGMegaPlantBaseSettings
+{
+	GENERATED_BODY()
+
+public:
+#if WITH_EDITOR
+	virtual FName GetDefaultNodeName() const override { return FName(TEXT("MegaPlantGravity")); }
+	virtual FText GetDefaultNodeTitle() const override { return NSLOCTEXT("PCGMegaPlantGravity", "NodeTitle", "Gravity"); }
+	virtual FText GetNodeTooltipText() const override { return NSLOCTEXT("PCGMegaPlantGravity", "NodeTooltip", "Apply gravity to branches and main trunk"); }
+	virtual FLinearColor GetNodeTitleColor() const override { return FLinearColor::Yellow; }
+
+	virtual TArray<EMegaPlantRenderType> GetRenderType() const override { return TArray{ EMegaPlantRenderType::Skeleton }; }
+#endif
+
+protected:
+	virtual FPCGElementPtr CreateElement() const override;
+public:
+	UPROPERTY(EditAnywhere, Category="Gravity", meta=(ShowOnlyInnerProperties))
+	FMegaPlantGravitySettings GravitySettings;
+};
+
+class FPCGMegaPlantGravityElement : public FPCGMegaPlantBaseElement
+{
+protected:
+	virtual bool ExecuteInternal(FPCGContext* InContext) const override;
+};
diff --git a/Engine/Plugins/Experimental/ProceduralVegetationEditor/Source/MegaPlants/Private/Nodes/PCGMegaPlantImporter.cpp b/Engine/Plugins/Experimental/ProceduralVegetationEditor/Source/MegaPlants/Private/Nodes/PCGMegaPlantImporter.cpp
new file mode 100644
--- /dev/null
+++ b/Engine/Plugins/Experimental/ProceduralVegetationEditor/Source/MegaPlants/Private/Nodes/PCGMegaPlantImporter.cpp
@@ -0,0 +1,48 @@
+﻿// Copyright Epic Games, Inc. All Rights Reserved.
+
+#include "PCGMegaPlantImporter.h"
+
+#include "MegaPlantsModule.h"
+#include "PCGMegaPlantData.h"
+#include "GeometryCollection/ManagedArrayCollection.h"
+#include "Helpers/MegaPlantJSONHelper.h"
+#include "Helpers/PCGHelpers.h"
+#include "Helpers/PCGSettingsHelpers.h"
+
+TArray<FPCGPinProperties> UPCGMegaPlantImporterSettings::InputPinProperties() const
+{
+	return TArray<FPCGPinProperties>();
+}
+
+FPCGElementPtr UPCGMegaPlantImporterSettings::CreateElement() const
+{
+	return MakeShared<FPCGMegaPlantImporterElement>();
+}
+
+bool FPCGMegaPlantImporterElement::ExecuteInternal(FPCGContext* InContext) const
+{
+	TRACE_CPUPROFILER_EVENT_SCOPE(FPCGMegaPlantImporterElement::Execute);
+
+	check(InContext);
+
+	const UPCGMegaPlantImporterSettings* Settings = InContext->GetInputSettings<UPCGMegaPlantImporterSettings>();
+	check(Settings);
+	
+	const FString& SkeletonFilePath = Settings->SkeletonFile.FilePath;
+	FString ErrorMessage;
+
+	FManagedArrayCollection Collection;
+	bool bSuccess = MegaPlants::LoadMegaPlantJsonToCollection(Collection, SkeletonFilePath, ErrorMessage);
+	if (!bSuccess)
+	{
+		UE_LOG(LogMegaPlants, Error, TEXT("%s"), *ErrorMessage);
+		PCGLog::InputOutput::LogInvalidInputDataError(InContext);
+		return true;
+	}
+	UPCGMegaPlantData* OutManagedArrayCollectionData = FPCGContext::NewObject_AnyThread<UPCGMegaPlantData>(InContext);
+	OutManagedArrayCollectionData->Initialize(MoveTemp(Collection));
+		
+	InContext->OutputData.TaggedData.Emplace(OutManagedArrayCollectionData);
+	
+	return true;
+}
diff --git a/Engine/Plugins/Experimental/ProceduralVegetationEditor/Source/MegaPlants/Private/Nodes/PCGMegaPlantImporter.h b/Engine/Plugins/Experimental/ProceduralVegetationEditor/Source/MegaPlants/Private/Nodes/PCGMegaPlantImporter.h
new file mode 100644
--- /dev/null
+++ b/Engine/Plugins/Experimental/ProceduralVegetationEditor/Source/MegaPlants/Private/Nodes/PCGMegaPlantImporter.h
@@ -0,0 +1,38 @@
+﻿// Copyright Epic Games, Inc. All Rights Reserved.
+#pragma once
+
+#include "CoreMinimal.h"
+#include "Nodes/PCGMegaPlantBaseSettings.h"
+#include "PCGSettings.h"
+
+#include "PCGMegaPlantImporter.generated.h"
+
+UCLASS(BlueprintType, ClassGroup = (Procedural))
+class UPCGMegaPlantImporterSettings : public UPCGMegaPlantBaseSettings
+{
+	GENERATED_BODY()
+
+public:
+#if WITH_EDITOR
+	virtual FName GetDefaultNodeName() const override { return FName(TEXT("MegaPlantImporter")); }
+	virtual FText GetDefaultNodeTitle() const override { return NSLOCTEXT("PCGMegaPlantImporter", "NodeTitle", "MegaPlant Importer"); }
+	virtual FText GetNodeTooltipText() const override { return NSLOCTEXT("PCGMegaPlantImporter", "NodeTooltip", "Create MegaPlant skeleton"); }
+	virtual FLinearColor GetNodeTitleColor() const override { return FLinearColor(FColor::Black); }
+
+	virtual TArray<EMegaPlantRenderType> GetRenderType() const override { return TArray{ EMegaPlantRenderType::Skeleton }; }
+#endif
+
+protected:
+	virtual TArray<FPCGPinProperties> InputPinProperties() const override;
+	virtual FPCGElementPtr CreateElement() const override;
+
+public:
+	UPROPERTY(EditAnywhere, Category="Skeleton")
+	FFilePath SkeletonFile;
+};
+
+class FPCGMegaPlantImporterElement : public FPCGMegaPlantBaseElement
+{
+protected:
+	virtual bool ExecuteInternal(FPCGContext* InContext) const override;
+};
