Create ProceduralVegetationEditor plugin

diff --git a/Engine/Plugins/Experimental/ProceduralVegetationEditor/Source/MegaPlants/Private/Nodes/PCGMegaPlantMesher.cpp b/Engine/Plugins/Experimental/ProceduralVegetationEditor/Source/MegaPlants/Private/Nodes/PCGMegaPlantMesher.cpp
new file mode 100644
--- /dev/null
+++ b/Engine/Plugins/Experimental/ProceduralVegetationEditor/Source/MegaPlants/Private/Nodes/PCGMegaPlantMesher.cpp
@@ -0,0 +1,135 @@
+﻿// Copyright Epic Games, Inc. All Rights Reserved.
+
+#include "PCGMegaPlantMesher.h"
+#include "PCGMegaPlantData.h"
+#include "Data/PCGBasePointData.h"
+#include "Facades/MegaPlantBranchFacade.h"
+#include "Facades/MegaPlantProfileFacade.h"
+#include "GeometryCollection/ManagedArrayCollection.h"
+#include "Helpers/PCGSettingsHelpers.h"
+#include "Implementations/MegaPlantMesher.h"
+#include "Materials/MaterialInterface.h"
+
+#if WITH_EDITOR
+void UPCGMegaPlantMesherSettings::PostEditChangeProperty(FPropertyChangedEvent& PropertyChangedEvent)
+{
+	Super::PostEditChangeProperty(PropertyChangedEvent);
+
+	if (PropertyChangedEvent.Property
+		&& PropertyChangedEvent.Property->GetFName() == GET_MEMBER_NAME_CHECKED(FMegaPlantMesherSettings, DisplacementTexture))
+	{
+		if (MesherSettings.DisplacementTexture)
+		{
+			FMegaPlantMesher::ExtractDisplacementData(MesherSettings.DisplacementTexture, MesherSettings.DisplacementValues);
+		}
+		else
+		{
+			MesherSettings.DisplacementValues.Empty();
+		}
+
+		Modify();
+	}
+}
+
+void UPCGMegaPlantMesherSettings::PostLoad()
+{
+	Super::PostLoad();
+
+	if (MesherSettings.DisplacementTexture)
+	{
+		FMegaPlantMesher::ExtractDisplacementData(MesherSettings.DisplacementTexture, MesherSettings.DisplacementValues);
+	}
+	else
+	{
+		MesherSettings.DisplacementValues.Empty();
+	}
+}
+#endif
+
+FPCGElementPtr UPCGMegaPlantMesherSettings::CreateElement() const
+{
+	return MakeShared<FPCGMegaPlantMesherElement>();
+}
+
+bool FPCGMegaPlantMesherElement::ExecuteInternal(FPCGContext* InContext) const
+{
+	TRACE_CPUPROFILER_EVENT_SCOPE(FPCGMegaPlantMesherElement::Execute);
+
+	check(InContext);
+
+	const UPCGMegaPlantMesherSettings* Settings = InContext->GetInputSettings<UPCGMegaPlantMesherSettings>();
+	check(Settings);
+
+	UPCGMegaPlantMesherSettings* MutableSettings = const_cast<UPCGMegaPlantMesherSettings*>(Settings);
+	check(MutableSettings);
+
+	MutableSettings->MesherSettings.PlantProfileOptions.Empty();
+	MutableSettings->MesherSettings.PlantProfileOptions.Add(TEXT("None"));
+
+	const TArray<FPCGTaggedData>& MegaPlantInputs = InContext->InputData.GetInputsByPin(PCGPinConstants::DefaultInputLabel);
+	if (!MegaPlantInputs.IsEmpty())
+	{
+		if (const UPCGMegaPlantData* InputData = Cast<UPCGMegaPlantData>(MegaPlantInputs[0].Data))
+		{
+			const FManagedArrayCollection& SkeletonCollection = InputData->GetCollection();
+			auto& MaterialSettings = MutableSettings->MesherSettings.MaterialSettings;
+					
+			if (!Settings->MesherSettings.bOverrideMaterial || MaterialSettings.MaterialSetups.Num() == 0)
+			{
+				const MegaPlants::Facades::FBranchFacade BranchFacade(SkeletonCollection);
+				const FString TrunkMaterial = BranchFacade.GetTrunkMaterialPath();
+				if (!TrunkMaterial.IsEmpty())
+				{
+					MaterialSettings.SetMaterial(TrunkMaterial, 0);
+				}
+				else
+				{
+					MutableSettings->MesherSettings.bOverrideMaterial = true;
+				}
+			}
+
+			if (const MegaPlants::Facades::FPlantProfileFacade PlantProfileFacade = MegaPlants::Facades::FPlantProfileFacade(SkeletonCollection);
+				PlantProfileFacade.NumProfileEntries() > 0)
+			{
+				for (int32 i = 0; i < PlantProfileFacade.NumProfileEntries(); i++)
+				{
+					MutableSettings->MesherSettings.PlantProfileOptions.Add(FString::FromInt(i));
+				}
+
+				MutableSettings->MesherSettings.bIsPlantProfileDropdownEnabled = true;
+			}
+			else
+			{
+				MutableSettings->MesherSettings.bIsPlantProfileDropdownEnabled = false;
+			}
+
+			UPCGMegaPlantData* OutManagedArrayCollectionData = FPCGContext::NewObject_AnyThread<UPCGMegaPlantData>(InContext);
+
+			FGeometryCollection OutGeometryCollection;
+			FMegaPlantMesher::GenerateGeometryCollection(SkeletonCollection, Settings->MesherSettings, OutGeometryCollection);
+
+			OutManagedArrayCollectionData->Initialize(MoveTemp(OutGeometryCollection));
+
+			FPCGTaggedData& CollectionOutput = InContext->OutputData.TaggedData.Emplace_GetRef();
+			CollectionOutput.Data = OutManagedArrayCollectionData;
+			CollectionOutput.Pin = PCGPinConstants::DefaultOutputLabel;
+
+			// TODO : Figure out if we want to keep the dynamic mesh inn the megaplant data or as a pin
+			// UPCGDynamicMeshData* DynamicMeshData = FPCGContext::NewObject_AnyThread<UPCGDynamicMeshData>(InContext);
+			// TObjectPtr<UDynamicMesh> OutMesh = NewObject<UDynamicMesh>();
+			// FMegaPlantMesher::GenerateDynamicMesh(SkeletonCollection, OutMesh);
+			// DynamicMeshData->Initialize(OutMesh, true);
+			//
+			// FPCGTaggedData& MeshOutput = InContext->OutputData.TaggedData.Emplace_GetRef();
+			// MeshOutput.Data = DynamicMeshData;
+			// MeshOutput.Pin = TEXT("Mesh");
+		}
+		else
+		{
+			PCGLog::InputOutput::LogInvalidInputDataError(InContext);
+			return true;
+		}
+	}
+
+	return true;
+}
diff --git a/Engine/Plugins/Experimental/ProceduralVegetationEditor/Source/MegaPlants/Private/Nodes/PCGMegaPlantMesher.h b/Engine/Plugins/Experimental/ProceduralVegetationEditor/Source/MegaPlants/Private/Nodes/PCGMegaPlantMesher.h
new file mode 100644
--- /dev/null
+++ b/Engine/Plugins/Experimental/ProceduralVegetationEditor/Source/MegaPlants/Private/Nodes/PCGMegaPlantMesher.h
@@ -0,0 +1,49 @@
+﻿// Copyright Epic Games, Inc. All Rights Reserved.
+#pragma once
+
+#include "CoreMinimal.h"
+#include "Implementations/MegaPlantMesher.h"
+#include "Nodes/PCGMegaPlantBaseSettings.h"
+#include "PCGMegaPlantData.h"
+#include "PCGMegaPlantMesher.generated.h"
+
+class UMaterialInterface;
+
+UCLASS(BlueprintType, ClassGroup = (Procedural))
+class UPCGMegaPlantMesherSettings : public UPCGMegaPlantBaseSettings
+{
+	GENERATED_BODY()
+
+public:
+#if WITH_EDITOR
+	virtual FName GetDefaultNodeName() const override { return FName(TEXT("MegaPlantMesher")); }
+	virtual FText GetDefaultNodeTitle() const override { return NSLOCTEXT("PCGMegaPlantMesher", "NodeTitle", "Mesh Builder"); }
+	virtual FText GetNodeTooltipText() const override { return NSLOCTEXT("PCGMegaPlantMesher", "NodeTooltip", "Create a mesh for the skeleton"); }
+	virtual FLinearColor GetNodeTitleColor() const override { return FLinearColor(FColor::Cyan); }
+
+	virtual TArray<EMegaPlantRenderType> GetRenderType() const override { return TArray{EMegaPlantRenderType::Mesh}; }
+
+	virtual void PostEditChangeProperty(FPropertyChangedEvent& PropertyChangedEvent) override;
+
+	virtual void PostLoad() override;
+#endif
+
+protected:
+	virtual FPCGElementPtr CreateElement() const override;
+
+public:	
+	UPROPERTY(EditAnywhere, Category = "Mesh", meta = (ShowOnlyInnerProperties))
+	FMegaPlantMesherSettings MesherSettings;
+
+	UFUNCTION()
+	TArray<FString> GetPlantProfileOptions() const
+	{
+		return MesherSettings.PlantProfileOptions;
+	}
+};
+
+class FPCGMegaPlantMesherElement : public FPCGMegaPlantBaseElement
+{
+protected:
+	virtual bool ExecuteInternal(FPCGContext* InContext) const override;
+};
diff --git a/Engine/Plugins/Experimental/ProceduralVegetationEditor/Source/MegaPlants/Private/Nodes/PCGMegaPlantOutput.cpp b/Engine/Plugins/Experimental/ProceduralVegetationEditor/Source/MegaPlants/Private/Nodes/PCGMegaPlantOutput.cpp
new file mode 100644
--- /dev/null
+++ b/Engine/Plugins/Experimental/ProceduralVegetationEditor/Source/MegaPlants/Private/Nodes/PCGMegaPlantOutput.cpp
@@ -0,0 +1,98 @@
+﻿// Copyright Epic Games, Inc. All Rights Reserved.
+
+#include "Nodes/PCGMegaPlantOutput.h"
+#include "PCGContext.h"
+#include "Helpers/PCGHelpers.h"
+#include "Misc/PackageName.h"
+
+#define LOCTEXT_NAMESPACE "PCGMegaPlantOutput"
+
+FString FormatLongPackageNameErrorCode(const FPackageName::EErrorCode ErrorCode)
+{
+	switch (ErrorCode)
+	{
+	case FPackageName::EErrorCode::PackageNameUnknown:
+		return "Unknown Export Path error";
+	case FPackageName::EErrorCode::PackageNameEmptyPath:
+		return "Empty Export Path";
+	case FPackageName::EErrorCode::PackageNamePathNotMounted:
+		return "Mount Point for Export Path is not valid";
+	case FPackageName::EErrorCode::PackageNamePathIsMemoryOnly:
+		return "Export Path exists in memory only. ";
+	case FPackageName::EErrorCode::PackageNameSpacesNotAllowed:
+		return "Spaces are not allowed in Export Path";	
+	case FPackageName::EErrorCode::PackageNameContainsInvalidCharacters:
+		return FString::Printf(TEXT("Export Path contains one of \n%s\n Invalid Characters"), TEXT(R"(\\:*?\"<>|' ,.&!~\n\r\t@#)"));
+	case FPackageName::EErrorCode::LongPackageNames_PathTooShort:
+		return "Export Path is too small to be a valid path";
+	case FPackageName::EErrorCode::LongPackageNames_PathWithNoStartingSlash:
+		return "Export Path has to start with a slash [/]";
+	case FPackageName::EErrorCode::LongPackageNames_PathWithTrailingSlash:
+		return "Export Path cannot end with a slash [/]";
+	case FPackageName::EErrorCode::LongPackageNames_PathWithDoubleSlash:
+		return "Export Path contains a double slash [/]";
+	}
+	return FString();
+}
+
+UPCGMegaPlantOutputSettings::UPCGMegaPlantOutputSettings()
+{
+	if (PCGHelpers::IsNewObjectAndNotDefault(this))
+	{
+		static FString DefaultMeshName = "MegaPlantOutput";
+		static FString DefaultAssetPath = "/Game/MegaPlants/";
+		ExportSettings.Initialize(DefaultAssetPath, DefaultMeshName);
+	}
+}
+
+TArray<FPCGPinProperties> UPCGMegaPlantOutputSettings::OutputPinProperties() const
+{
+	TArray<FPCGPinProperties> Properties;
+
+	FPCGPinProperties& Pin = Properties.Emplace_GetRef(PCGPinConstants::DefaultOutputLabel, EPCGDataType::Any);
+	Pin.SetRequiredPin();
+	Pin.SetAllowMultipleConnections(false);
+	Pin.bInvisiblePin = true;
+
+	return Properties;
+}
+
+FPCGElementPtr UPCGMegaPlantOutputSettings::CreateElement() const
+{
+	return MakeShared<UPCGMegaPlantOutputElement>();
+}
+
+bool UPCGMegaPlantOutputElement::ExecuteInternal(FPCGContext* InContext) const
+{
+	const UPCGMegaPlantOutputSettings* Settings = InContext->GetInputSettings<UPCGMegaPlantOutputSettings>();
+	check(Settings);
+
+	for (const FPCGTaggedData& Input : InContext->InputData.GetInputsByPin(PCGPinConstants::DefaultInputLabel))
+	{
+		if (const UPCGMegaPlantData* InputData = Cast<UPCGMegaPlantData>(Input.Data))
+		{
+			const FString OutputPath = Settings->ExportSettings.ContentBrowserFolder.Path;
+			if (FPackageName::EErrorCode ErrorCode; !FPackageName::IsValidLongPackageName(OutputPath, false, &ErrorCode))
+			{
+				PCGLog::LogErrorOnGraph(FText::FromString(FormatLongPackageNameErrorCode(ErrorCode)), InContext);
+				return true;
+			}
+
+			FManagedArrayCollection OutCollection = InputData->GetCollection();
+
+			UPCGMegaPlantData* MegaPlantData = FPCGContext::NewObject_AnyThread<UPCGMegaPlantData>(InContext);
+			MegaPlantData->Initialize(MoveTemp(OutCollection));
+
+			InContext->OutputData.TaggedData.Emplace(MegaPlantData);
+		}
+		else
+		{
+			PCGLog::InputOutput::LogInvalidInputDataError(InContext);
+			return true;
+		}
+	}
+
+	return true;
+}
+
+#undef LOCTEXT_NAMESPACE
diff --git a/Engine/Plugins/Experimental/ProceduralVegetationEditor/Source/MegaPlants/Private/Nodes/PCGMegaPlantRecipeLoad.cpp b/Engine/Plugins/Experimental/ProceduralVegetationEditor/Source/MegaPlants/Private/Nodes/PCGMegaPlantRecipeLoad.cpp
new file mode 100644
--- /dev/null
+++ b/Engine/Plugins/Experimental/ProceduralVegetationEditor/Source/MegaPlants/Private/Nodes/PCGMegaPlantRecipeLoad.cpp
@@ -0,0 +1,57 @@
+// Copyright Epic Games, Inc. All Rights Reserved.
+
+#include "PCGMegaPlantRecipeLoad.h"
+#include "PCGContext.h"
+
+TArray<FPCGPinProperties> UPCGMegaPlantRecipeLoadSettings::InputPinProperties() const
+{
+	return TArray<FPCGPinProperties>();
+}
+
+TArray<FPCGPinProperties> UPCGMegaPlantRecipeLoadSettings::OutputPinProperties() const
+{
+	TArray<FPCGPinProperties> Properties;
+	if (RecipeDataAsset)
+	{
+		for (const auto& [VariantName, VariantData] : RecipeDataAsset->Variants)
+		{
+			FPCGPinProperties& Pin = Properties.Emplace_GetRef(*VariantName, EPCGDataType::Any);
+			Pin.bAllowMultipleData = false;
+		}
+	}
+
+	return Properties;
+}
+
+FPCGElementPtr UPCGMegaPlantRecipeLoadSettings::CreateElement() const
+{
+	return MakeShared<FPCGMegaPlantRecipeLoadElement>();
+}
+
+bool FPCGMegaPlantRecipeLoadElement::ExecuteInternal(FPCGContext* InContext) const
+{
+	TRACE_CPUPROFILER_EVENT_SCOPE(FPCGMegaPlantRecipeLoadElement::ExecuteInternal);
+
+	check(InContext);
+
+	const UPCGMegaPlantRecipeLoadSettings* Settings = InContext->GetInputSettings<UPCGMegaPlantRecipeLoadSettings>();
+	check(Settings);
+
+	if (!Settings->RecipeDataAsset)
+	{
+		PCGLog::LogWarningOnGraph(FText::FromString("Recipe is not set"), InContext);
+		return true;
+	}
+
+	for (const auto& [VariantName, VariantData] : Settings->RecipeDataAsset->Variants)
+	{
+		UPCGMegaPlantData* OutVariantData = FPCGContext::NewObject_AnyThread<UPCGMegaPlantData>(InContext);
+		OutVariantData->Initialize(CopyTemp(VariantData));
+
+		FPCGTaggedData& CollectionOutput = InContext->OutputData.TaggedData.Emplace_GetRef();
+		CollectionOutput.Data = OutVariantData;
+		CollectionOutput.Pin = *VariantName;
+	}
+
+	return true;
+}
diff --git a/Engine/Plugins/Experimental/ProceduralVegetationEditor/Source/MegaPlants/Private/Nodes/PCGMegaPlantRecipeLoad.h b/Engine/Plugins/Experimental/ProceduralVegetationEditor/Source/MegaPlants/Private/Nodes/PCGMegaPlantRecipeLoad.h
new file mode 100644
--- /dev/null
+++ b/Engine/Plugins/Experimental/ProceduralVegetationEditor/Source/MegaPlants/Private/Nodes/PCGMegaPlantRecipeLoad.h
@@ -0,0 +1,40 @@
+// Copyright Epic Games, Inc. All Rights Reserved.
+#pragma once
+
+#include "CoreMinimal.h"
+#include "MegaPlantRecipeDataAsset.h"
+
+#include "Nodes/PCGMegaPlantBaseSettings.h"
+#include "PCGSettings.h"
+
+#include "PCGMegaPlantRecipeLoad.generated.h"
+
+UCLASS(BlueprintType, ClassGroup = (Procedural))
+class UPCGMegaPlantRecipeLoadSettings : public UPCGMegaPlantBaseSettings
+{
+	GENERATED_BODY()
+
+public:
+#if WITH_EDITOR
+	virtual FName GetDefaultNodeName() const override { return FName(TEXT("MegaPlantRecipeLoad")); }
+	virtual FText GetDefaultNodeTitle() const override { return NSLOCTEXT("PCGMegaPlantRecipeLoad", "NodeTitle", "MegaPlant Recipe Load"); }
+	virtual FText GetNodeTooltipText() const override { return NSLOCTEXT("PCGMegaPlantRecipeLoad", "NodeTooltip", "Load a MegaPlant recipe"); }
+
+	virtual TArray<EMegaPlantRenderType> GetRenderType() const override { return TArray{ EMegaPlantRenderType::Skeleton }; }
+#endif
+
+protected:
+	virtual TArray<FPCGPinProperties> InputPinProperties() const override;
+	virtual TArray<FPCGPinProperties> OutputPinProperties() const override;
+	virtual FPCGElementPtr CreateElement() const override;
+
+public:
+	UPROPERTY(EditAnywhere, Category="Recipe")
+	TObjectPtr<UMegaPlantRecipeDataAsset> RecipeDataAsset = nullptr;
+};
+
+class FPCGMegaPlantRecipeLoadElement : public FPCGMegaPlantBaseElement
+{
+protected:
+	virtual bool ExecuteInternal(FPCGContext* InContext) const override;
+};
diff --git a/Engine/Plugins/Experimental/ProceduralVegetationEditor/Source/MegaPlants/Private/Nodes/PCGMegaPlantRemoveBranches.cpp b/Engine/Plugins/Experimental/ProceduralVegetationEditor/Source/MegaPlants/Private/Nodes/PCGMegaPlantRemoveBranches.cpp
new file mode 100644
--- /dev/null
+++ b/Engine/Plugins/Experimental/ProceduralVegetationEditor/Source/MegaPlants/Private/Nodes/PCGMegaPlantRemoveBranches.cpp
@@ -0,0 +1,45 @@
+﻿// Copyright Epic Games, Inc. All Rights Reserved.
+
+#include "PCGMegaPlantRemoveBranches.h"
+#include "PCGMegaPlantData.h"
+#include "Data/PCGBasePointData.h"
+#include "GeometryCollection/ManagedArrayCollection.h"
+#include "Helpers/PCGSettingsHelpers.h"
+#include "Implementations/MegaPlantRemoveBranches.h"
+
+FPCGElementPtr UPCGMegaPlantRemoveBranchesSettings::CreateElement() const
+{
+	return MakeShared<FPCGMegaPlantRemoveBranchesElement>();
+}
+
+bool FPCGMegaPlantRemoveBranchesElement::ExecuteInternal(FPCGContext* InContext) const
+{
+	TRACE_CPUPROFILER_EVENT_SCOPE(FPCGMegaPlantRemoveBranchesElement::Execute);
+
+	check(InContext);
+
+	const UPCGMegaPlantRemoveBranchesSettings* Settings = InContext->GetInputSettings<UPCGMegaPlantRemoveBranchesSettings>();
+	check(Settings);
+	
+	for (const FPCGTaggedData& Input : InContext->InputData.GetInputsByPin(PCGPinConstants::DefaultInputLabel))
+	{
+		if(const UPCGMegaPlantData* InputData = Cast<UPCGMegaPlantData>(Input.Data))
+		{
+			FManagedArrayCollection OutCollection = InputData->GetCollection();
+		
+			UPCGMegaPlantData* OutManagedArrayCollectionData = FPCGContext::NewObject_AnyThread<UPCGMegaPlantData>(InContext);
+			
+			FMegaPlantRemoveBranches::ApplyRemoveBranches(Settings->BranchRemoveBasis, Settings->Threshold, OutCollection);
+			
+			OutManagedArrayCollectionData->Initialize(MoveTemp(OutCollection));
+			InContext->OutputData.TaggedData.Emplace(OutManagedArrayCollectionData);
+		}
+		else
+		{
+			PCGLog::InputOutput::LogInvalidInputDataError(InContext);
+			return true;
+		}
+	}
+	
+	return true;
+}
diff --git a/Engine/Plugins/Experimental/ProceduralVegetationEditor/Source/MegaPlants/Private/Nodes/PCGMegaPlantRemoveBranches.h b/Engine/Plugins/Experimental/ProceduralVegetationEditor/Source/MegaPlants/Private/Nodes/PCGMegaPlantRemoveBranches.h
new file mode 100644
--- /dev/null
+++ b/Engine/Plugins/Experimental/ProceduralVegetationEditor/Source/MegaPlants/Private/Nodes/PCGMegaPlantRemoveBranches.h
@@ -0,0 +1,38 @@
+﻿// Copyright Epic Games, Inc. All Rights Reserved.
+#pragma once
+
+#include "CoreMinimal.h"
+#include "Nodes/PCGMegaPlantBaseSettings.h"
+#include "Implementations/MegaPlantRemoveBranches.h"
+#include "PCGMegaPlantRemoveBranches.generated.h"
+
+UCLASS(BlueprintType, ClassGroup = (Procedural))
+class UPCGMegaPlantRemoveBranchesSettings : public UPCGMegaPlantBaseSettings
+{
+	GENERATED_BODY()
+
+public:
+#if WITH_EDITOR
+	virtual FName GetDefaultNodeName() const override { return FName(TEXT("MegaPlantRemoveBranches")); }
+	virtual FText GetDefaultNodeTitle() const override { return NSLOCTEXT("PCGMegaPlantScale", "NodeTitle", "Remove Branches"); }
+	virtual FText GetNodeTooltipText() const override { return NSLOCTEXT("PCGMegaPlantScale", "NodeTooltip", "Remove branches on the basis of selected criteria"); }
+	virtual FLinearColor GetNodeTitleColor() const override { return FLinearColor::Yellow; }
+
+	virtual TArray<EMegaPlantRenderType> GetRenderType() const override { return TArray{ EMegaPlantRenderType::Skeleton }; }
+#endif
+
+protected:
+	virtual FPCGElementPtr CreateElement() const override;
+public:
+	UPROPERTY(EditAnywhere, Category = Settings)
+	ERemoveBranchesBasis BranchRemoveBasis = ERemoveBranchesBasis::Length;
+
+	UPROPERTY(EditAnywhere, Category = Settings, meta=(ClampMin = 0f, ClampMax = 10.0f, UIMin = 0f, UIMax = 1.0f))
+	float Threshold = 0;
+};
+
+class FPCGMegaPlantRemoveBranchesElement : public FPCGMegaPlantBaseElement
+{
+protected:
+	virtual bool ExecuteInternal(FPCGContext* InContext) const override;
+};
diff --git a/Engine/Plugins/Experimental/ProceduralVegetationEditor/Source/MegaPlants/Private/Nodes/PCGMegaPlantScale.cpp b/Engine/Plugins/Experimental/ProceduralVegetationEditor/Source/MegaPlants/Private/Nodes/PCGMegaPlantScale.cpp
new file mode 100644
--- /dev/null
+++ b/Engine/Plugins/Experimental/ProceduralVegetationEditor/Source/MegaPlants/Private/Nodes/PCGMegaPlantScale.cpp
@@ -0,0 +1,45 @@
+﻿// Copyright Epic Games, Inc. All Rights Reserved.
+
+#include "PCGMegaPlantScale.h"
+#include "PCGMegaPlantData.h"
+#include "Data/PCGBasePointData.h"
+#include "GeometryCollection/ManagedArrayCollection.h"
+#include "Helpers/PCGSettingsHelpers.h"
+#include "Implementations/MegaPlantScale.h"
+
+FPCGElementPtr UPCGMegaPlantScaleSettings::CreateElement() const
+{
+	return MakeShared<FPCGMegaPlantScaleElement>();
+}
+
+bool FPCGMegaPlantScaleElement::ExecuteInternal(FPCGContext* InContext) const
+{
+	TRACE_CPUPROFILER_EVENT_SCOPE(FPCGMegaPlantScaleElement::Execute);
+
+	check(InContext);
+
+	const UPCGMegaPlantScaleSettings* Settings = InContext->GetInputSettings<UPCGMegaPlantScaleSettings>();
+	check(Settings);
+	
+	for (const FPCGTaggedData& Input : InContext->InputData.GetInputsByPin(PCGPinConstants::DefaultInputLabel))
+	{
+		if(const UPCGMegaPlantData* InputData = Cast<UPCGMegaPlantData>(Input.Data))
+		{
+			FManagedArrayCollection OutCollection = InputData->GetCollection();
+		
+			UPCGMegaPlantData* OutManagedArrayCollectionData = FPCGContext::NewObject_AnyThread<UPCGMegaPlantData>(InContext);
+			
+			FMegaPlantScale::ApplyScale(Settings->Scale, OutCollection);
+			
+			OutManagedArrayCollectionData->Initialize(MoveTemp(OutCollection));
+			InContext->OutputData.TaggedData.Emplace(OutManagedArrayCollectionData);
+		}
+		else
+		{
+			PCGLog::InputOutput::LogInvalidInputDataError(InContext);
+			return true;
+		}
+	}
+	
+	return true;
+}
diff --git a/Engine/Plugins/Experimental/ProceduralVegetationEditor/Source/MegaPlants/Private/Nodes/PCGMegaPlantScale.h b/Engine/Plugins/Experimental/ProceduralVegetationEditor/Source/MegaPlants/Private/Nodes/PCGMegaPlantScale.h
new file mode 100644
--- /dev/null
+++ b/Engine/Plugins/Experimental/ProceduralVegetationEditor/Source/MegaPlants/Private/Nodes/PCGMegaPlantScale.h
@@ -0,0 +1,34 @@
+﻿// Copyright Epic Games, Inc. All Rights Reserved.
+#pragma once
+
+#include "CoreMinimal.h"
+#include "Nodes/PCGMegaPlantBaseSettings.h"
+#include "PCGMegaPlantScale.generated.h"
+
+UCLASS(BlueprintType, ClassGroup = (Procedural))
+class UPCGMegaPlantScaleSettings : public UPCGMegaPlantBaseSettings
+{
+	GENERATED_BODY()
+
+public:
+#if WITH_EDITOR
+	virtual FName GetDefaultNodeName() const override { return FName(TEXT("MegaPlantScale")); }
+	virtual FText GetDefaultNodeTitle() const override { return NSLOCTEXT("PCGMegaPlantScale", "NodeTitle", "Scale"); }
+	virtual FText GetNodeTooltipText() const override { return NSLOCTEXT("PCGMegaPlantScale", "NodeTooltip", "Scale the tree based on the settings"); }
+	virtual FLinearColor GetNodeTitleColor() const override { return FLinearColor::Yellow; }
+
+	virtual TArray<EMegaPlantRenderType> GetRenderType() const override { return TArray{ EMegaPlantRenderType::Skeleton }; }
+#endif
+
+protected:
+	virtual FPCGElementPtr CreateElement() const override;
+public:
+	UPROPERTY(EditAnywhere, Category = Settings, meta=(ClampMin = 0.01f, ClampMax = 100.0f, UIMin = 0.01f, UIMax = 10.0f))
+	float Scale = 1;
+};
+
+class FPCGMegaPlantScaleElement : public FPCGMegaPlantBaseElement
+{
+protected:
+	virtual bool ExecuteInternal(FPCGContext* InContext) const override;
+};
diff --git a/Engine/Plugins/Experimental/ProceduralVegetationEditor/Source/MegaPlants/Private/PCGMegaPlantData.cpp b/Engine/Plugins/Experimental/ProceduralVegetationEditor/Source/MegaPlants/Private/PCGMegaPlantData.cpp
new file mode 100644
--- /dev/null
+++ b/Engine/Plugins/Experimental/ProceduralVegetationEditor/Source/MegaPlants/Private/PCGMegaPlantData.cpp
@@ -0,0 +1,86 @@
+// Copyright Epic Games, Inc. All Rights Reserved.
+
+#include "PCGMegaPlantData.h"
+#include "PCGContext.h"
+#include "Data/PCGPointArrayData.h"
+#include "Data/PCGPointData.h"
+#include "Helpers/PCGHelpers.h"
+#include "Serialization/ArchiveCrc32.h"
+#include "UObject/Package.h"
+
+UPCGMegaPlantData::UPCGMegaPlantData(const FObjectInitializer& ObjectInitializer)
+	: UPCGSpatialData(ObjectInitializer)
+{
+	
+}
+
+void UPCGMegaPlantData::Initialize(FManagedArrayCollection&& InCollection, bool bCanTakeOwnership)
+{
+	Collection = InCollection;
+}
+
+void UPCGMegaPlantData::AddToCrc(FArchiveCrc32& Ar, bool bFullDataCrc) const
+{
+	Super::AddToCrc(Ar, bFullDataCrc);
+
+	if (bFullDataCrc)
+	{
+		// Implementation note: metadata not supported at this point.
+
+		FString ClassName = StaticClass()->GetPathName();
+
+		Ar << ClassName;
+		Ar << Collection;
+#if WITH_EDITOR
+		Ar << RenderType;
+#endif
+	}
+	else
+	{
+		AddUIDToCrc(Ar);
+	}
+}
+
+FBox UPCGMegaPlantData::GetBounds() const
+{
+	return CachedBounds;
+}
+
+bool UPCGMegaPlantData::SamplePoint(const FTransform& InTransform, const FBox& InBounds, FPCGPoint& OutPoint, UPCGMetadata* OutMetadata) const
+{
+	return false;
+}
+
+const UPCGBasePointData* UPCGMegaPlantData::ToBasePointData(FPCGContext* Context, const FBox& InBounds, TSubclassOf<UPCGBasePointData> PointDataClass) const
+{
+	TRACE_CPUPROFILER_EVENT_SCOPE(UPCGPrimitiveData::CreatePointData);
+
+	UPCGBasePointData* Data = FPCGContext::NewObject_AnyThread<UPCGBasePointData>(Context, GetTransientPackage(), PointDataClass);
+	return Data;
+}
+
+UPCGSpatialData* UPCGMegaPlantData::CopyInternal(FPCGContext* Context) const
+{
+	UPCGMegaPlantData* NewMegaPlantData = FPCGContext::NewObject_AnyThread<UPCGMegaPlantData>(Context);
+
+	FManagedArrayCollection NewCollection;
+	Collection.CopyTo(&NewCollection);
+	
+	NewMegaPlantData->Initialize(MoveTemp(NewCollection), /*bCanTakeOwnership=*/false);
+
+#if WITH_EDITOR
+	NewMegaPlantData->RenderType = RenderType;
+#endif
+	
+	return NewMegaPlantData;
+}
+
+const UPCGPointArrayData* UPCGMegaPlantData::ToPointArrayData(FPCGContext* Context, const FBox& InBounds) const
+{
+	return CastChecked<UPCGPointArrayData>(ToBasePointData(Context, InBounds, UPCGPointArrayData::StaticClass()));
+}
+const UPCGPointData* UPCGMegaPlantData::ToPointData(FPCGContext* Context, const FBox& InBounds) const
+{
+	return CastChecked<UPCGPointData>(ToBasePointData(Context, InBounds, UPCGPointData::StaticClass()));
+}
+
diff --git a/Engine/Plugins/Experimental/ProceduralVegetationEditor/Source/MegaPlants/Public/Facades/MegaPlantAttributesNames.h b/Engine/Plugins/Experimental/ProceduralVegetationEditor/Source/MegaPlants/Public/Facades/MegaPlantAttributesNames.h
new file mode 100644
--- /dev/null
+++ b/Engine/Plugins/Experimental/ProceduralVegetationEditor/Source/MegaPlants/Public/Facades/MegaPlantAttributesNames.h
@@ -0,0 +1,84 @@
+// Copyright Epic Games, Inc. All Rights Reserved.
+#pragma once
+
+#include "CoreMinimal.h"
+
+namespace MegaPlants::GroupNames
+{
+	//Group Names
+	inline static const FName BranchGroup = FName(TEXT("Primitives"));
+	inline static const FName PointGroup = FName(TEXT("Points"));
+	inline static const FName FoliageGroup = FName(TEXT("Foliage"));
+	inline static const FName FoliageNamesGroup = FName(TEXT("FoliageNames"));
+	inline static const FName DetailsGroup = FName(TEXT("Details"));
+	inline static const FName PlantProfilesGroup = FName(TEXT("PlantProfiles"));
+	inline static const FName BonesGroup = FName(TEXT("Bones"));
+	inline static const FName VerticesGroup = FName(TEXT("Vertices"));
+}
+
+namespace MegaPlants::AttributeNames
+{
+	//Branch Attributes
+	inline static const FName BranchParents = FName(TEXT("parents"));
+	inline static const FName BranchChildren = FName(TEXT("children"));
+	inline static const FName BranchPoints = FName(TEXT("Points"));
+	inline static const FName BranchNumber = FName(TEXT("branchNumber"));
+	inline static const FName BranchSourceBudNumber = FName(TEXT("branchSourceBudNumber"));
+	inline static const FName BranchFoliageIDs = FName(TEXT("FoliageID"));
+	inline static const FName BranchUVMaterial = FName(TEXT("branchUVMat"));
+	inline static const FName TrunkMaterialPath = FName(TEXT("TrunkMaterialPath"));
+	inline static const FName BranchHierarchyNumber = FName(TEXT("branchHierarchyNumber"));
+
+	//Point Attributes
+	inline static const FName PointPosition = FName(TEXT("Position"));
+	inline static const FName LengthFromRoot = FName(TEXT("lengthFromRoot"));
+	inline static const FName PointScaleGradient = FName(TEXT("LOD_totalPscaleGradient"));
+	inline static const FName PointScale = FName(TEXT("Scale"));
+	
+	inline static const FName LengthFromSeed = FName(TEXT("lengthFromSeed"));
+	inline static const FName BudLightDetected = FName(TEXT("budLightDetected"));
+	inline static const FName BudDevelopment = FName(TEXT("budDevelopment"));
+	inline static const FName TextureCoordV = FName(TEXT("uv_v"));
+	inline static const FName TextureCoordUOffset = FName(TEXT("uv_uOffset"));
+	
+	inline static const FName HullGradient = FName(TEXT("LOD_hullGradient"));
+	inline static const FName MainTrunkGradient = FName(TEXT("LOD_mainTrunkGradient"));
+	inline static const FName GroundGradient = FName(TEXT("LOD_groundGradient"));
+	inline static const FName BudNumber = FName(TEXT("budNumber"));
+	inline static const FName BudHormoneLevels = FName(TEXT("budHormoneLevels"));
+	inline static const FName BudDirection = FName(TEXT("budDirection"));
+	inline static const FName BranchGradient = FName(TEXT("branchGradient"));
+	inline static const FName PlantGradient = FName(TEXT("plantGradient"));
+	inline static const FName NjordPixelIndex = FName(TEXT("njord_pixelIdx"));
+
+	//Foliage Attributes
+	inline static const FName FoliageNameID = FName(TEXT("FoliageNameID"));
+	inline static const FName FoliageBranchID = FName(TEXT("FoliageBranchID"));
+	inline static const FName FoliagePivotPoint = FName(TEXT("FoliagePivotPoint"));
+	inline static const FName FoliageUPVector = FName(TEXT("FoliageUPVector"));
+	inline static const FName FoliageNormalVector = FName(TEXT("FoliageNormalVector"));
+	inline static const FName FoliageScale = FName(TEXT("FoliageScale"));
+	inline static const FName FoliageLengthFromRoot = FName(TEXT("LengthFromRoot"));
+	inline static const FName FoliageParentBoneID = FName(TEXT("FoliageParentBoneID"));
+
+	//FoliageNames Attributes
+	inline static const FName FoliageName = FName(TEXT("Name"));
+
+	//Details Attributes
+	inline static const FName LeafPhyllotaxy = FName(TEXT("phyllotaxyLeaf"));
+	inline static const FName FoliagePath = FName(TEXT("FoliagePath"));
+	inline static const FName Guid = FName(TEXT("Guid"));
+
+	//PlantProfiles Attributes
+	inline static const FName ProfilePoints = FName(TEXT("ProfilePoints"));
+
+	//Bone Attributes
+	inline static const FName BoneName = FName(TEXT("BoneName"));
+	inline static const FName BoneParentIndex = FName(TEXT("BoneParentIndex"));
+	inline static const FName BonePose = FName(TEXT("BonePose"));
+	inline static const FName BonePointIndex = FName(TEXT("BonePointIndex"));
+	inline static const FName BoneParentPointIndex = FName(TEXT("BoneParentPointIndex"));
+	inline static const FName BoneBranchIndex = FName(TEXT("BoneBranchIndex"));
+	inline static const FName NjordPixelId = FName(TEXT("NjordPixelId"));
+}
+
diff --git a/Engine/Plugins/Experimental/ProceduralVegetationEditor/Source/MegaPlants/Public/Facades/MegaPlantBoneFacade.h b/Engine/Plugins/Experimental/ProceduralVegetationEditor/Source/MegaPlants/Public/Facades/MegaPlantBoneFacade.h
new file mode 100644
--- /dev/null
+++ b/Engine/Plugins/Experimental/ProceduralVegetationEditor/Source/MegaPlants/Public/Facades/MegaPlantBoneFacade.h
@@ -0,0 +1,64 @@
+﻿// Copyright Epic Games, Inc. All Rights Reserved.
+
+#pragma once
+
+#include "CoreMinimal.h"
+#include "GeometryCollection/ManagedArrayAccessor.h"
+
+namespace MegaPlants::Facades
+{
+	/**
+	 * FBoneNode contains the bone data and linked point and branch ids
+	 */
+	struct MEGAPLANTS_API FBoneNode
+	{
+		FName BoneName = NAME_None;
+		int32 BoneIndex = INDEX_NONE;
+		int32 ParentBoneIndex = INDEX_NONE;
+		int32 PointIndex = INDEX_NONE;
+		int32 BranchIndex = INDEX_NONE;
+		float NjordPixelID = 0.0;
+
+		FTransform BoneTransform{};
+		FVector AbsolutePosition = FVector::ZeroVector;
+	};
+	
+	/**
+	 * FBoneFacade is used to access and manipulate the Bone data from Bones Group with in the MegaPlant's FManagedArrayCollection.
+	 * NjordPixel id here is stored in vertex Group and is used for bone assignment
+	 */
+	class MEGAPLANTS_API FBoneFacade
+	{
+	public:
+		FBoneFacade(FManagedArrayCollection& InCollection);
+		
+		FBoneFacade(const FManagedArrayCollection& InCollection);
+
+		static void DefineSchema(FManagedArrayCollection& InCollection);
+		
+		bool IsValid() const;
+
+		void GetNjordPixelIds(TArray<float>& VertexNjordIds) const;
+		
+		const TManagedArray<FString>& GetBoneNames() const;
+		
+		const TManagedArray<int32>& GetBoneParentIndices() const;
+		
+		const TManagedArray<FTransform>& GetBonePoses() const;
+
+		const TManagedArray<float>& GetNjordPixelIds() const;
+		
+		TManagedArray<float>& ModifyNjordPixelIds();
+
+		void CreateBoneData(TArray<FBoneNode>& Bones) const;
+		
+	private:
+		TManagedArrayAccessor<FString> BoneName;
+		TManagedArrayAccessor<int32> BoneParentIndex;
+		TManagedArrayAccessor<FTransform> BonePose;
+		TManagedArrayAccessor<int32> BonePointIndex;
+		TManagedArrayAccessor<int32> BoneParentPointIndex;
+		TManagedArrayAccessor<int32> BoneBranchIndex;
+		TManagedArrayAccessor<float> NjordPixelIds;
+	};
+}
diff --git a/Engine/Plugins/Experimental/ProceduralVegetationEditor/Source/MegaPlants/Public/Facades/MegaPlantBranchFacade.h b/Engine/Plugins/Experimental/ProceduralVegetationEditor/Source/MegaPlants/Public/Facades/MegaPlantBranchFacade.h
new file mode 100644
--- /dev/null
+++ b/Engine/Plugins/Experimental/ProceduralVegetationEditor/Source/MegaPlants/Public/Facades/MegaPlantBranchFacade.h
@@ -0,0 +1,81 @@
+﻿// Copyright Epic Games, Inc. All Rights Reserved.
+
+#pragma once
+
+#include "CoreMinimal.h"
+#include "IShrinkableFacade.h"
+#include "GeometryCollection/ManagedArrayAccessor.h"
+
+namespace MegaPlants::Facades
+{
+	/**
+	 * FBranchFacade is used to access and manipulate the data from Branch Group with in the MegaPlant's FManagedArrayCollection
+	 * A branch is made up of multiple points. Each Branch will know about its parents and children, along with other attributes
+	 * Only add the frequently used branch attributes and their access to this facade, for the specific access write a new facade
+	 */
+	class MEGAPLANTS_API FBranchFacade final : public IShrinkable
+	{
+	public:
+		FBranchFacade(FManagedArrayCollection& InCollection);
+		FBranchFacade(const FManagedArrayCollection& InCollection);
+
+		bool IsConst() const { return Collection == nullptr; }
+
+		bool IsValid() const;
+
+		virtual int32 GetElementCount() const override;
+
+		const TArray<int32>& GetPoints(const int32 Index) const;
+
+		void SetPoints(const int32 Index, const TArray<int32>& InPoints);
+
+		const TArray<int32>& GetChildren(const int32 Index) const;
+
+		void SetChildren(const int32 Index, const TArray<int32>& InChildren);
+
+		const TArray<int32>& GetParents(const int32 Index) const;
+
+		void SetParents(const int32 Index, const TArray<int32>& InParents);
+
+		int32 GetBranchNumber(const int32 Index) const;
+
+		int32 GetBranchSourceBudNumber(const int32 Index) const;
+
+		int32 GetTrunkIndex() const;
+
+		TArray<int32> GetParentBranchIndices(int BranchIndex) const;
+
+		int32 GetGenerationNumber(const int32 Index) const;
+
+		const TManagedArray<int32>& GetBranchNumbers() const;
+
+		int32 GetBranchUVMaterial(const int32 Index) const;
+
+		void SetBranchUVMaterial(const int32 Index, const int32 InMaterial);
+		
+		int32 GetBranchHierarchyNumber(int32 Index) const;
+
+		void SetTrunkMaterialPath(const FString& InPath);
+
+		FString GetTrunkMaterialPath() const;
+
+		virtual void CopyEntry(const int32 FromIndex, const int32 ToIndex) override;
+
+		virtual void RemoveEntries(const int32 NumEntries, const int32 StartIndex) override;
+
+		void GetSortedBranchIndicesByHierarchy(TArray<int32>& OutSortedIndices) const;
+
+	private:
+		FManagedArrayCollection* Collection = nullptr;
+
+		TManagedArrayAccessor<TArray<int32>> Parents;
+		TManagedArrayAccessor<TArray<int32>> Children;
+		TManagedArrayAccessor<TArray<int32>> BranchPoints;
+		TManagedArrayAccessor<int32> BranchNumbers;
+		TManagedArrayAccessor<int32> BranchSourceBudNumber;
+		TManagedArrayAccessor<TArray<int32>> BranchFoliageIDs;
+		TManagedArrayAccessor<int32> BranchUVMaterial;
+		TManagedArrayAccessor<int32> BranchHierarchyNumber;
+		TManagedArrayAccessor<FString> TrunkMaterialPathAttribute;
+	};
+}
diff --git a/Engine/Plugins/Experimental/ProceduralVegetationEditor/Source/MegaPlants/Public/Facades/MegaPlantFoliageFacade.h b/Engine/Plugins/Experimental/ProceduralVegetationEditor/Source/MegaPlants/Public/Facades/MegaPlantFoliageFacade.h
new file mode 100644
--- /dev/null
+++ b/Engine/Plugins/Experimental/ProceduralVegetationEditor/Source/MegaPlants/Public/Facades/MegaPlantFoliageFacade.h
@@ -0,0 +1,105 @@
+﻿// Copyright Epic Games, Inc. All Rights Reserved.
+
+#pragma once
+
+#include "CoreMinimal.h"
+#include "IShrinkableFacade.h"
+#include "GeometryCollection/ManagedArrayAccessor.h"
+
+namespace MegaPlants::Facades
+{
+	struct MEGAPLANTS_API FFoliageEntryData
+	{
+		int32 NameId;
+		int32 BranchId;
+		FVector3f PivotPoint;
+		FVector3f UpVector;
+		FVector3f NormalVector;
+		float Scale;
+		float LengthFromRoot;
+		int32 ParentBoneID = INDEX_NONE;
+	};
+
+	/**
+	 * FFoliageFacade is used to access and manipulate Foliage related data for a MegaPlant's FManagedArrayCollection
+	 * It presents access to foliage Mesh names, their respective Branch ids, and instancing transform data (Pivot point, Up Vector, and Scale)
+	 * It lays out data in two Groups: FoliageNames and Foliage where Foliage references FoliageNames and Branch/Primitives with Ids
+	 * It also adds Foliage entry Ids to the Branch/Primitives group. 
+	 */
+	class MEGAPLANTS_API FFoliageFacade final : public IShrinkable
+	{
+	public:
+		FFoliageFacade(FManagedArrayCollection& InCollection, const int32 InitialSize = 0);
+		FFoliageFacade(const FManagedArrayCollection& InCollection);
+
+		bool IsConst() const { return Collection == nullptr; }
+
+		bool IsValid() const;
+
+		int32 NumFoliageEntries() const { return NameIdsAttribute.Num(); };
+
+		FFoliageEntryData GetFoliageEntry(const int32 Index) const;
+
+		const float& GetLengthFromRoot(const int32 Index) const;
+
+		void SetLengthFromRoot(const int32 Index, const float Input);
+
+		int32 GetParentBoneID(int32 Index) const;
+
+		void SetParentBoneID(int32 Index, int32 Input);
+
+		const TArray<int32>& GetFoliageEntryIdsForBranch(int32 Index) const;
+
+		int32 AddFoliageEntry(const FFoliageEntryData& InputData);
+
+		void SetFoliageEntry(const int32 Index, const FFoliageEntryData& InputData);
+
+		void SetFoliageIdsArray(const int32 Index, const TArray<int32>& InputIds);
+
+		void SetFoliageBranchId(const int32 Index, const int32 InputId);
+
+		int32 NumFoliageNames() const { return NamesAttribute.Num(); };
+
+		FString GetFoliageName(const int32 Index) const;
+
+		const FVector3f& GetPivotPoint(const int32 Index) const;
+
+		void SetPivotPoint(const int32 Index, const FVector3f& Input);
+
+		void SetScale(const int32 Index, const float& Input);
+
+		void SetFoliageName(const int32 Index, const FString& Input);
+
+		void SetFoliageNames(const TArray<FString>& InputNames);
+
+		void SetFoliagePath(const FString InPath);
+
+		FString GetFoliagePath() const;
+
+		virtual int32 GetElementCount() const override;
+
+		virtual void CopyEntry(const int32 FromIndex, const int32 ToIndex) override;
+
+		virtual void RemoveEntries(const int32 NumEntries, const int32 StartIndex) override;
+		
+		FTransform GetFoliageTransform(int32 Id) const;
+
+	protected:
+		void DefineSchema(const int32 InitialSize = 0);
+
+		const FManagedArrayCollection& ConstCollection;
+		FManagedArrayCollection* Collection = nullptr;
+
+		TManagedArrayAccessor<FString> NamesAttribute;
+		TManagedArrayAccessor<int32> NameIdsAttribute;
+		TManagedArrayAccessor<int32> BranchIdsAttribute;
+		TManagedArrayAccessor<FVector3f> PivotPointsAttribute;
+		TManagedArrayAccessor<FVector3f> UpVectorsAttribute;
+		TManagedArrayAccessor<FVector3f> NormalVectorsAttribute;
+		TManagedArrayAccessor<float> ScalesAttribute;
+		TManagedArrayAccessor<float> LengthFromRootAttribute;
+		TManagedArrayAccessor<int32> ParentBoneIdsAttribute;
+		TManagedArrayAccessor<TArray<int32>> FoliageIdsAttribute;
+		TManagedArrayAccessor<FString> FoliagePathAttribute;
+	};
+}
diff --git a/Engine/Plugins/Experimental/ProceduralVegetationEditor/Source/MegaPlants/Public/Facades/MegaPlantMetaInfoFacade.h b/Engine/Plugins/Experimental/ProceduralVegetationEditor/Source/MegaPlants/Public/Facades/MegaPlantMetaInfoFacade.h
new file mode 100644
--- /dev/null
+++ b/Engine/Plugins/Experimental/ProceduralVegetationEditor/Source/MegaPlants/Public/Facades/MegaPlantMetaInfoFacade.h
@@ -0,0 +1,30 @@
+﻿// Copyright Epic Games, Inc. All Rights Reserved.
+
+#pragma once
+
+#include "CoreMinimal.h"
+#include "GeometryCollection/ManagedArrayAccessor.h"
+
+namespace MegaPlants::Facades
+{
+	class MEGAPLANTS_API FMetaInfoFacade
+	{
+	public:
+		FMetaInfoFacade(FManagedArrayCollection& InCollection, const int32 InitialSize = 0);
+		FMetaInfoFacade(const FManagedArrayCollection& InCollection);
+
+		bool IsConst() const { return Collection == nullptr; }
+		bool IsValid() const;
+
+		void CreateGuid(const FString& InPath);
+		FGuid GetGuid() const;
+
+	protected:
+		void DefineSchema(const int32 InitialSize = 0);
+
+		const FManagedArrayCollection& ConstCollection;
+		FManagedArrayCollection* Collection = nullptr;
+		
+		TManagedArrayAccessor<FGuid> GuidAttribute;
+	};
+}
diff --git a/Engine/Plugins/Experimental/ProceduralVegetationEditor/Source/MegaPlants/Public/Facades/MegaPlantPointFacade.h b/Engine/Plugins/Experimental/ProceduralVegetationEditor/Source/MegaPlants/Public/Facades/MegaPlantPointFacade.h
new file mode 100644
--- /dev/null
+++ b/Engine/Plugins/Experimental/ProceduralVegetationEditor/Source/MegaPlants/Public/Facades/MegaPlantPointFacade.h
@@ -0,0 +1,111 @@
+﻿// Copyright Epic Games, Inc. All Rights Reserved.
+
+#pragma once
+
+#include "CoreMinimal.h"
+#include "IShrinkableFacade.h"
+#include "GeometryCollection/ManagedArrayAccessor.h"
+
+namespace MegaPlants::Facades
+{
+	/**
+	 * FPointFacade is used to access and manipulate the Point Group data from the MegaPlant's FManagedArrayCollection
+	 * Only add the frequently used Point attributes and their access to this facade, for the specific access write a new facade
+	 */
+	class MEGAPLANTS_API FPointFacade final : public IShrinkable
+	{
+	public:
+		FPointFacade(FManagedArrayCollection& InCollection);
+		FPointFacade(const FManagedArrayCollection& InCollection);
+
+		bool IsValid() const;
+
+		virtual int32 GetElementCount() const override;
+
+		const FVector3f& GetPosition(const int32 Index) const;
+
+		const float& GetLengthFromRoot(const int32 Index) const;
+
+		const float& GetLengthFromSeed(const int32 Index) const;
+
+		const float& GetPointScaleGradient(const int32 Index) const;
+
+		const float& GetHullGradient(const int32 Index) const;
+
+		const float& GetMainTrunkGradient(const int32 Index) const;
+
+		const float& GetGroundGradient(const int32 Index) const;
+		
+		const float& GetPointScale(const int32 Index) const;
+
+		const int32 GetBudNumber(const int32 Index) const;
+
+		const TArray<float>& GetBudLightDetected(const int32 Index) const;
+
+		const TArray<int>& GetBudDevelopment(const int32 Index) const;
+		
+		const TArray<FVector3f>& GetBudDirection(const int32 Index) const;
+
+		const TArray<float>& GetBudHormoneLevels(const int32 Index) const;
+
+		void SetBudHormoneLevels(const int32 Index, const TArray<float>& InBudHormoneLevels);
+
+		const float& GetBranchGradient(const int32 Index) const;
+
+		void SetBranchGradient(const int32 Index, const float InBranchGradient);
+
+		const float& GetPlantGradient(const int32 Index) const;
+
+		const float& GetNjordPixelIndex(const int32 Index) const;
+
+		void SetPlantGradient(const int32 Index, const float InPlantGradient);
+
+		const float& GetTextureCoordV(int32 Index) const;
+		
+		void SeTextureCoordV(int32 Index, float InTextureCoordV);
+		
+		const float& GetTextureCoordUOffset(int32 Index) const;
+		
+		void SetTextureCoordUOffset(int32 Index, float InTextureCoordUOffset);
+
+		TManagedArray<FVector3f>& ModifyPositions();
+		
+		TManagedArray<float>& ModifyPointScales();
+
+		TManagedArray<float>& ModifyLengthFromRoots();
+		
+		TManagedArray<float>& ModifyLengthFromSeeds();
+		
+		TManagedArray<float>& ModifyNjordPixelIDs();
+
+		const TManagedArray<FVector3f>& GetPositions() const;
+		
+		const TManagedArray<float>& GetPointScales() const;
+
+		const TManagedArray<float>& GetLengthFromRootsArray() const;
+
+		virtual void CopyEntry(const int32 FromIndex, const int32 ToIndex) override;
+
+		virtual void RemoveEntries(const int32 NumEntries, const int32 StartIndex) override;
+
+	private:
+		TManagedArrayAccessor<FVector3f> Positions;
+		TManagedArrayAccessor<float> LengthFromRoot;
+		TManagedArrayAccessor<float> PointScaleGradient;
+		TManagedArrayAccessor<float> HullGradient;
+		TManagedArrayAccessor<float> MainTrunkGradient;
+		TManagedArrayAccessor<float> GroundGradient;
+		TManagedArrayAccessor<float> PointScale;
+		TManagedArrayAccessor<TArray<FVector3f>> BudDirections;
+		TManagedArrayAccessor<float> BranchGradients;
+		TManagedArrayAccessor<TArray<float>> BudHormoneLevels;
+		TManagedArrayAccessor<float> PlantGradients;
+		TManagedArrayAccessor<float> LengthFromSeed;
+		TManagedArrayAccessor<int32> BudNumber;
+		TManagedArrayAccessor<float> NjordPixelIndex;
+		TManagedArrayAccessor<TArray<float>> BudLightDetected;
+		TManagedArrayAccessor<TArray<int>> BudDevelopment;
+		TManagedArrayAccessor<float> TextureCoordV;
+		TManagedArrayAccessor<float> TextureCoordUOffset;
+	};
+}
diff --git a/Engine/Plugins/Experimental/ProceduralVegetationEditor/Source/MegaPlants/Public/Facades/MegaPlantProfileFacade.h b/Engine/Plugins/Experimental/ProceduralVegetationEditor/Source/MegaPlants/Public/Facades/MegaPlantProfileFacade.h
new file mode 100644
--- /dev/null
+++ b/Engine/Plugins/Experimental/ProceduralVegetationEditor/Source/MegaPlants/Public/Facades/MegaPlantProfileFacade.h
@@ -0,0 +1,37 @@
+﻿// Copyright Epic Games, Inc. All Rights Reserved.
+
+#pragma once
+
+#include "CoreMinimal.h"
+#include "GeometryCollection/ManagedArrayAccessor.h"
+
+namespace MegaPlants::Facades
+{
+	/**
+	 * FTrunkProfileFacade is used to access and manipulate Trunk profile data for a MegaPlant's FManagedArrayCollection
+	 */
+	class MEGAPLANTS_API FPlantProfileFacade
+	{
+	public:
+		FPlantProfileFacade(FManagedArrayCollection& InCollection);
+		FPlantProfileFacade(const FManagedArrayCollection& InCollection);
+
+		bool IsConst() const { return Collection == nullptr; }
+
+		bool IsValid() const;
+
+		int32 NumProfileEntries() const { return ProfilePointsAttribute.Num(); };
+
+		int32 AddProfileEntry(const TArray<float>& InProfilePoints);
+
+		const TArray<float>& GetProfilePoints(const int32 Index) const;
+
+	protected:
+		void DefineSchema();
+
+		const FManagedArrayCollection& ConstCollection;
+		FManagedArrayCollection* Collection = nullptr;
+
+		TManagedArrayAccessor<TArray<float>> ProfilePointsAttribute;
+	};
+}
diff --git a/Engine/Plugins/Experimental/ProceduralVegetationEditor/Source/MegaPlants/Public/Facades/MegaPlantRenderingFacade.h b/Engine/Plugins/Experimental/ProceduralVegetationEditor/Source/MegaPlants/Public/Facades/MegaPlantRenderingFacade.h
new file mode 100644
--- /dev/null
+++ b/Engine/Plugins/Experimental/ProceduralVegetationEditor/Source/MegaPlants/Public/Facades/MegaPlantRenderingFacade.h
@@ -0,0 +1,31 @@
+﻿// Copyright Epic Games, Inc. All Rights Reserved.
+
+#pragma once
+#include "GeometryCollection/Facades/CollectionRenderingFacade.h"
+#include "MegaPlantFoliageFacade.h"
+
+namespace MegaPlants::Facades
+{
+	/**
+	 * FMegaPlantRenderingFacade is used to build a collection that contains the render data required to render Mega Plants
+	 * Add functions to access and modify render data for mega plants
+	 */
+	class MEGAPLANTS_API FMegaPlantRenderingFacade : public GeometryCollection::Facades::FRenderingFacade
+	{
+	public:
+		FMegaPlantRenderingFacade(FManagedArrayCollection& InSelf);
+		FMegaPlantRenderingFacade(const FManagedArrayCollection& InSelf);
+
+		void DefineSchema();
+
+		const FFoliageFacade& GetFoliageFacade() const;
+
+		void CopyFrom(const FManagedArrayCollection& Other);
+
+	protected:
+		const FManagedArrayCollection& ConstMegaPlantCollection;
+		FManagedArrayCollection* MegaPlantCollection = nullptr;
+
+		FFoliageFacade FoliageFacade;
+	};
+}
diff --git a/Engine/Plugins/Experimental/ProceduralVegetationEditor/Source/MegaPlants/Public/Helpers/MegaPlantExportHelper.h b/Engine/Plugins/Experimental/ProceduralVegetationEditor/Source/MegaPlants/Public/Helpers/MegaPlantExportHelper.h
new file mode 100644
--- /dev/null
+++ b/Engine/Plugins/Experimental/ProceduralVegetationEditor/Source/MegaPlants/Public/Helpers/MegaPlantExportHelper.h
@@ -0,0 +1,138 @@
+﻿// Copyright Epic Games, Inc. All Rights Reserved.
+#pragma once
+
+#include "UDynamicMesh.h"
+#include "NaniteAssemblyDataBuilder.h"
+#include "Engine/EngineTypes.h"
+#include "GeometryCollection/ManagedArrayCollection.h"
+#include "Materials/MaterialInterface.h"
+#include "Facades/MegaPlantBoneFacade.h"
+#include "MegaPlantExportHelper.generated.h"
+
+class UInstancedStaticMeshComponent;
+class UStaticMesh;
+class USkeletalMesh;
+class UMegaPlant;
+struct FGeometryScriptCopyMeshToAssetOptions;
+
+UENUM()
+enum class EMegaPlantExportMeshType: uint8
+{
+	StaticMesh,
+	SkeletalMesh,
+};
+
+USTRUCT()
+struct MEGAPLANTS_API FMegaPlantExportSettings
+{
+	GENERATED_BODY()
+
+	UPROPERTY()
+	bool bShouldExport = true;
+
+	UPROPERTY(EditAnywhere, Category = "Export", Meta = (ContentDir, DisplayName = "Content Browser Folder"))
+	FDirectoryPath ContentBrowserFolder;
+
+	UPROPERTY(EditAnywhere, Category = "Export", Meta = (DisplayName = "Mesh Name"))
+	FName MeshName;
+
+	UPROPERTY(EditAnywhere, Category = "Export", Meta = (DisplayName = "Export Mesh Type"))
+	EMegaPlantExportMeshType ExportMeshType = EMegaPlantExportMeshType::StaticMesh;
+
+	UPROPERTY(EditAnywhere, Category = "Export", Meta = (DisplayName = "Create Nanite Assemblies"))
+	bool bCreateNaniteAssemblies = true;
+	
+	UPROPERTY(EditAnywhere, Category = "Export", Meta = (DisplayName = "Nanite Shape Preservation Method"))
+	ENaniteShapePreservation NaniteShapePreservation = ENaniteShapePreservation::Voxelize;
+
+	void Initialize(const FString& InAssetPath, const FString& InName)
+	{
+		FString DefaultDirectory = FPaths::GetPath(InAssetPath);
+		ContentBrowserFolder.Path = DefaultDirectory;
+		
+		MeshName = FName(InName);
+	}
+
+	FString GetOutputObjectPath() const
+	{
+		const FString AssetName = MeshName.ToString();
+		const FString OutputPath = ContentBrowserFolder.Path / AssetName + '.' + AssetName;
+		return OutputPath;
+	}
+};
+
+namespace MegaPlants::Export::Internal
+{
+	typedef TFunction<void(const FString& MeshName, int32 InstanceID, int32 VertexStart, int32 VertexCount)> FMeshInstanceCombined;
+	
+	TObjectPtr<UDynamicMesh> CollectionToDynamicMesh(const FManagedArrayCollection& Collection);
+
+	FGeometryScriptCopyMeshToAssetOptions GetCopyMeshToAssetOptions(
+		TArray<TObjectPtr<UMaterialInterface>> InMaterials
+	);
+	
+#if WITH_EDITORONLY_DATA
+	FGeometryScriptCopyMeshToAssetOptions GetCopyMeshToAssetOptions(
+		TArray<TObjectPtr<UMaterialInterface>> InMaterials,
+		const FMeshNaniteSettings& InNaniteSettings,
+		const ENaniteShapePreservation InShapePreservation
+	);
+#endif
+	
+	void CombineMeshInstancesToDynamicMesh(
+		const FManagedArrayCollection& Collection,
+		UDynamicMesh* DynamicMesh,
+		TArray<TObjectPtr<UMaterialInterface>>& Materials,
+		const FMeshInstanceCombined& OnMeshInstanceCombined
+	);
+
+#if WITH_EDITOR
+	void BuildNaniteAssemblyData(
+		const FManagedArrayCollection& Collection,
+		UStaticMesh* StaticMesh
+	);
+
+	void AddNodeToBuilder(FNaniteAssemblyDataBuilder& AssemblyBuilder,
+		const FManagedArrayCollection& Collection,
+		const TMap<FString, int32>& InMeshNamePartMap,
+		const TMap<FString, TArray<TObjectPtr<UMaterialInterface>>>& InMeshMaterialsMap);
+	
+#endif
+
+	void ExportCollectionToStaticMesh(
+		TObjectPtr<UStaticMesh> ExportMesh,
+		const FManagedArrayCollection& Collection,
+		bool bBuildNaniteAssemblies,
+		ENaniteShapePreservation ShapePreservation
+	);
+
+	void ExportCollectionToSkeletalMesh(
+		TObjectPtr<USkeletalMesh> ExportMesh,
+		const FManagedArrayCollection& Collection,
+		ENaniteShapePreservation ShapePreservation,
+		bool bBuildNaniteAssemblies
+	);
+
+	void AttachMegaPlantLink(UObject* InExportedMesh, const TObjectPtr<UMegaPlant>& InMegaPlant);
+	
+	void AssignBoneIDsToFoliage(const TArray<Facades::FBoneNode>& Bones, FManagedArrayCollection& Collection);
+
+	void RemoveUnwantedSkinWeights(TObjectPtr<UDynamicMesh> DynamicMesh, const FName ProfileName, const TArray<int32>& VertexBoneIDs, const TArray<float>& NjordIDs,
+		int32 FoliageVertexStart, const TArray<Facades::FBoneNode>& BoneNodes);
+
+	void ConvertToDefaultSkeletalMesh(USkeletalMesh* SkeletalMesh,
+		FDynamicMesh3* Mesh,
+		TArray<TObjectPtr<UMaterialInterface>> Materials
+		);
+	
+}
+
+namespace MegaPlants::Export
+{
+	void MEGAPLANTS_API ExportCollectionAsMesh(
+		const TObjectPtr<UMegaPlant> InMegaPlant,
+		const FManagedArrayCollection& Collection,
+		const FManagedArrayCollection& FoliageCollection,
+		const FMegaPlantExportSettings& Settings
+	);
+}
diff --git a/Engine/Plugins/Experimental/ProceduralVegetationEditor/Source/MegaPlants/Public/Helpers/MegaPlantJSONHelper.h b/Engine/Plugins/Experimental/ProceduralVegetationEditor/Source/MegaPlants/Public/Helpers/MegaPlantJSONHelper.h
new file mode 100644
--- /dev/null
+++ b/Engine/Plugins/Experimental/ProceduralVegetationEditor/Source/MegaPlants/Public/Helpers/MegaPlantJSONHelper.h
@@ -0,0 +1,485 @@
+// Copyright Epic Games, Inc. All Rights Reserved.
+#pragma once
+
+#include "CoreMinimal.h"
+#include "Chaos/Deformable/MuscleActivationConstraints.h"
+#include "GeometryCollection/ManagedArrayCollection.h"
+#include "Facades/MegaPlantAttributesNames.h"
+#include "Dom/JsonObject.h"
+#include "Dom/JsonValue.h"
+#include "Facades/MegaPlantMetaInfoFacade.h"
+#include "Facades/MegaPlantFoliageFacade.h"
+#include "Facades/MegaPlantProfileFacade.h"
+#include "Misc/FileHelper.h"
+#include "Serialization/JsonSerializer.h"
+#include "Misc/Paths.h"
+#include "Misc/PackageName.h"
+
+namespace MegaPlants
+{
+	static void FillAttributes(FManagedArrayCollection& Collection, const FName Group, const TSharedPtr<FJsonObject>& AttributesObject)
+	{
+		for (const auto& [Key, Attribute] : AttributesObject->Values)
+		{
+			const TSharedPtr<FJsonObject>& AttributeObject = Attribute->AsObject();
+
+			const FString Type = AttributeObject->GetStringField(TEXT("type"));
+			const int32 ElementCount = AttributeObject->GetNumberField(TEXT("size"));
+			const bool IsArray = AttributeObject->GetBoolField(TEXT("isArray"));
+			const TArray<TSharedPtr<FJsonValue>>& Values = AttributeObject->GetArrayField(TEXT("values"));
+
+			if (IsArray)
+			{
+				if (ElementCount == 1)
+				{
+					if (Type == "int")
+					{
+						TManagedArray<TArray<int32>>& Attr = Collection.AddAttribute<TArray<int32>>(*Key, Group);
+						for (int32 Index = 0; Index < Values.Num(); ++Index)
+						{
+							TArray<int32>& CurrentAttrib = Attr[Index];
+							const TArray<TSharedPtr<FJsonValue>>& AttribValues = Values[Index]->AsArray();
+							for (const TSharedPtr<FJsonValue>& AttribValue : AttribValues)
+							{
+								CurrentAttrib.Add(AttribValue->AsNumber());
+							}
+						}
+					}
+					else if (Type == "float")
+					{
+						TManagedArray<TArray<float>>& Attr = Collection.AddAttribute<TArray<float>>(*Key, Group);
+						for (int32 Index = 0; Index < Values.Num(); ++Index)
+						{
+							TArray<float>& CurrentAttrib = Attr[Index];
+							const TArray<TSharedPtr<FJsonValue>>& AttribValues = Values[Index]->AsArray();
+							for (const TSharedPtr<FJsonValue>& AttribValue : AttribValues)
+							{
+								CurrentAttrib.Add(AttribValue->AsNumber());
+							}
+						}
+					}
+					else if (Type == "string")
+					{
+						// TODO: Figure out how to manage string arrays
+					}
+				}
+				else if (ElementCount == 3)
+				{
+					if (Type == "int")
+					{
+						TManagedArray<TArray<FIntVector>>& Attr = Collection.AddAttribute<TArray<FIntVector>>(*Key, Group);
+						for (int32 Index = 0; Index < Values.Num(); ++Index)
+						{
+							TArray<FIntVector>& CurrentAttrib = Attr[Index];
+							const TArray<TSharedPtr<FJsonValue>>& AttribValues = Values[Index]->AsArray();
+							for (int ElemIndex = 0; ElemIndex < AttribValues.Num(); ElemIndex += ElementCount)
+							{
+								CurrentAttrib.Emplace(AttribValues[ElemIndex]->AsNumber(), AttribValues[ElemIndex + 2]->AsNumber(),
+									AttribValues[ElemIndex + 1]->AsNumber());
+							}
+						}
+					}
+					else if (Type == "float")
+					{
+						TManagedArray<TArray<FVector3f>>& Attr = Collection.AddAttribute<TArray<FVector3f>>(*Key, Group);
+						for (int32 Index = 0; Index < Values.Num(); ++Index)
+						{
+							TArray<FVector3f>& CurrentAttrib = Attr[Index];
+							const TArray<TSharedPtr<FJsonValue>>& AttribValues = Values[Index]->AsArray();
+							for (int ElemIndex = 0; ElemIndex < AttribValues.Num(); ElemIndex += ElementCount)
+							{
+								CurrentAttrib.Emplace(AttribValues[ElemIndex]->AsNumber(), AttribValues[ElemIndex + 2]->AsNumber(),
+									AttribValues[ElemIndex + 1]->AsNumber());
+							}
+						}
+					}
+				}
+			}
+			else
+			{
+				if (ElementCount == 1)
+				{
+					if (Type == "int")
+					{
+						TManagedArray<int32>& Attr = Collection.AddAttribute<int32>(*Key, Group);
+						for (int32 Index = 0; Index < Values.Num(); ++Index)
+						{
+							Attr[Index] = Values[Index]->AsNumber();
+						}
+					}
+					else if (Type == "float")
+					{
+						TManagedArray<float>& Attr = Collection.AddAttribute<float>(*Key, Group);
+						for (int32 Index = 0; Index < Values.Num(); ++Index)
+						{
+							Attr[Index] = Values[Index]->AsNumber();
+						}
+					}
+					else if (Type == "string")
+					{
+						TManagedArray<FString>& Attr = Collection.AddAttribute<FString>(*Key, Group);
+						for (int32 Index = 0; Index < Values.Num(); ++Index)
+						{
+							Attr[Index] = Values[Index]->AsString();
+						}
+					}
+				}
+				else if (ElementCount == 3)
+				{
+					if (Type == "int")
+					{
+						TManagedArray<FIntVector>& Attr = Collection.AddAttribute<FIntVector>(*Key, Group);
+						for (int32 Index = 0; Index < Values.Num(); ++Index)
+						{
+							const TArray<TSharedPtr<FJsonValue>>& VectorAttribute = Values[Index]->AsArray();
+							Attr[Index] = FIntVector(VectorAttribute[0]->AsNumber(), VectorAttribute[2]->AsNumber(), VectorAttribute[1]->AsNumber());
+						}
+					}
+					else if (Type == "float")
+					{
+						TManagedArray<FVector3f>& Attr = Collection.AddAttribute<FVector3f>(*Key, Group);
+						for (int32 Index = 0; Index < Values.Num(); ++Index)
+						{
+							const TArray<TSharedPtr<FJsonValue>>& VectorAttribute = Values[Index]->AsArray();
+							Attr[Index] = FVector3f(VectorAttribute[0]->AsNumber(), VectorAttribute[2]->AsNumber(), VectorAttribute[1]->AsNumber());
+						}
+					}
+				}
+			}
+		}
+	}
+
+	static void FillDetailsAttributes(FManagedArrayCollection& Collection, const FName Group, const TSharedPtr<FJsonObject>& AttributesObject)
+	{
+		for (const auto& [Key, Attribute] : AttributesObject->Values)
+		{
+			const TSharedPtr<FJsonObject>& AttributeObject = Attribute->AsObject();
+
+			const FString Type = AttributeObject->GetStringField(TEXT("type"));
+			const bool IsArray = AttributeObject->GetBoolField(TEXT("isArray"));
+
+			if (IsArray)
+			{
+				const TArray<TSharedPtr<FJsonValue>>& Values = AttributeObject->GetArrayField(TEXT("value"));
+				if (Type == "int")
+				{
+					TManagedArray<TArray<int32>>& Attr = Collection.AddAttribute<TArray<int32>>(*Key, Group);
+					TArray<int32>& CurrentAttrib = Attr[0];
+					for (const TSharedPtr<FJsonValue>& AttribValue : Values)
+					{
+						CurrentAttrib.Add(AttribValue->AsNumber());
+					}
+				}
+				else if (Type == "float")
+				{
+					TManagedArray<TArray<float>>& Attr = Collection.AddAttribute<TArray<float>>(*Key, Group);
+					TArray<float>& CurrentAttrib = Attr[0];
+					for (const TSharedPtr<FJsonValue>& AttribValue : Values)
+					{
+						CurrentAttrib.Add(AttribValue->AsNumber());
+					}
+				}
+				else if (Type == "string")
+				{
+					// TODO: Figure out how to manage string arrays
+				}
+			}
+			else
+			{
+				if (Type == "int")
+				{
+					TManagedArray<int32>& Attr = Collection.AddAttribute<int32>(*Key, Group);
+					Attr[0] = AttributeObject->GetNumberField(TEXT("value"));
+				}
+				else if (Type == "float")
+				{
+					TManagedArray<float>& Attr = Collection.AddAttribute<float>(*Key, Group);
+					Attr[0] = AttributeObject->GetNumberField(TEXT("value"));
+				}
+				else if (Type == "string")
+				{
+					TManagedArray<FString>& Attr = Collection.AddAttribute<FString>(*Key, Group);
+					Attr[0] = AttributeObject->GetStringField(TEXT("value"));
+				}
+			}
+		}
+	}
+
+	static void FillPlantProfilesData(FManagedArrayCollection& Collection, const FName Group, const TSharedPtr<FJsonObject>& AttributesObject)
+	{
+		Facades::FPlantProfileFacade Facade(Collection);
+
+		for (const auto& [Key, Attribute] : AttributesObject->Values)
+		{
+			if (const TSharedPtr<FJsonObject>& AttributeObject = Attribute->AsObject();
+				AttributeObject
+				&& AttributeObject->HasTypedField<EJson::Array>(TEXT("value"))
+				&& Key.StartsWith(TEXT("plantProfile_")))
+			{
+				const TArray<TSharedPtr<FJsonValue>>& Values = AttributeObject->GetArrayField(TEXT("value"));
+				TArray<float> Points;
+				Points.Reserve(Values.Num());
+
+				for (const TSharedPtr<FJsonValue>& AttribValue : Values)
+				{
+					if (double Point; AttribValue->TryGetNumber(Point))
+					{
+						Points.Add(static_cast<float>(AttribValue->AsNumber()));
+					}
+				}
+
+				Facade.AddProfileEntry(Points);
+			}
+		}
+	}
+
+	static void FillFoliageData(FManagedArrayCollection& Collection, const TSharedPtr<FJsonObject>& PrimitiveAttributesObject, const FString& InPath)
+	{
+		const TSharedPtr<FJsonObject>& InstancerNameObject = PrimitiveAttributesObject->GetObjectField(TEXT("instancer_name"));
+		const TArray<TSharedPtr<FJsonValue>>& InstancerNameValues = InstancerNameObject->GetArrayField(TEXT("values"));
+
+		const TSharedPtr<FJsonObject>& InstancerPivotObject = PrimitiveAttributesObject->GetObjectField(TEXT("instancer_pivot"));
+		const TArray<TSharedPtr<FJsonValue>>& InstancerPivotValues = InstancerPivotObject->GetArrayField(TEXT("values"));
+
+		const TSharedPtr<FJsonObject>& InstancerUpVectorObject = PrimitiveAttributesObject->GetObjectField(TEXT("instancer_UP"));
+		const TArray<TSharedPtr<FJsonValue>>& InstancerUpVectorValues = InstancerUpVectorObject->GetArrayField(TEXT("values"));
+
+		const TSharedPtr<FJsonObject>& InstancerNormalVectorObject = PrimitiveAttributesObject->GetObjectField(TEXT("instancer_N"));
+		const TArray<TSharedPtr<FJsonValue>>& InstancerNormalVectorValues = InstancerNormalVectorObject->GetArrayField(TEXT("values"));
+
+		const TSharedPtr<FJsonObject>& InstancerScaleObject = PrimitiveAttributesObject->GetObjectField(TEXT("instancer_scale"));
+		const TArray<TSharedPtr<FJsonValue>>& InstancerScaleValues = InstancerScaleObject->GetArrayField(TEXT("values"));
+
+		const TSharedPtr<FJsonObject>& InstancerLFRObject = PrimitiveAttributesObject->GetObjectField(TEXT("instancer_LFR"));
+		const TArray<TSharedPtr<FJsonValue>>& InstancerLFRValues = InstancerLFRObject->GetArrayField(TEXT("values"));
+
+		// Compute unique names and number of elements for group
+		int32 TotalValues = 0;
+		int32 NextFoliageIndex = 0;
+		TMap<FString, int32> FoliageNamesToIndex;
+		TArray<FString> FoliageNames;
+		for (int32 i = 0; i < InstancerNameValues.Num(); ++i)
+		{
+			const TArray<TSharedPtr<FJsonValue>>& InstancerNames = InstancerNameValues[i]->AsArray();
+			for (int32 Index = 0; Index < InstancerNames.Num(); ++Index)
+			{
+				const FString& Name = InstancerNames[Index]->AsString();
+				if (!FoliageNamesToIndex.Contains(Name))
+				{
+					FoliageNamesToIndex.Emplace(Name, NextFoliageIndex++);
+					FoliageNames.Emplace(Name);
+				}
+			}
+			TotalValues += InstancerNames.Num();
+		}
+
+		Facades::FFoliageFacade Facade(Collection, TotalValues);
+
+		FString AbsolutePath = FPaths::ConvertRelativePathToFull(InPath);
+		FString PackagePath;
+
+		if (FPackageName::TryConvertFilenameToLongPackageName(AbsolutePath, PackagePath))
+		{
+			FString PackageFolder = FPackageName::GetLongPackagePath(PackagePath);
+			Facade.SetFoliagePath(PackageFolder);
+			for (FString& FoliageName : FoliageNames)
+			{
+				FoliageName = PackageFolder / FoliageName + '.' + FoliageName;
+			}
+		}
+		Facade.SetFoliageNames(FoliageNames);
+
+		int32 FoliageEntryIndex = 0;
+		for (int32 Id = 0; Id < InstancerNameValues.Num(); ++Id)
+		{
+			const TArray<TSharedPtr<FJsonValue>>& InstancerNames = InstancerNameValues[Id]->AsArray();
+			const TArray<TSharedPtr<FJsonValue>>& PivotPoints = InstancerPivotValues[Id]->AsArray();
+			const TArray<TSharedPtr<FJsonValue>>& UpVectors = InstancerUpVectorValues[Id]->AsArray();
+			const TArray<TSharedPtr<FJsonValue>>& NormalVectors = InstancerNormalVectorValues[Id]->AsArray();
+			const TArray<TSharedPtr<FJsonValue>>& Scales = InstancerScaleValues[Id]->AsArray();
+			const TArray<TSharedPtr<FJsonValue>>& LFRs = InstancerLFRValues[Id]->AsArray();
+
+			TArray<int32> FoliageEntryIds;
+			for (int32 j = 0; j < InstancerNames.Num(); ++j)
+			{
+				const FString Name = InstancerNames[j]->AsString();
+				const FVector3f PivotPoint = FVector3f(
+					PivotPoints[j * 3]->AsNumber(),
+					PivotPoints[(j * 3) + 2]->AsNumber(),
+					PivotPoints[(j * 3) + 1]->AsNumber()) * 100.0f;
+				const FVector3f UpVector = FVector3f(
+					UpVectors[j * 3]->AsNumber(),
+					UpVectors[(j * 3) + 2]->AsNumber(),
+					UpVectors[(j * 3) + 1]->AsNumber());
+				const FVector3f NormalVector = FVector3f(
+					NormalVectors[j * 3]->AsNumber(),
+					NormalVectors[(j * 3) + 2]->AsNumber(),
+					NormalVectors[(j * 3) + 1]->AsNumber());
+				const float Scale = Scales[j]->AsNumber();
+				const float LengthFromRoot = LFRs[j]->AsNumber();
+
+				Facade.SetFoliageEntry(FoliageEntryIndex, {
+					.NameId = FoliageNamesToIndex[Name],
+					.BranchId = Id,
+					.PivotPoint = PivotPoint,
+					.UpVector = UpVector,
+					.NormalVector = NormalVector,
+					.Scale = Scale,
+					.LengthFromRoot = LengthFromRoot,
+				});
+
+				FoliageEntryIds.Add(FoliageEntryIndex);
+				FoliageEntryIndex++;
+			}
+
+			Facade.SetFoliageIdsArray(Id, FoliageEntryIds);
+		}
+	}
+
+	static bool HasJsonFieldPath(const TSharedPtr<FJsonObject>& JsonObject, const FString& Path)
+	{
+		if (!JsonObject.IsValid())
+		{
+			return false;
+		}
+
+		TArray<FString> Keys;
+		Path.ParseIntoArray(Keys, TEXT("."), true);
+
+		TSharedPtr<FJsonObject> CurrentObject = JsonObject;
+
+		for (int32 i = 0; i < Keys.Num(); ++i)
+		{
+			const FString& CurrentKey = Keys[i];
+			if (i == Keys.Num() - 1)
+			{
+				if (!CurrentObject->HasField(CurrentKey))
+				{
+					return false;
+				}
+			}
+			else
+			{
+				const TSharedPtr<FJsonObject>* NextObject;
+				if (CurrentObject->TryGetObjectField(CurrentKey, NextObject))
+				{
+					CurrentObject = *NextObject;
+				}
+				else
+				{
+					return false;
+				}
+			}
+		}
+
+		return true;
+	}
+
+	static bool LoadMegaPlantJsonToCollection(FManagedArrayCollection& Collection, const FString& FilePath, FString& OutErrorMessage)
+	{
+		FString MegaPlantData;
+		if (!FFileHelper::LoadFileToString(MegaPlantData, *FilePath))
+		{
+			OutErrorMessage = TEXT("Failed to open skeleton file");
+			return false;
+		}
+
+		TSharedPtr<FJsonObject> LoadedData = MakeShareable(new FJsonObject);
+		if (!FJsonSerializer::Deserialize(TJsonReaderFactory<>::Create(MegaPlantData), LoadedData))
+		{
+			OutErrorMessage = TEXT("Failed to deserialize JSON file for skeleton");
+			return false;
+		}
+
+		TArray<FString> RequiredJSONPaths = {
+			TEXT("points.attributes.pscale"),
+			TEXT("points.positions"),
+			TEXT("points.attributes.lengthFromRoot.values"),
+			TEXT("points.attributes.LOD_totalPscaleGradient.values"),
+			TEXT("points.attributes.budDirection.values"),
+			TEXT("primitives.points"),
+			TEXT("primitives.attributes.instancer_name.values"),
+			TEXT("primitives.attributes.instancer_pivot.values"),
+			TEXT("primitives.attributes.instancer_UP.values"),
+			TEXT("primitives.attributes.instancer_scale.values"),
+			TEXT("primitives.attributes.instancer_LFR.values"),
+			TEXT("primitives.attributes.parents.values"),
+			TEXT("primitives.attributes.children.values"),
+			TEXT("primitives.attributes.branchNumber.values"),
+			TEXT("globalAttributes.phyllotaxyLeaf.value")
+		};
+
+		for (const FString& Path : RequiredJSONPaths)
+		{
+			if (!HasJsonFieldPath(LoadedData, Path))
+			{
+				OutErrorMessage = FString::Printf(TEXT("Failed to find required path in JSON file: %s"), *Path);
+				return false;
+			}
+		}
+
+		const TSharedPtr<FJsonObject>& Points = LoadedData->GetObjectField(TEXT("points"));
+		const TSharedPtr<FJsonObject>& PointAttributes = Points->GetObjectField(TEXT("attributes"));
+
+		const TArray<TSharedPtr<FJsonValue>>& PointPositions = Points->GetArrayField(TEXT("positions"));
+		const TArray<TSharedPtr<FJsonValue>>& PointScales = PointAttributes->GetObjectField(TEXT("pscale"))->GetArrayField(TEXT("values"));
+
+		Collection.AddGroup(GroupNames::PointGroup);
+		Collection.AddElements(PointPositions.Num(), GroupNames::PointGroup);
+
+		TManagedArray<FVector3f>& PositionArray = Collection.AddAttribute<FVector3f>(AttributeNames::PointPosition, GroupNames::PointGroup);
+		TManagedArray<float>& ScaleArray = Collection.AddAttribute<float>("Scale", GroupNames::PointGroup);
+
+		for (int32 PointIndex = 0; PointIndex < PointPositions.Num(); ++PointIndex)
+		{
+			const TArray<TSharedPtr<FJsonValue>>& PointPosition = PointPositions[PointIndex]->AsArray();
+
+			const FVector3f Position = FVector3f(PointPosition[0]->AsNumber(), PointPosition[2]->AsNumber(), PointPosition[1]->AsNumber()) * 100.0f;
+			const float Scale = PointScales[PointIndex]->AsNumber() * 100.0f;
+
+			PositionArray[PointIndex] = Position;
+			ScaleArray[PointIndex] = Scale;
+		}
+
+		const TSharedPtr<FJsonObject>& Primitives = LoadedData->GetObjectField(TEXT("primitives"));
+		const TSharedPtr<FJsonObject>& PrimitiveAttributes = Primitives->GetObjectField(TEXT("attributes"));
+
+		const TSharedPtr<FJsonObject>& GlobalAttributes = LoadedData->GetObjectField(TEXT("globalAttributes"));
+
+		const TArray<TSharedPtr<FJsonValue>>& PrimitivesPoints = Primitives->GetArrayField(TEXT("points"));
+
+		Collection.AddGroup(GroupNames::BranchGroup);
+		Collection.AddElements(PrimitivesPoints.Num(), GroupNames::BranchGroup);
+
+		TManagedArray<TArray<int32>>& PointsArray = Collection.AddAttribute<TArray<int32>>(
+			MegaPlants::AttributeNames::BranchPoints, MegaPlants::GroupNames::BranchGroup);
+
+		for (int32 PrimitiveIndex = 0; PrimitiveIndex < PrimitivesPoints.Num(); ++PrimitiveIndex)
+		{
+			const TArray<TSharedPtr<FJsonValue>>& PrimitivePoints = PrimitivesPoints[PrimitiveIndex]->AsArray();
+			for (const TSharedPtr<FJsonValue>& PrimitivePoint : PrimitivePoints)
+			{
+				PointsArray[PrimitiveIndex].Add(PrimitivePoint->AsNumber());
+			}
+		}
+
+		Collection.AddGroup(GroupNames::DetailsGroup);
+		Collection.AddElements(1, GroupNames::DetailsGroup);
+
+		// Please note that LFR (length from root) for both skeleton points and foliage points are not being scaled by
+		// 100 at the moment, which is OK for now since they're only being compared relative to each other for now,
+		// however if they're used in any computation in the future in any meaningful way where other quantities are
+		// involved they will end up giving erroneous results.
+		FillAttributes(Collection, GroupNames::PointGroup, PointAttributes);
+		FillAttributes(Collection, GroupNames::BranchGroup, PrimitiveAttributes);
+		FillDetailsAttributes(Collection, GroupNames::DetailsGroup, GlobalAttributes);
+		FillPlantProfilesData(Collection, GroupNames::PlantProfilesGroup, GlobalAttributes);
+
+		FillFoliageData(Collection, PrimitiveAttributes, FilePath);
+
+		Facades::FMetaInfoFacade MetaInfoFacade(Collection);
+		MetaInfoFacade.CreateGuid(FilePath);
+
+		return true;
+	}
+}
diff --git a/Engine/Plugins/Experimental/ProceduralVegetationEditor/Source/MegaPlants/Public/Helpers/MegaPlantUtilities.h b/Engine/Plugins/Experimental/ProceduralVegetationEditor/Source/MegaPlants/Public/Helpers/MegaPlantUtilities.h
new file mode 100644
--- /dev/null
+++ b/Engine/Plugins/Experimental/ProceduralVegetationEditor/Source/MegaPlants/Public/Helpers/MegaPlantUtilities.h
@@ -0,0 +1,17 @@
+﻿// Copyright Epic Games, Inc. All Rights Reserved.
+#pragma once
+#include "Facades/MegaPlantRenderingFacade.h"
+
+class UInstancedStaticMeshComponent;
+class UMeshComponent;
+class FFoliageFacade;
+class AActor;
+
+namespace MegaPlants::Utilities
+{
+	DECLARE_DELEGATE_OneParam(FFoliageComponentCreatedCallback, UMeshComponent*);
+	DECLARE_DELEGATE_OneParam(FFoliageInstanceAddedCallback, int32);
+	
+	void MEGAPLANTS_API AddFoliageInstancesToActor(const Facades::FFoliageFacade& InFoliageFacade, AActor* InActor, TMap<FString, TObjectPtr<UMeshComponent>>& OutInstancedComponentMap);
+	void MEGAPLANTS_API AddFoliageInstances(const Facades::FFoliageFacade& InFoliageFacade, const UObject* InParent, FFoliageComponentCreatedCallback InCallback, TMap<FString, TObjectPtr<UMeshComponent>>& OutInstancedComponentMap);
+}
diff --git a/Engine/Plugins/Experimental/ProceduralVegetationEditor/Source/MegaPlants/Public/MegaPlant.h b/Engine/Plugins/Experimental/ProceduralVegetationEditor/Source/MegaPlants/Public/MegaPlant.h
new file mode 100644
--- /dev/null
+++ b/Engine/Plugins/Experimental/ProceduralVegetationEditor/Source/MegaPlants/Public/MegaPlant.h
@@ -0,0 +1,47 @@
+// Copyright Epic Games, Inc. All Rights Reserved.
+
+#pragma once
+#include "CoreMinimal.h"
+#include "PCGGraph.h"
+#include "MegaPlant.generated.h"
+
+class UMegaPlantRecipeDataAsset;
+
+UCLASS(MinimalAPI)
+class UMegaPlantGraph : public UPCGGraph
+{
+	GENERATED_BODY()
+
+public:
+	UMegaPlantGraph(const FObjectInitializer& ObjectInitializer)
+		: UPCGGraph(ObjectInitializer)
+	{
+#if WITH_EDITOR
+		// Hides input and output nodes
+		SetHiddenFlagInputNode(/*bHidden=*/true);
+		SetHiddenFlagOutputNode(/*bHidden=*/true);
+#endif
+	}
+};
+
+/**
+ * Asset type for procedural plant generation
+ */
+UCLASS()
+class MEGAPLANTS_API UMegaPlant : public UObject
+{
+	GENERATED_BODY()
+
+	UPROPERTY()
+	TObjectPtr<UMegaPlantGraph> Graph;
+
+public:
+	TObjectPtr<UMegaPlantGraph>& GetGraph() { return Graph; }
+	const TObjectPtr<UMegaPlantGraph>& GetGraph() const { return Graph; }
+
+	void CreateGraph(const UMegaPlantGraph* InGraph = nullptr);
+
+#if WITH_EDITOR
+	void CreateGraphFromRecipe(const TObjectPtr<UMegaPlantRecipeDataAsset> Recipe);
+#endif
+};
diff --git a/Engine/Plugins/Experimental/ProceduralVegetationEditor/Source/MegaPlants/Public/MegaPlantLink.h b/Engine/Plugins/Experimental/ProceduralVegetationEditor/Source/MegaPlants/Public/MegaPlantLink.h
new file mode 100644
--- /dev/null
+++ b/Engine/Plugins/Experimental/ProceduralVegetationEditor/Source/MegaPlants/Public/MegaPlantLink.h
@@ -0,0 +1,17 @@
+// Copyright Epic Games, Inc. All Rights Reserved.
+
+#pragma once
+
+#include "Engine/AssetUserData.h"
+#include "MegaPlant.h"
+#include "MegaPlantLink.generated.h"
+
+UCLASS(MinimalAPI)
+class UMegaPlantLink : public UAssetUserData
+{
+	GENERATED_BODY()
+
+public:
+	UPROPERTY(VisibleAnywhere, Category = Settings)
+	TObjectPtr<UMegaPlant> Source;
+};
diff --git a/Engine/Plugins/Experimental/ProceduralVegetationEditor/Source/MegaPlants/Public/MegaPlantRecipeDataAsset.h b/Engine/Plugins/Experimental/ProceduralVegetationEditor/Source/MegaPlants/Public/MegaPlantRecipeDataAsset.h
new file mode 100644
--- /dev/null
+++ b/Engine/Plugins/Experimental/ProceduralVegetationEditor/Source/MegaPlants/Public/MegaPlantRecipeDataAsset.h
@@ -0,0 +1,24 @@
+// Copyright Epic Games, Inc. All Rights Reserved.
+
+#pragma once
+
+#include "Engine/DataAsset.h"
+
+#include "GeometryCollection/ManagedArrayCollection.h"
+
+#include "MegaPlantRecipeDataAsset.generated.h"
+
+class UMegaPlantGraph;
+
+UCLASS(BlueprintType)
+class MEGAPLANTS_API UMegaPlantRecipeDataAsset : public UDataAsset
+{
+	GENERATED_BODY()
+
+public:
+	UPROPERTY(EditAnywhere, Category="Recipe")
+	TMap<FString, FManagedArrayCollection> Variants;
+
+public:
+	void LoadFromVariantFiles(const TMap<FString, FString>& InVariantFiles);
+};
diff --git a/Engine/Plugins/Experimental/ProceduralVegetationEditor/Source/MegaPlants/Public/MegaPlantsModule.h b/Engine/Plugins/Experimental/ProceduralVegetationEditor/Source/MegaPlants/Public/MegaPlantsModule.h
new file mode 100644
--- /dev/null
+++ b/Engine/Plugins/Experimental/ProceduralVegetationEditor/Source/MegaPlants/Public/MegaPlantsModule.h
@@ -0,0 +1,18 @@
+// Copyright Epic Games, Inc. All Rights Reserved.
+#pragma once
+
+#include "CoreMinimal.h"
+#include "Modules/ModuleManager.h"
+
+DECLARE_LOG_CATEGORY_EXTERN(LogMegaPlants, Log, All);
+
+class FMegaPlantsModule : public FDefaultGameModuleImpl
+{
+
+public:
+	/** IModuleInterface implementation */
+	virtual void StartupModule() override;
+	virtual void ShutdownModule() override;
+};
+
+
diff --git a/Engine/Plugins/Experimental/ProceduralVegetationEditor/Source/MegaPlants/Public/Nodes/PCGMegaPlantBaseSettings.h b/Engine/Plugins/Experimental/ProceduralVegetationEditor/Source/MegaPlants/Public/Nodes/PCGMegaPlantBaseSettings.h
new file mode 100644
--- /dev/null
+++ b/Engine/Plugins/Experimental/ProceduralVegetationEditor/Source/MegaPlants/Public/Nodes/PCGMegaPlantBaseSettings.h
@@ -0,0 +1,32 @@
+﻿// Copyright Epic Games, Inc. All Rights Reserved.
+#pragma once
+
+#include "CoreMinimal.h"
+#include "PCGSettings.h"
+#include "PCGMegaPlantData.h"
+#include "PCGMegaPlantBaseSettings.generated.h"
+
+UCLASS(BlueprintType, Abstract, HideCategories=(Debug), ClassGroup = (Procedural))
+class UPCGMegaPlantBaseSettings : public UPCGSettings
+{
+	GENERATED_BODY()
+
+public:
+#if WITH_EDITOR
+	virtual EPCGSettingsType GetType() const override { return EPCGSettingsType::Spatial; }
+	virtual TArray<EMegaPlantRenderType> GetRenderType() const { return TArray<EMegaPlantRenderType>(); }
+#endif
+
+protected:
+	virtual TArray<FPCGPinProperties> InputPinProperties() const override;
+	virtual TArray<FPCGPinProperties> OutputPinProperties() const override;
+	virtual FPCGElementPtr CreateElement() const override;
+};
+
+class FPCGMegaPlantBaseElement : public IPCGElement
+{
+protected:
+	virtual void PostExecuteInternal(FPCGContext* InContext) const override;
+	virtual bool ExecuteInternal(FPCGContext* Context) const override;
+};
+
diff --git a/Engine/Plugins/Experimental/ProceduralVegetationEditor/Source/MegaPlants/Public/Nodes/PCGMegaPlantOutput.h b/Engine/Plugins/Experimental/ProceduralVegetationEditor/Source/MegaPlants/Public/Nodes/PCGMegaPlantOutput.h
new file mode 100644
--- /dev/null
+++ b/Engine/Plugins/Experimental/ProceduralVegetationEditor/Source/MegaPlants/Public/Nodes/PCGMegaPlantOutput.h
@@ -0,0 +1,37 @@
+﻿// Copyright Epic Games, Inc. All Rights Reserved.
+#pragma once
+
+#include "CoreMinimal.h"
+#include "Helpers/MegaPlantExportHelper.h"
+#include "Nodes/PCGMegaPlantBaseSettings.h"
+#include "PCGMegaPlantOutput.generated.h"
+
+UCLASS(ClassGroup = (Procedural))
+class UPCGMegaPlantOutputSettings : public UPCGMegaPlantBaseSettings
+{
+	GENERATED_BODY()
+
+public:
+	UPCGMegaPlantOutputSettings();
+
+	#if WITH_EDITOR
+	virtual FName GetDefaultNodeName() const override { return FName(TEXT("MegaPlantOutput")); }
+	virtual FText GetDefaultNodeTitle() const override { return NSLOCTEXT("PCGMegaPlantOutput", "NodeTitle", "MegaPlant Output"); }
+	virtual FText GetNodeTooltipText() const override { return NSLOCTEXT("PCGMegaPlantOutput", "NodeTooltip", "Output node for Mega plant"); }
+	virtual EPCGSettingsType GetType() const override { return EPCGSettingsType::InputOutput; }
+	#endif
+
+protected:
+	
+	virtual TArray<FPCGPinProperties> OutputPinProperties() const override;
+	virtual FPCGElementPtr CreateElement() const override;
+	
+public:
+	UPROPERTY(EditAnywhere, Category = ExportSettings, meta = (ShowOnlyInnerProperties))
+	FMegaPlantExportSettings ExportSettings;
+};
+
+class UPCGMegaPlantOutputElement : public FPCGMegaPlantBaseElement
+{
+	virtual bool ExecuteInternal(FPCGContext* Context) const override;
+};
\ No newline at end of file
diff --git a/Engine/Plugins/Experimental/ProceduralVegetationEditor/Source/MegaPlants/Public/PCGMegaPlantData.h b/Engine/Plugins/Experimental/ProceduralVegetationEditor/Source/MegaPlants/Public/PCGMegaPlantData.h
new file mode 100644
--- /dev/null
+++ b/Engine/Plugins/Experimental/ProceduralVegetationEditor/Source/MegaPlants/Public/PCGMegaPlantData.h
@@ -0,0 +1,77 @@
+// Copyright Epic Games, Inc. All Rights Reserved.
+
+#pragma once
+
+#include "Data/PCGSpatialData.h"
+#include "GeometryCollection/ManagedArrayCollection.h"
+#include "PCGMegaPlantData.generated.h"
+
+struct FPCGContext;
+
+UENUM()
+enum class EMegaPlantRenderType : uint8
+{
+	None 		= 0,
+	Skeleton 	= 1 << 0,
+	Mesh 		= 1 << 1,
+	Foliage 	= 1 << 2,
+
+	// Add new renter type above
+	Count	
+};
+
+ENUM_CLASS_FLAGS(EMegaPlantRenderType)
+
+UCLASS(BlueprintType, ClassGroup = (Procedural))
+class MEGAPLANTS_API UPCGMegaPlantData : public UPCGSpatialData
+{
+	GENERATED_BODY()
+
+public:
+	UPCGMegaPlantData(const FObjectInitializer& ObjectInitializer);
+	
+	void Initialize(FManagedArrayCollection&& InCollection, bool bCanTakeOwnership = false);
+	
+	// ~Begin UPCGData interface
+	virtual EPCGDataType GetDataType() const override { return EPCGDataType::Other; }
+	virtual void AddToCrc(FArchiveCrc32& Ar, bool bFullDataCrc) const override;
+	// ~End UPCGData interface
+
+	// ~Begin UPCGSpatialData interface
+	virtual int GetDimension() const override { return 3; }
+	virtual FBox GetBounds() const override;
+	virtual bool SamplePoint(const FTransform& Transform, const FBox& Bounds, FPCGPoint& OutPoint, UPCGMetadata* OutMetadata) const override;
+	//~End UPCGSpatialData interface
+
+	const FManagedArrayCollection& GetCollection() const { return Collection; }
+
+#if WITH_EDITOR
+	const TArray<EMegaPlantRenderType>& GetRenderType() const { return RenderType; }
+	void SetRenderType(TArray<EMegaPlantRenderType> InRenderType) const { RenderType = InRenderType; }
+#endif
+	
+protected:
+
+	FManagedArrayCollection Collection;
+
+#if WITH_EDITOR
+	mutable TArray<EMegaPlantRenderType> RenderType;
+#endif
+	
+	// ~Begin UPCGData interface
+	virtual bool SupportsFullDataCrc() const override { return true; }
+	// ~End UPCGData interface
+
+	// ~Begin UPCGSpatialData interface
+	virtual UPCGSpatialData* CopyInternal(FPCGContext* Context) const override;
+public:
+	virtual const UPCGPointData* ToPointData(FPCGContext* Context, const FBox& InBounds) const override;
+	virtual const UPCGPointArrayData* ToPointArrayData(FPCGContext* Context, const FBox& InBounds) const override;
+	//~End UPCGSpatialData interface
+	
+private:
+	const UPCGBasePointData* ToBasePointData(FPCGContext* Context, const FBox& InBounds, TSubclassOf<UPCGBasePointData> PointDataClass) const;
+protected:
+	
+	mutable FBox CachedBounds = FBox(EForceInit::ForceInit);
+};
diff --git a/Engine/Plugins/Experimental/ProceduralVegetationEditor/Source/MegaPlantsEditor/MegaPlantsEditor.Build.cs b/Engine/Plugins/Experimental/ProceduralVegetationEditor/Source/MegaPlantsEditor/MegaPlantsEditor.Build.cs
new file mode 100644
--- /dev/null
+++ b/Engine/Plugins/Experimental/ProceduralVegetationEditor/Source/MegaPlantsEditor/MegaPlantsEditor.Build.cs
@@ -0,0 +1,58 @@
+// Copyright Epic Games, Inc. All Rights Reserved.
+
+using UnrealBuildTool;
+
+public class MegaPlantsEditor : ModuleRules
+{
+	public MegaPlantsEditor(ReadOnlyTargetRules Target) : base(Target)
+	{
+		PCHUsage = PCHUsageMode.UseExplicitOrSharedPCHs;
+
+		PublicDependencyModuleNames.AddRange(
+			new string[]
+			{
+				// ... add other public dependencies that you statically link with here ...
+			}
+		);
+
+		PrivateDependencyModuleNames.AddRange(
+			new string[]
+			{
+				"AdvancedPreviewScene",
+				"AssetDefinition",
+				"AssetTools",
+				"Chaos",
+				"ContentBrowser",
+				"ContentBrowserData",
+				"Core",
+				"CoreUObject",
+				"DataflowEditor",
+				"Engine",
+				"Fab",
+				"GeometryCollectionEngine",
+				"GeometryCore",
+				"GeometryFramework",
+				"InputCore",
+				"InterchangeCore",
+				"InterchangeEngine",
+				"InterchangeFactoryNodes",
+				"InterchangeNodes",
+				"InterchangePipelines",
+				"Json",
+				"JsonUtilities",
+				"LevelEditor",
+				"MegaPlants",
+				"MeshConversionEngineTypes",
+				"PCG",
+				"PCGEditor",
+				"Projects",
+				"RenderCore",
+				"Slate",
+				"SlateCore",
+				"ToolMenus",
+				"ToolWidgets",
+				"UnrealEd"
+			}
+		);
+	}
+}
\ No newline at end of file
diff --git a/Engine/Plugins/Experimental/ProceduralVegetationEditor/Source/MegaPlantsEditor/Private/AssetDefinition_MegaPlant.cpp b/Engine/Plugins/Experimental/ProceduralVegetationEditor/Source/MegaPlantsEditor/Private/AssetDefinition_MegaPlant.cpp
new file mode 100644
--- /dev/null
+++ b/Engine/Plugins/Experimental/ProceduralVegetationEditor/Source/MegaPlantsEditor/Private/AssetDefinition_MegaPlant.cpp
@@ -0,0 +1,37 @@
+﻿// Copyright Epic Games, Inc. All Rights Reserved.
+
+#include "AssetDefinition_MegaPlant.h"
+#include "MegaPlantEditor.h"
+#include "MegaPlant.h"
+
+FText UAssetDefinition_MegaPlant::GetAssetDisplayName() const
+{
+	return INVTEXT("MegaPlant");
+}
+
+FLinearColor UAssetDefinition_MegaPlant::GetAssetColor() const
+{
+	return FColor::Green;
+}
+
+TSoftClassPtr<UObject> UAssetDefinition_MegaPlant::GetAssetClass() const
+{
+	return UMegaPlant::StaticClass();
+}
+
+TConstArrayView<FAssetCategoryPath> UAssetDefinition_MegaPlant::GetAssetCategories() const
+{
+	static const auto Categories = { EAssetCategoryPaths::Foliage };
+	return Categories;
+}
+
+EAssetCommandResult UAssetDefinition_MegaPlant::OpenAssets(const FAssetOpenArgs& OpenArgs) const
+{
+	for (UMegaPlant* MegaPlant : OpenArgs.LoadObjects<UMegaPlant>())
+	{
+		const TSharedRef<FMegaPlantEditor> MegaPlantEditor = MakeShared<FMegaPlantEditor>();
+		MegaPlantEditor->Initialize(EToolkitMode::Standalone, OpenArgs.ToolkitHost, MegaPlant);
+	}
+
+	return EAssetCommandResult::Handled;
+}
diff --git a/Engine/Plugins/Experimental/ProceduralVegetationEditor/Source/MegaPlantsEditor/Private/AssetDefinition_MegaPlant.h b/Engine/Plugins/Experimental/ProceduralVegetationEditor/Source/MegaPlantsEditor/Private/AssetDefinition_MegaPlant.h
new file mode 100644
--- /dev/null
+++ b/Engine/Plugins/Experimental/ProceduralVegetationEditor/Source/MegaPlantsEditor/Private/AssetDefinition_MegaPlant.h
@@ -0,0 +1,19 @@
+﻿// Copyright Epic Games, Inc. All Rights Reserved.
+
+#pragma once
+
+#include "AssetDefinitionDefault.h"
+#include "AssetDefinition_MegaPlant.generated.h"
+
+UCLASS()
+class UAssetDefinition_MegaPlant : public UAssetDefinitionDefault
+{
+	GENERATED_BODY()
+
+public:
+	virtual FText GetAssetDisplayName() const override;
+	virtual FLinearColor GetAssetColor() const override;
+	virtual TSoftClassPtr<UObject> GetAssetClass() const override;
+	virtual TConstArrayView<FAssetCategoryPath> GetAssetCategories() const override;
+	virtual EAssetCommandResult OpenAssets(const FAssetOpenArgs& OpenArgs) const override;
+};
diff --git a/Engine/Plugins/Experimental/ProceduralVegetationEditor/Source/MegaPlantsEditor/Private/Importer/MegaPlantImportPipeline.cpp b/Engine/Plugins/Experimental/ProceduralVegetationEditor/Source/MegaPlantsEditor/Private/Importer/MegaPlantImportPipeline.cpp
new file mode 100644
--- /dev/null
+++ b/Engine/Plugins/Experimental/ProceduralVegetationEditor/Source/MegaPlantsEditor/Private/Importer/MegaPlantImportPipeline.cpp
@@ -0,0 +1,234 @@
+// Copyright Epic Games, Inc. All Rights Reserved.
+
+#include "MegaPlantImportPipeline.h"
+
+#include "InterchangeMaterialFactoryNode.h"
+#include "InterchangeStaticMeshFactoryNode.h"
+#include "InterchangeTextureFactoryNode.h"
+#include "RenderUtils.h"
+
+#include "Serialization/JsonSerializer.h"
+
+UMegaPlantImportPipeline::UMegaPlantImportPipeline()
+{}
+
+void UMegaPlantImportPipeline::ExecutePipeline(
+	UInterchangeBaseNodeContainer* NodeContainer,
+	const TArray<UInterchangeSourceData*>& SourceDatas,
+	const FString& ContentBasePath
+)
+{
+	Super::ExecutePipeline(NodeContainer, SourceDatas, ContentBasePath);
+
+	BaseNodeContainer = NodeContainer;
+
+	const UInterchangeSourceData* const* GltfSourceData = SourceDatas.FindByPredicate(
+		[](const UInterchangeSourceData* SourceData)
+			{
+				return FPaths::GetExtension(SourceData->GetFilename()) == "gltf";
+			}
+	);
+	if (!GltfSourceData)
+	{
+		return;
+	}
+
+	if (!LoadGltfSource((*GltfSourceData)->GetFilename()))
+	{
+		return;
+	}
+
+	TextureFactoryNodes = GetNodesOfType<UInterchangeTextureFactoryNode>();
+	MaterialInstanceFactoryNodes = GetNodesOfType<UInterchangeMaterialInstanceFactoryNode>();
+	StaticMeshFactoryNodes = GetNodesOfType<UInterchangeStaticMeshFactoryNode>();
+
+	for (const TObjectPtr<UInterchangeStaticMeshFactoryNode>& StaticMeshFactoryNode : StaticMeshFactoryNodes)
+	{
+		static const FString AttrKey = "Nanite_Shape_Preservation";
+
+		static const bool bAllowNaniteFoliage = NaniteVoxelsSupported();
+
+		StaticMeshFactoryNode->SetAttribute(AttrKey,
+			static_cast<uint8>(bAllowNaniteFoliage ? ENaniteShapePreservation::Voxelize : ENaniteShapePreservation::PreserveArea));
+		StaticMeshFactoryNode->AddApplyAndFillDelegates<uint8>(
+			AttrKey,
+			StaticMeshFactoryNode->GetObjectClass(),
+			"NaniteSettings.ShapePreservation"
+		);
+	}
+
+	const TArray<TSharedPtr<FJsonValue>>& Images = GltfJson->GetArrayField(TEXT("images"));
+	for (const TSharedPtr<FJsonValue>& Image : Images)
+	{
+		const TSharedPtr<FJsonObject>& Texture = Image->AsObject();
+		const FString TextureName = Texture->GetStringField(TEXT("name"));
+		TObjectPtr<UInterchangeTextureFactoryNode> TextureFactoryNode = FindTextureFactoryNodeByName(TextureName);
+		if (TextureFactoryNode == nullptr)
+		{
+			return;
+		}
+
+		TextureFactoryNode->SetCustomVirtualTextureStreaming(true);
+
+		if (const TSharedPtr<FJsonObject>* TextureExtras; Texture->TryGetObjectField(TEXT("extras"), TextureExtras))
+		{
+			if (FString CompressionSettings; (*TextureExtras)->TryGetStringField(TEXT("compression"), CompressionSettings))
+			{
+				if (CompressionSettings == "mask")
+					TextureFactoryNode->SetCustomCompressionSettings(TC_Masks);
+				else if (CompressionSettings == "displacement" || CompressionSettings == "alpha")
+					TextureFactoryNode->SetCustomCompressionSettings(TC_Alpha);
+			}
+			if (FString TextureSlot; (*TextureExtras)->TryGetStringField(TEXT("textureSlot"), TextureSlot))
+			{
+				int32 MaterialIndex = 0;
+				(*TextureExtras)->TryGetNumberField(TEXT("materialIndex"), MaterialIndex);
+				if (const TSharedPtr<FJsonObject>* MaterialObject = GetMaterialAtIndex(MaterialIndex))
+				{
+					const FString MaterialName = (*MaterialObject)->GetStringField(TEXT("name"));
+					if (UInterchangeMaterialInstanceFactoryNode* MaterialInstanceFactoryNode = FindMaterialInstanceFactoryNodeByName(MaterialName))
+					{
+						const FString ParameterName = UInterchangeShaderPortsAPI::MakeInputValueKey(TextureSlot);
+						const FString TextureFactoryUid = TextureFactoryNode->GetUniqueID();
+						MaterialInstanceFactoryNode->AddStringAttribute(ParameterName, TextureFactoryUid);
+						MaterialInstanceFactoryNode->AddFactoryDependencyUid(TextureFactoryUid);
+					}
+				}
+			}
+		}
+	}
+
+
+	const TArray<TSharedPtr<FJsonValue>>& GltfMaterials = GltfJson->GetArrayField(TEXT("materials"));
+	for (const TSharedPtr<FJsonValue>& GltfMaterial : GltfMaterials)
+	{
+		const TSharedPtr<FJsonObject>& MaterialObject = GltfMaterial->AsObject();
+		const FString& MaterialName = MaterialObject->GetStringField(TEXT("name"));
+		if (const TSharedPtr<FJsonObject>* MaterialExtras; MaterialObject->TryGetObjectField(TEXT("extras"), MaterialExtras))
+		{
+			UInterchangeMaterialInstanceFactoryNode* const& MaterialInstanceFactoryNode = FindMaterialInstanceFactoryNodeByName(MaterialName);
+			if (FString MaterialType; (*MaterialExtras)->TryGetStringField(TEXT("type"), MaterialType))
+			{
+				if (MaterialType == "foliage")
+				{
+					const FString FoliageParentMaterial =
+						"/MegaPlants/SampleTrees/MSPresets/MS_Foliage_Material_FAB_VT_V1/MasterMaterials/MA_Foliage_Trees.MA_Foliage_Trees";
+					MaterialInstanceFactoryNode->SetCustomParent(FoliageParentMaterial);
+				}
+				else if (MaterialType == "trunk" || MaterialType == "bark")
+				{
+					const FString BarkParentMaterial =
+						"/MegaPlants/SampleTrees/MSPresets/MS_Foliage_Material_FAB_VT_V1/MasterMaterials/MA_Foliage_Trees.MA_Foliage_Trees";
+					MaterialInstanceFactoryNode->SetCustomParent(BarkParentMaterial);
+				}
+			}
+
+			if (const TSharedPtr<FJsonObject>* CustomProperties; (*MaterialExtras)->TryGetObjectField(TEXT("properties"), CustomProperties))
+			{
+				if (EBlendMode BlendMode; (*CustomProperties)->TryGetNumberField(TEXT("blend_mode"), reinterpret_cast<int&>(BlendMode)))
+				{
+					static const FString BlendModeOverrideAttrKey = "MegaSeed_BlendModeOverride";
+					static const FString BlendModeValueAttrKey = "MegaSeed_BlendModeValue";
+
+					MaterialInstanceFactoryNode->SetAttribute(BlendModeOverrideAttrKey, true);
+					MaterialInstanceFactoryNode->AddApplyAndFillDelegates<bool>(
+						BlendModeOverrideAttrKey,
+						MaterialInstanceFactoryNode->GetObjectClass(),
+						"BasePropertyOverrides.bOverride_BlendMode"
+					);
+					MaterialInstanceFactoryNode->SetAttribute(BlendModeValueAttrKey, static_cast<int>(BlendMode));
+					MaterialInstanceFactoryNode->AddApplyAndFillDelegates<int>(
+						BlendModeValueAttrKey,
+						MaterialInstanceFactoryNode->GetObjectClass(),
+						"BasePropertyOverrides.BlendMode"
+					);
+				}
+			}
+
+			const TArray<TSharedPtr<FJsonValue>>* CustomParams;
+			if (!(*MaterialExtras)->TryGetArrayField(TEXT("params"), CustomParams))
+			{
+				continue;
+			}
+
+			for (const TSharedPtr<FJsonValue>& Param : *CustomParams)
+			{
+				if (Param->Type != EJson::Object)
+					continue;
+
+				const TSharedPtr<FJsonObject>& ParamObject = Param->AsObject();
+				const FString Name = ParamObject->GetStringField(TEXT("Name"));
+				const TSharedPtr<FJsonValue> Value = ParamObject->TryGetField(TEXT("Value"));
+
+				const FString ParameterName = UInterchangeShaderPortsAPI::MakeInputValueKey(Name);
+				if (Value->Type == EJson::Boolean)
+				{
+					MaterialInstanceFactoryNode->AddBooleanAttribute(ParameterName, Value->AsBool());
+				}
+				else if (Value->Type == EJson::Number)
+				{
+					MaterialInstanceFactoryNode->AddFloatAttribute(ParameterName, Value->AsNumber());
+				}
+				else if (Value->Type == EJson::Array)
+				{
+					const TArray<TSharedPtr<FJsonValue>>& ArrayValue = Value->AsArray();
+					if (ArrayValue.Num() >= 4)
+					{
+						MaterialInstanceFactoryNode->AddLinearColorAttribute(
+							ParameterName,
+							FLinearColor(
+								ArrayValue[0]->AsNumber(),
+								ArrayValue[1]->AsNumber(),
+								ArrayValue[2]->AsNumber(),
+								ArrayValue[3]->AsNumber()
+							)
+						);
+					}
+				}
+			}
+		}
+	}
+}
+
+bool UMegaPlantImportPipeline::LoadGltfSource(const FString& SourceFile)
+{
+	if (FString GltfFileData; FFileHelper::LoadFileToString(GltfFileData, *SourceFile))
+	{
+		GltfJson = MakeShareable(new FJsonObject);
+		return FJsonSerializer::Deserialize(TJsonReaderFactory<TCHAR>::Create(GltfFileData), GltfJson);
+	}
+	return false;
+}
+
+const TSharedPtr<FJsonObject>* UMegaPlantImportPipeline::GetMaterialAtIndex(uint32 Index) const
+{
+	TSharedPtr<FJsonObject>* Material = nullptr;
+	const TArray<TSharedPtr<FJsonValue>> Materials = GltfJson->GetArrayField(TEXT("materials"));
+	if (Materials.IsValidIndex(Index))
+	{
+		Materials[Index]->TryGetObject(Material);
+	}
+	return Material;
+}
+
+TObjectPtr<UInterchangeTextureFactoryNode> UMegaPlantImportPipeline::FindTextureFactoryNodeByName(const FString& DisplayName) const
+{
+	const TObjectPtr<UInterchangeTextureFactoryNode>* FoundNode = TextureFactoryNodes.FindByPredicate(
+		[&DisplayName](const TObjectPtr<UInterchangeTextureFactoryNode>& Node)
+			{
+				return Node->GetDisplayLabel() == DisplayName;
+			}
+	);
+	return FoundNode ? *FoundNode : nullptr;
+}
+
+TObjectPtr<UInterchangeMaterialInstanceFactoryNode> UMegaPlantImportPipeline::FindMaterialInstanceFactoryNodeByName(const FString& DisplayName) const
+{
+	const TObjectPtr<UInterchangeMaterialInstanceFactoryNode>* FoundNode = MaterialInstanceFactoryNodes.FindByPredicate(
+		[&DisplayName](const TObjectPtr<UInterchangeMaterialInstanceFactoryNode>& Node)
+			{
+				return Node->GetDisplayLabel() == DisplayName;
+			}
+	);
+	return FoundNode ? *FoundNode : nullptr;
+}
diff --git a/Engine/Plugins/Experimental/ProceduralVegetationEditor/Source/MegaPlantsEditor/Private/Importer/MegaPlantImportPipeline.h b/Engine/Plugins/Experimental/ProceduralVegetationEditor/Source/MegaPlantsEditor/Private/Importer/MegaPlantImportPipeline.h
new file mode 100644
--- /dev/null
+++ b/Engine/Plugins/Experimental/ProceduralVegetationEditor/Source/MegaPlantsEditor/Private/Importer/MegaPlantImportPipeline.h
@@ -0,0 +1,78 @@
+// Copyright Epic Games, Inc. All Rights Reserved.
+
+#pragma once
+
+#include "InterchangePipelineBase.h"
+
+#include "Dom/JsonObject.h"
+
+#include "MegaPlantImportPipeline.generated.h"
+
+class UInterchangeTextureFactoryNode;
+class UInterchangeMaterialInstanceFactoryNode;
+class UInterchangeStaticMeshFactoryNode;
+
+UCLASS(BlueprintType)
+class UMegaPlantImportPipeline : public UInterchangePipelineBase
+{
+	GENERATED_BODY()
+
+public:
+	UMegaPlantImportPipeline();
+
+protected:
+	virtual void ExecutePipeline(
+		UInterchangeBaseNodeContainer* NodeContainer,
+		const TArray<UInterchangeSourceData*>& SourceDatas,
+		const FString& ContentBasePath
+	) override;
+	
+	virtual bool CanExecuteOnAnyThread(EInterchangePipelineTask PipelineTask) override { return true; }
+
+private:
+	bool LoadGltfSource(const FString& SourceFile);
+
+	const TSharedPtr<FJsonObject>* GetMaterialAtIndex(uint32 Index) const;
+
+	TObjectPtr<UInterchangeTextureFactoryNode> FindTextureFactoryNodeByName(const FString& DisplayName) const;
+	TObjectPtr<UInterchangeMaterialInstanceFactoryNode> FindMaterialInstanceFactoryNodeByName(const FString& DisplayName) const;
+
+	template <class NodeType>
+	TArray<TObjectPtr<NodeType>> GetNodesOfType()
+	{
+		TArray<TObjectPtr<NodeType>> Nodes;
+		BaseNodeContainer->IterateNodesOfType<NodeType>(
+			[&Nodes](const FString& Uid, NodeType* Node)
+			{
+				Nodes.Add(Node);
+			}
+		);
+		return Nodes;
+	}
+
+	template <class NodeType>
+	NodeType* FindNodeOfTypeByName(const FString& DisplayName)
+	{
+		NodeType* FoundNode = nullptr;
+		BaseNodeContainer->IterateNodesOfType<NodeType>(
+			[&](const FString& Uid, NodeType* Node)
+			{
+				if (FoundNode)
+					return;
+				if (Node->GetDisplayLabel() == DisplayName)
+					FoundNode = Node;
+			}
+		);
+		return FoundNode;
+	}
+
+private:
+	UPROPERTY()
+	TObjectPtr<UInterchangeBaseNodeContainer> BaseNodeContainer = nullptr;
+
+	TSharedPtr<FJsonObject> GltfJson = nullptr;
+
+	TArray<TObjectPtr<UInterchangeTextureFactoryNode>> TextureFactoryNodes;
+	TArray<TObjectPtr<UInterchangeMaterialInstanceFactoryNode>> MaterialInstanceFactoryNodes;
+	TArray<TObjectPtr<UInterchangeStaticMeshFactoryNode>> StaticMeshFactoryNodes;
+};
