Create ProceduralVegetationEditor plugin

diff --git a/Engine/Plugins/Experimental/ProceduralVegetationEditor/Source/MegaPlantsEditor/Private/Importer/MegaPlantImportWorkflow.cpp b/Engine/Plugins/Experimental/ProceduralVegetationEditor/Source/MegaPlantsEditor/Private/Importer/MegaPlantImportWorkflow.cpp
new file mode 100644
--- /dev/null
+++ b/Engine/Plugins/Experimental/ProceduralVegetationEditor/Source/MegaPlantsEditor/Private/Importer/MegaPlantImportWorkflow.cpp
@@ -0,0 +1,214 @@
+// Copyright Epic Games, Inc. All Rights Reserved.
+
+#include "MegaPlantImportWorkflow.h"
+#include "FabDownloader.h"
+#include "MegaPlantImporter.h"
+#include "MegaPlantRecipeDataAsset.h"
+#include "AssetRegistry/IAssetRegistry.h"
+#include "Async/Async.h"
+#include "Dom/JsonObject.h"
+#include "Dom/JsonValue.h"
+#include "Facades/MegaPlantBranchFacade.h"
+#include "Facades/MegaPlantFoliageFacade.h"
+#include "HAL/PlatformFileManager.h"
+#include "Misc/FileHelper.h"
+#include "MegaPlant.h"
+#include "Serialization/JsonReader.h"
+#include "Serialization/JsonSerializer.h"
+#include "Utilities/AssetUtils.h"
+#include "Utilities/FabAssetsCache.h"
+
+FMegaPlantImportWorkflow::FMegaPlantImportWorkflow(const FString& InAssetId, const FString& InAssetName, const FString& InDownloadUrl)
+	: IFabWorkflow(InAssetId, InAssetName, InDownloadUrl)
+{}
+
+void FMegaPlantImportWorkflow::Execute()
+{
+	DownloadContent();
+}
+
+void FMegaPlantImportWorkflow::DownloadContent()
+{
+	const FString DownloadLocation = FFabAssetsCache::GetCacheLocation() / AssetId;
+
+	DownloadRequest = MakeShared<FFabDownloadRequest>(AssetId, DownloadUrl, DownloadLocation, EFabDownloadType::HTTP);
+	DownloadRequest->OnDownloadProgress().AddRaw(this, &FMegaPlantImportWorkflow::OnContentDownloadProgress);
+	DownloadRequest->OnDownloadComplete().AddRaw(this, &FMegaPlantImportWorkflow::OnContentDownloadComplete);
+	DownloadRequest->ExecuteRequest();
+
+	CreateDownloadNotification();
+}
+
+void FMegaPlantImportWorkflow::ImportContent(const TArray<FString>& SourceFiles)
+{
+	CreateImportNotification();
+
+	FString AssetImportFolder = AssetName.IsEmpty() ? AssetId : AssetName;
+	FAssetUtils::SanitizeFolderName(AssetImportFolder);
+	ImportLocation = "/Game/Fab/MegaPlants" / AssetImportFolder;
+
+	const FString RecipeFilePath = SourceFiles[0];
+	const FString RecipeDownloadFolder = FPaths::GetPath(RecipeFilePath);
+	TSharedPtr<FJsonObject> RecipeData = nullptr;
+
+	bool bSuccess = false;
+
+	FString RecipeFile;
+	bSuccess = FFileHelper::LoadFileToString(RecipeFile, *RecipeFilePath);
+	if (!bSuccess)
+	{
+		ExpireImportNotification(false);
+		CancelWorkflow();
+	}
+
+	RecipeData = MakeShareable(new FJsonObject);
+	bSuccess = FJsonSerializer::Deserialize(TJsonReaderFactory<TCHAR>::Create(RecipeFile), RecipeData);
+	if (!bSuccess)
+	{
+		ExpireImportNotification(false);
+		CancelWorkflow();
+	}
+
+	const TSharedPtr<FJsonObject>& RecipeResources = RecipeData->GetObjectField(TEXT("resources"));
+
+	const TArray<TSharedPtr<FJsonValue>>& RecipeVariations = RecipeResources->GetArrayField(TEXT("variations"));
+	const TArray<TSharedPtr<FJsonValue>>& PreBakedMeshes = RecipeResources->GetArrayField(TEXT("prebaked_meshes"));
+
+	TArray<FString> PreBakedMeshesPaths;
+	for (const TSharedPtr<FJsonValue>& PreBakedMeshData : PreBakedMeshes)
+	{
+		const TSharedPtr<FJsonObject>& PreBakedMesh = PreBakedMeshData->AsObject();
+		const FString MeshUri = RecipeDownloadFolder / PreBakedMesh->GetStringField(TEXT("uri"));
+		PreBakedMeshesPaths.Add(MeshUri);
+	}
+
+	const FString RecipeName = AssetImportFolder + "_Recipe";
+	const FString RecipePackageName = ImportLocation / RecipeName;
+
+	UPackage* RecipePackage = CreatePackage(*RecipePackageName);
+	UMegaPlantRecipeDataAsset* RecipeAsset = NewObject<UMegaPlantRecipeDataAsset>(RecipePackage, *RecipeName, RF_Public | RF_Standalone);
+
+	TMap<FString, FString> VariantFiles;
+	for (const auto& Variation : RecipeVariations)
+	{
+		const TSharedPtr<FJsonObject>& VariationObject = Variation->AsObject();
+		const FString& VariationUri = VariationObject->GetStringField(TEXT("uri"));
+		const FString& VariationName = VariationObject->GetStringField(TEXT("name"));
+		const FString VariationPath = RecipeDownloadFolder / VariationUri;
+		VariantFiles.Add(VariationName, VariationPath);
+	}
+	RecipeAsset->LoadFromVariantFiles(VariantFiles);
+
+	IAssetRegistry::Get()->AssetCreated(RecipeAsset);
+	RecipePackage->SetDirtyFlag(true);
+
+	const FString MegaPlantName = AssetImportFolder + "_Graph";
+	const FString MegaPlantPackageName = ImportLocation / MegaPlantName;
+	UPackage* MegaPlantPackage = CreatePackage(*MegaPlantPackageName);
+	UMegaPlant* MegaPlant = NewObject<UMegaPlant>(MegaPlantPackage, *MegaPlantName, RF_Public | RF_Standalone);
+
+	IAssetRegistry::Get()->AssetCreated(MegaPlant);
+
+	const TArray<TSharedPtr<FJsonValue>>& RecipeTrunkMaterials = RecipeResources->GetArrayField(TEXT("trunkMaterials"));
+	// TODO: Add option to have materials per generation
+	const FString& TrunkMaterialName = RecipeTrunkMaterials[0]->AsObject()->GetStringField(TEXT("name"));
+
+	const TArray<TSharedPtr<FJsonValue>>& RecipeFoliage = RecipeResources->GetArrayField(TEXT("foliage"));
+	for (const TSharedPtr<FJsonValue>& Foliage : RecipeFoliage)
+	{
+		const TSharedPtr<FJsonObject>& FoliageObject = Foliage->AsObject();
+		const FString& FoliageUri = FoliageObject->GetStringField(TEXT("uri"));
+		const FString FoliagePath = RecipeDownloadFolder / FoliageUri;
+
+		const FString FoliageImportLocation = ImportLocation / "Foliage";
+		for (auto& [VariantName, VariantCollection] : RecipeAsset->Variants)
+		{
+			MegaPlants::Facades::FFoliageFacade FoliageFacade(VariantCollection);
+			FoliageFacade.SetFoliagePath(FoliageImportLocation);
+
+			MegaPlants::Facades::FBranchFacade BranchFacade(VariantCollection);
+			BranchFacade.SetTrunkMaterialPath(FoliageImportLocation / TrunkMaterialName + '.' + TrunkMaterialName);
+		}
+
+		FMegaPlantImporter::ImportFoliage(
+			FoliagePath,
+			FoliageImportLocation,
+			[this, MegaPlant, RecipeAsset, PreBakedMeshesPaths](const TArray<UObject*>& InImportedObjects)
+				{
+					AsyncTask(
+						ENamedThreads::GameThread,
+						[this, MegaPlant, RecipeAsset]()
+							{
+#if WITH_EDITOR
+								MegaPlant->CreateGraphFromRecipe(RecipeAsset);
+								MegaPlant->MarkPackageDirty();
+#endif
+								ExpireImportNotification(true);
+								CompleteWorkflow();
+							}
+					);
+					for (auto& [VariantName, VariantCollection] : RecipeAsset->Variants)
+					{
+						MegaPlants::Facades::FFoliageFacade FoliageFacade(VariantCollection);
+						for (int32 FoliageIndex = 0; FoliageIndex < FoliageFacade.NumFoliageNames(); FoliageIndex++)
+						{
+							UObject* const* FoundFoliageMesh = InImportedObjects.FindByPredicate([&](const UObject* ImportedObject)
+								{
+									return ImportedObject->GetName().Contains(FoliageFacade.GetFoliageName(FoliageIndex));
+								});
+							if (FoundFoliageMesh)
+							{
+								FoliageFacade.SetFoliageName(FoliageIndex, (*FoundFoliageMesh)->GetPathName());
+							}
+						}
+					}
+
+					TArray<FString> ScanFiles;
+					for (const FString& MeshUri : PreBakedMeshesPaths)
+					{
+						const FString DestPath = FPaths::ConvertRelativePathToFull(FPaths::ProjectContentDir()) / ImportLocation.RightChop(6) /
+							FPaths::GetCleanFilename(MeshUri);
+						IFileManager::Get().Copy(*DestPath, *MeshUri);
+						ScanFiles.Add(DestPath);
+					}
+
+					IAssetRegistry::Get()->ScanFilesSynchronous(ScanFiles, true);
+				}
+		);
+	}
+}
+
+void FMegaPlantImportWorkflow::OnContentDownloadProgress(const FFabDownloadRequest* Request, const FFabDownloadStats& DownloadStats)
+{
+	SetDownloadNotificationProgress(DownloadStats.PercentComplete);
+}
+
+void FMegaPlantImportWorkflow::OnContentDownloadComplete(const FFabDownloadRequest* Request, const FFabDownloadStats& DownloadStats)
+{
+	if (!DownloadStats.bIsSuccess)
+	{
+		UE_LOG(LogTemp, Error, TEXT("Failed to download FAB Asset %s"), *AssetName);
+		ExpireDownloadNotification(false);
+		CancelWorkflow();
+		return;
+	}
+
+	const FString& ZipArchive = DownloadStats.DownloadedFiles[0];
+	const FString ExtractLocation = FPaths::GetBaseFilename(ZipArchive, false) + "_extracted";
+	if (!FAssetUtils::Unzip(ZipArchive, ExtractLocation))
+	{
+		UE_LOG(LogTemp, Error, TEXT("Failed to unzip MegaPlant Asset %s"), *AssetId);
+		ExpireDownloadNotification(false);
+		CancelWorkflow();
+		return;
+	}
+
+	IPlatformFile& FileManager = FPlatformFileManager::Get().GetPlatformFile();
+
+	TArray<FString> ImportFiles;
+	FileManager.FindFiles(ImportFiles, *ExtractLocation, TEXT(".json"));
+	ImportFiles = ImportFiles.FilterByPredicate([](const FString& File) { return File.EndsWith("_recipe.json"); });
+
+	ExpireDownloadNotification(true);
+	ImportContent(ImportFiles);
+}
diff --git a/Engine/Plugins/Experimental/ProceduralVegetationEditor/Source/MegaPlantsEditor/Private/Importer/MegaPlantImportWorkflow.h b/Engine/Plugins/Experimental/ProceduralVegetationEditor/Source/MegaPlantsEditor/Private/Importer/MegaPlantImportWorkflow.h
new file mode 100644
--- /dev/null
+++ b/Engine/Plugins/Experimental/ProceduralVegetationEditor/Source/MegaPlantsEditor/Private/Importer/MegaPlantImportWorkflow.h
@@ -0,0 +1,23 @@
+// Copyright Epic Games, Inc. All Rights Reserved.
+
+#pragma once
+
+#include "Workflows/FabWorkflow.h"
+
+class FMegaPlantImportWorkflow : public IFabWorkflow
+{
+public:
+	FMegaPlantImportWorkflow(const FString& InAssetId, const FString& InAssetName, const FString& InDownloadUrl);
+
+	virtual void Execute() override;
+
+protected:
+	virtual void DownloadContent() override;
+	virtual void ImportContent(const TArray<FString>& SourceFiles) override;
+
+	virtual void OnContentDownloadProgress(const FFabDownloadRequest* Request, const FFabDownloadStats& DownloadStats) override;
+	virtual void OnContentDownloadComplete(const FFabDownloadRequest* Request, const FFabDownloadStats& DownloadStats) override;
+
+private:
+	TSharedPtr<FFabDownloadRequest> DownloadRequest;
+};
diff --git a/Engine/Plugins/Experimental/ProceduralVegetationEditor/Source/MegaPlantsEditor/Private/Importer/MegaPlantImportWorkflowFactory.cpp b/Engine/Plugins/Experimental/ProceduralVegetationEditor/Source/MegaPlantsEditor/Private/Importer/MegaPlantImportWorkflowFactory.cpp
new file mode 100644
--- /dev/null
+++ b/Engine/Plugins/Experimental/ProceduralVegetationEditor/Source/MegaPlantsEditor/Private/Importer/MegaPlantImportWorkflowFactory.cpp
@@ -0,0 +1,16 @@
+// Copyright Epic Games, Inc. All Rights Reserved.
+
+#include "MegaPlantImportWorkflowFactory.h"
+
+#include "MegaPlantImportWorkflow.h"
+
+const TArray<FString>& FMegaPlantImportWorkflowFactory::GetImportAssetTypes()
+{
+	static TArray<FString> ImportAssetTypes{"megaplant"};
+	return ImportAssetTypes;
+}
+
+TSharedPtr<IFabWorkflow> FMegaPlantImportWorkflowFactory::Create(const FFabAssetMetadata& InImportAssetMetadata, const FString& InDownloadUrl)
+{
+	return MakeShared<FMegaPlantImportWorkflow>(InImportAssetMetadata.AssetId, InImportAssetMetadata.AssetName, InDownloadUrl);
+}
diff --git a/Engine/Plugins/Experimental/ProceduralVegetationEditor/Source/MegaPlantsEditor/Private/Importer/MegaPlantImportWorkflowFactory.h b/Engine/Plugins/Experimental/ProceduralVegetationEditor/Source/MegaPlantsEditor/Private/Importer/MegaPlantImportWorkflowFactory.h
new file mode 100644
--- /dev/null
+++ b/Engine/Plugins/Experimental/ProceduralVegetationEditor/Source/MegaPlantsEditor/Private/Importer/MegaPlantImportWorkflowFactory.h
@@ -0,0 +1,11 @@
+// Copyright Epic Games, Inc. All Rights Reserved.
+
+#pragma once
+#include "Workflows/FabWorkflowFactory.h"
+
+class FMegaPlantImportWorkflowFactory: public IFabWorkflowFactory
+{
+public:
+	virtual const TArray<FString>& GetImportAssetTypes() override;
+	virtual TSharedPtr<IFabWorkflow> Create(const FFabAssetMetadata& InImportAssetMetadata, const FString& InDownloadUrl = "") override;
+};
diff --git a/Engine/Plugins/Experimental/ProceduralVegetationEditor/Source/MegaPlantsEditor/Private/Importer/MegaPlantImporter.cpp b/Engine/Plugins/Experimental/ProceduralVegetationEditor/Source/MegaPlantsEditor/Private/Importer/MegaPlantImporter.cpp
new file mode 100644
--- /dev/null
+++ b/Engine/Plugins/Experimental/ProceduralVegetationEditor/Source/MegaPlantsEditor/Private/Importer/MegaPlantImporter.cpp
@@ -0,0 +1,141 @@
+// Copyright Epic Games, Inc. All Rights Reserved.
+
+#include "MegaPlantImporter.h"
+
+#include "AssetImportTask.h"
+#include "IAssetTools.h"
+#include "InterchangeGenericAssetsPipeline.h"
+#include "InterchangeGenericMaterialPipeline.h"
+#include "InterchangeGenericMeshPipeline.h"
+#include "InterchangeGenericTexturePipeline.h"
+#include "InterchangeProjectSettings.h"
+#include "MegaPlantImportPipeline.h"
+
+#include "Async/Async.h"
+
+UObject* FMegaPlantImporter::GetImportOptions(const FString& SourceFile, UObject* const OptionsOuter)
+{
+	check(OptionsOuter);
+
+	const UInterchangeManager& InterchangeManager = UInterchangeManager::GetInterchangeManager();
+	const UInterchangeSourceData* SourceData = UInterchangeManager::CreateSourceData(SourceFile);
+	if (InterchangeManager.IsInterchangeImportEnabled() && InterchangeManager.CanTranslateSourceData(SourceData))
+	{
+		const FName PipelineStackName = FInterchangeProjectSettingsUtils::GetDefaultPipelineStackName(false, *SourceData);
+		const FInterchangeImportSettings& InterchangeImportSettings = FInterchangeProjectSettingsUtils::GetDefaultImportSettings(false);
+		if (const FInterchangePipelineStack* const PipelineStack = InterchangeImportSettings.PipelineStacks.Find(PipelineStackName))
+		{
+			const TArray<FSoftObjectPath>* Pipelines = &PipelineStack->Pipelines;
+			UE::Interchange::FScopedTranslator ScopedTranslator(SourceData);
+			for (const FInterchangeTranslatorPipelines& TranslatorPipelines : PipelineStack->PerTranslatorPipelines)
+			{
+				const UClass* TranslatorClass = TranslatorPipelines.Translator.LoadSynchronous();
+				if (ScopedTranslator.GetTranslator() && ScopedTranslator.GetTranslator()->IsA(TranslatorClass))
+				{
+					Pipelines = &TranslatorPipelines.Pipelines;
+					break;
+				}
+			}
+
+			UInterchangePipelineStackOverride* StackOverride = NewObject<UInterchangePipelineStackOverride>(OptionsOuter);
+
+			for (const FSoftObjectPath& Pipeline : *Pipelines)
+			{
+				UInterchangePipelineBase* const DefaultPipeline = Cast<UInterchangePipelineBase>(Pipeline.TryLoad());
+				if (!DefaultPipeline)
+				{
+					continue;
+				}
+				if (UInterchangePipelineBase* GeneratedPipeline = UE::Interchange::GeneratePipelineInstance(Pipeline))
+				{
+					GeneratedPipeline->TransferAdjustSettings(DefaultPipeline);
+					GeneratedPipeline->AddToRoot();
+					if (UInterchangeGenericAssetsPipeline* const GenericAssetsPipeline = Cast<UInterchangeGenericAssetsPipeline>(GeneratedPipeline))
+					{
+						GenericAssetsPipeline->bSceneNameSubFolder = false;
+						GenericAssetsPipeline->bAssetTypeSubFolders = false;
+						GenericAssetsPipeline->MeshPipeline->bImportStaticMeshes = true;
+						GenericAssetsPipeline->MeshPipeline->bImportSkeletalMeshes = true;
+						GenericAssetsPipeline->MeshPipeline->bCombineStaticMeshes = false;
+						GenericAssetsPipeline->MeshPipeline->SkeletalMeshImportContentType = EInterchangeSkeletalMeshContentType::All;
+						GenericAssetsPipeline->MeshPipeline->bGenerateLightmapUVs = true;
+						GenericAssetsPipeline->MeshPipeline->bBuildNanite = true;
+						GenericAssetsPipeline->MaterialPipeline->bImportMaterials = true;
+						GenericAssetsPipeline->MaterialPipeline->TexturePipeline->bImportTextures = true;
+
+						GenericAssetsPipeline->MaterialPipeline->MaterialImport = EInterchangeMaterialImportOption::ImportAsMaterialInstances;
+						GenericAssetsPipeline->CommonMeshesProperties->bRecomputeNormals = false;
+						GenericAssetsPipeline->CommonMeshesProperties->bComputeWeightedNormals = false;
+						GenericAssetsPipeline->CommonMeshesProperties->VertexColorImportOption = EInterchangeVertexColorImportOption::IVCIO_Replace;
+					}
+					if (UInterchangeGenericTexturePipeline* const GenericTexturePipeline = Cast<
+						UInterchangeGenericTexturePipeline>(GeneratedPipeline))
+					{
+						GenericTexturePipeline->bAllowNonPowerOfTwo = true;
+						GenericTexturePipeline->bDetectNormalMapTexture = true;
+					}
+
+					StackOverride->OverridePipelines.Add(GeneratedPipeline);
+				}
+			}
+
+			UMegaPlantImportPipeline* MegaPlantImportPipeline = NewObject<UMegaPlantImportPipeline>();
+			MegaPlantImportPipeline->AddToRoot();
+			StackOverride->OverridePipelines.Add(MegaPlantImportPipeline);
+
+			return StackOverride;
+		}
+		return nullptr;
+	}
+
+	return nullptr;
+}
+
+void FMegaPlantImporter::CleanImportOptions(UObject* const Options)
+{
+	if (const UInterchangePipelineStackOverride* const InterchangeOptions = Cast<UInterchangePipelineStackOverride>(Options))
+	{
+		for (const FSoftObjectPath& OverridePipeline : InterchangeOptions->OverridePipelines)
+		{
+			if (UObject* const LoadedPipeline = OverridePipeline.TryLoad())
+				LoadedPipeline->RemoveFromRoot();
+		}
+	}
+}
+
+void FMegaPlantImporter::ImportFoliage(const FString& Source, const FString& Destination, const TFunction<void(const TArray<UObject*>&)>& Callback)
+{
+	UAssetImportTask* MeshImportTask = NewObject<UAssetImportTask>();
+	MeshImportTask->AddToRoot();
+
+	MeshImportTask->bAutomated = true;
+	MeshImportTask->bSave = false;
+	MeshImportTask->bAsync = true;
+	MeshImportTask->Filename = Source;
+	
+	MeshImportTask->DestinationPath = Destination;
+	MeshImportTask->bReplaceExisting = true;
+	MeshImportTask->Options = GetImportOptions(Source, MeshImportTask);
+
+	IAssetTools::Get().ImportAssetTasks({MeshImportTask});
+
+	Async(
+		EAsyncExecution::Thread,
+		[MeshImportTask, Callback]()
+		{
+			TArray<UObject*> ImportedObjects;
+			if (MeshImportTask->AsyncResults.IsValid())
+			{
+				FPlatformProcess::ConditionalSleep([=]() { return MeshImportTask->IsAsyncImportComplete(); }, 0.25f);
+				ImportedObjects.Append(MeshImportTask->AsyncResults->GetImportedObjects());
+			}
+			else
+			{
+				ImportedObjects.Append(MeshImportTask->GetObjects());
+			}
+			Callback(ImportedObjects);
+			CleanImportOptions(MeshImportTask->Options);
+			MeshImportTask->RemoveFromRoot();
+		}
+	);
+}
diff --git a/Engine/Plugins/Experimental/ProceduralVegetationEditor/Source/MegaPlantsEditor/Private/Importer/MegaPlantImporter.h b/Engine/Plugins/Experimental/ProceduralVegetationEditor/Source/MegaPlantsEditor/Private/Importer/MegaPlantImporter.h
new file mode 100644
--- /dev/null
+++ b/Engine/Plugins/Experimental/ProceduralVegetationEditor/Source/MegaPlantsEditor/Private/Importer/MegaPlantImporter.h
@@ -0,0 +1,15 @@
+// Copyright Epic Games, Inc. All Rights Reserved.
+
+#pragma once
+
+#include "CoreMinimal.h"
+
+class FMegaPlantImporter
+{
+private:
+	static UObject* GetImportOptions(const FString& SourceFile, UObject* const OptionsOuter);
+	static void CleanImportOptions(UObject* const Options);
+
+public:
+	static void ImportFoliage(const FString& Source, const FString& Destination, const TFunction<void(const TArray<UObject*>&)>& Callback);
+};
diff --git a/Engine/Plugins/Experimental/ProceduralVegetationEditor/Source/MegaPlantsEditor/Private/MegaPlantEditor.cpp b/Engine/Plugins/Experimental/ProceduralVegetationEditor/Source/MegaPlantsEditor/Private/MegaPlantEditor.cpp
new file mode 100644
--- /dev/null
+++ b/Engine/Plugins/Experimental/ProceduralVegetationEditor/Source/MegaPlantsEditor/Private/MegaPlantEditor.cpp
@@ -0,0 +1,423 @@
+﻿// Copyright Epic Games, Inc. All Rights Reserved.
+
+#include "MegaPlantEditor.h"
+#include "EngineAnalytics.h"
+#include "MegaPlantEditorCommands.h"
+#include "Editor/PCGEditorDefaultExecutionSource.h"
+#include "Subsystems/PCGEditorSubsystem.h"
+#include "MegaPlantEditorSchema.h"
+#include "MegaPlant.h"
+#include "PCGMegaPlantData.h"
+#include "Nodes/PCGMegaPlantOutput.h"
+#include "ToolMenus.h"
+#include "ToolMenuSection.h"
+#include "Dataflow/DataflowCollectionSpreadSheetWidget.h"
+#include "PCGEditor/Private/Nodes/PCGEditorGraphNodeBase.h"
+#include "Widgets/SMegaPlantCreateDialog.h"
+#include "Widgets/SMegaPlantExportSelectionDialog.h"
+
+#if MEGAPLANT_DEBUG_SHOW_COLLECTION_TAB
+#include "Widgets/Docking/SDockTab.h"
+#endif
+
+#define LOCTEXT_NAMESPACE "MegaPlantEditor"
+
+#if MEGAPLANT_DEBUG_SHOW_COLLECTION_TAB
+const FName FMegaPlantEditor::CollectionSpreadSheetTabId (TEXT("MegaPlantEditor_CollectionSpreadSheet"));
+#endif
+
+FMegaPlantEditor::FMegaPlantEditor()
+{
+}
+
+void FMegaPlantEditor::Initialize(const EToolkitMode::Type InMode, const TSharedPtr<class IToolkitHost>& InToolkitHost, UMegaPlant* InMegaPlant)
+{
+	check(InMegaPlant);
+
+	MegaPlantBeingEdited = InMegaPlant;
+	
+	UMegaPlantGraph* MegaPlantGraph = InMegaPlant->GetGraph();
+
+	check(MegaPlantGraph);
+
+	// Initialize widgets before calling base
+
+#if MEGAPLANT_DEBUG_SHOW_COLLECTION_TAB
+	CollectionSpreadSheetWidget = CreateCollectionSpreadSheetWidget();
+#endif
+	
+	FPCGEditor::Initialize(InMode, InToolkitHost, MegaPlantGraph);
+
+	AddEditingObject(InMegaPlant);
+
+	check(GEditor);
+	
+	UPCGEditorSubsystem* Subsystem = GEditor->GetEditorSubsystem<UPCGEditorSubsystem>();
+
+	check(Subsystem);
+	
+	const FPCGEditorDefaultExecutionSourceParams Params
+	{
+		.Graph = MegaPlantGraph
+	};
+	
+	ExecutionSource = Subsystem->CreateEditorExecutionSource(Params);
+
+	// Set the stack
+	FPCGStack Stack;
+	Stack.PushFrame(ExecutionSource.Get());
+	Stack.PushFrame(MegaPlantGraph);
+	SetStackBeingInspected(Stack);
+
+	// Ask for generation
+	ExecutionSource->Generate();
+
+	SessionStartTime = FDateTime::Now();
+
+	if (FEngineAnalytics::IsAvailable())
+	{
+		FEngineAnalytics::GetProvider().RecordEvent(TEXT("Editor.Usage.MegaPlants.SessionStarted"));
+	}
+}
+
+TSubclassOf<UPCGEditorGraphSchema> FMegaPlantEditor::GetSchemaClass() const
+{
+	return UMegaPlantEditorSchema::StaticClass();
+}
+
+void FMegaPlantEditor::AddReferencedObjects(FReferenceCollector& Collector)
+{
+	FPCGEditor::AddReferencedObjects(Collector);
+
+	Collector.AddReferencedObject(ExecutionSource);
+}
+
+FText FMegaPlantEditor::GetToolkitName() const
+{
+	return GetLabelForObject(MegaPlantBeingEdited);
+}
+
+FName FMegaPlantEditor::GetToolkitFName() const
+{
+	return TEXT("MegaPlantEditor");
+}
+
+FText FMegaPlantEditor::GetBaseToolkitName() const
+{
+	return INVTEXT("MegaPlant Editor");
+}
+
+FText FMegaPlantEditor::GetToolkitToolTipText() const
+{
+	return GetToolTipTextForObject(MegaPlantBeingEdited);
+}
+
+FLinearColor FMegaPlantEditor::GetWorldCentricTabColorScale() const
+{
+	return FLinearColor::Blue;
+}
+
+FString FMegaPlantEditor::GetWorldCentricTabPrefix() const
+{
+	return TEXT("MegaPlant Editor ");
+}
+
+void FMegaPlantEditor::OnClose()
+{
+	if (FEngineAnalytics::IsAvailable())
+	{
+		TArray<FAnalyticsEventAttribute> Attributes;
+		Attributes.Add(FAnalyticsEventAttribute(TEXT("TimeActive.Seconds"),  (FDateTime::Now() - SessionStartTime).GetTotalSeconds()));
+				
+		FEngineAnalytics::GetProvider().RecordEvent(TEXT("Editor.Usage.MegaPlants.SessionEnded"), Attributes);
+	}
+	
+	FPCGEditor::OnClose();
+
+	if (ExecutionSource)
+	{
+		ExecutionSource->Sunset();
+		ExecutionSource = nullptr;
+	}
+}
+
+void FMegaPlantEditor::RegisterTabSpawners(const TSharedRef<class FTabManager>& InTabManager)
+{
+	FPCGEditor::RegisterTabSpawners(InTabManager);
+
+#if MEGAPLANT_DEBUG_SHOW_COLLECTION_TAB
+	InTabManager->RegisterTabSpawner(CollectionSpreadSheetTabId, FOnSpawnTab::CreateSP(this, &FMegaPlantEditor::SpawnTab_CollectionSpreadSheet))
+		.SetDisplayName(LOCTEXT("CollectionSpreadsheetTab", "Collection Spreadsheet"))
+		.SetGroup(WorkspaceMenuCategory.ToSharedRef());
+#endif
+}
+
+void FMegaPlantEditor::UnregisterTabSpawners(const TSharedRef<class FTabManager>& InTabManager)
+{
+	FPCGEditor::UnregisterTabSpawners(InTabManager);
+
+#if MEGAPLANT_DEBUG_SHOW_COLLECTION_TAB
+	InTabManager->UnregisterTabSpawner(CollectionSpreadSheetTabId);
+#endif
+}
+
+IPCGBaseSubsystem* FMegaPlantEditor::GetSubsystem() const
+{
+	return GEditor ? GEditor->GetEditorSubsystem<UPCGEditorSubsystem>() : nullptr;
+}
+
+TAttribute<FGraphAppearanceInfo> FMegaPlantEditor::GetAppearanceInfo() const
+{
+	FGraphAppearanceInfo AppearanceInfo;
+	AppearanceInfo.CornerText = INVTEXT("Megaseed");
+
+	return AppearanceInfo;
+}
+
+TSharedRef<FTabManager::FLayout> FMegaPlantEditor::GetDefaultLayout() const
+{
+	return FTabManager::NewLayout("Standalone_MegaPlantPCGGraphEditor_DefaultLayout_v1.0.4")
+	->AddArea // Main PCG Graph Editor Area
+	(
+		FTabManager::NewPrimaryArea()->SetOrientation(Orient_Vertical)
+		->Split // Top Section - Graph, Data Viewport, HLSL Source Editor, and Details View
+		(
+			FTabManager::NewSplitter()->SetOrientation(Orient_Horizontal)->SetSizeCoefficient(0.9f)
+			->SetSizeCoefficient(0.40f)
+			->Split
+			(
+				FTabManager::NewSplitter()->SetOrientation(Orient_Vertical)
+				->SetSizeCoefficient(0.1f)
+				->Split // Viewport and AVL
+				(
+					FTabManager::NewStack()
+					->SetSizeCoefficient(0.95f)
+					->AddTab(GetPanelID(EPCGEditorPanel::Viewport1), ETabState::OpenedTab)
+					->SetForegroundTab(GetPanelID(EPCGEditorPanel::Viewport1))
+					->SetHideTabWell(true)
+				)
+				->Split // Bottom Section - ALV
+				(FTabManager::NewStack()
+					->SetSizeCoefficient(0.05f)
+					->SetHideTabWell(false)
+					->AddTab(GetPanelID(EPCGEditorPanel::Attributes1), ETabState::ClosedTab)
+
+#if MEGAPLANT_DEBUG_SHOW_COLLECTION_TAB
+					->AddTab(CollectionSpreadSheetTabId, ETabState::ClosedTab)
+#endif
+				)	
+			)
+			->Split // Graph and details panel
+			(
+				FTabManager::NewSplitter()
+				->SetOrientation(Orient_Vertical)
+				->SetSizeCoefficient(0.25f)
+				->Split
+				(
+				FTabManager::NewStack()
+					->SetSizeCoefficient(0.6f)
+					->AddTab(GetPanelID(EPCGEditorPanel::GraphEditor), ETabState::OpenedTab)
+					->SetHideTabWell(true)
+				)
+				->Split
+				(
+				FTabManager::NewStack()
+				->SetSizeCoefficient(0.4f)
+				->AddTab(GetPanelID(EPCGEditorPanel::PropertyDetails1), ETabState::OpenedTab)
+				)
+			)
+		)
+	);
+}
+
+void FMegaPlantEditor::RegisterToolbar() const
+{
+	UToolMenus* ToolMenus = UToolMenus::Get();
+	FName ParentName;
+	const FName ToolbarName = GetToolMenuToolbarName(ParentName);
+
+	UToolMenu* const AssetToolbar = UToolMenus::Get()->ExtendMenu(ToolbarName);
+	FToolMenuSection& Section = AssetToolbar->FindOrAddSection("MegaPlantTools");
+
+	const FMegaPlantEditorCommands& Commands = FMegaPlantEditorCommands::Get();
+	Section.AddEntry(FToolMenuEntry::InitToolBarButton(
+			Commands.Export,
+			TAttribute<FText>(),
+			TAttribute<FText>(),
+			FSlateIcon(FAppStyle::GetAppStyleSetName(), "Icons.Toolbar.Export")));
+}
+
+void FMegaPlantEditor::BindCommands()
+{
+	FPCGEditor::BindCommands();
+
+	const FMegaPlantEditorCommands& EditorCommands = FMegaPlantEditorCommands::Get();
+
+	ToolkitCommands->MapAction(
+		EditorCommands.Export,
+		FExecuteAction::CreateSP(this, &FMegaPlantEditor::OnExport_Clicked));
+}
+
+void FMegaPlantEditor::OnSelectedNodesChanged(const TSet<UObject*>& InNewSelection)
+{
+	FPCGEditor::OnSelectedNodesChanged(InNewSelection);
+
+#if MEGAPLANT_DEBUG_SHOW_COLLECTION_TAB
+	const UPCGNode* PCGNode = nullptr;
+	
+	if (!InNewSelection.IsEmpty())
+	{
+		const UPCGEditorGraphNodeBase* Node = Cast<UPCGEditorGraphNodeBase>(*InNewSelection.CreateConstIterator());
+		const FObjectProperty* ObjectProperty = CastField<FObjectProperty>(UPCGEditorGraphNodeBase::StaticClass()->FindPropertyByName(TEXT("PCGNode")));
+		if (Node && ObjectProperty)
+		{
+			PCGNode = Cast<UPCGNode>(ObjectProperty->GetObjectPropertyValue_InContainer(Node));
+		}
+	}
+
+	UpdateCollectionSpreadSheet(PCGNode);
+#endif
+}
+
+void FMegaPlantEditor::OnExport_Clicked()
+{
+	ExportInternal();
+}
+
+FReply FMegaPlantEditor::ExportInternal()
+{
+	const UMegaPlantGraph* MegaPlantGraph = Cast<UMegaPlantGraph>(GetPCGGraph());
+	
+	TSharedPtr<SMegaPlantExportSelectionDialog> ExportSelectionWidget =
+		SNew(SMegaPlantExportSelectionDialog)
+		.Title(LOCTEXT("Export Selection", "Choose outputs to export"))
+		.Graph(MegaPlantGraph);
+
+	if (ExportSelectionWidget->ShowModal() == EAppReturnType::Ok)
+	{
+		const TArray<UPCGNode*>& Nodes = MegaPlantGraph->GetNodes();
+
+		TArray<const UPCGNode*> NodesToExport;
+		
+		for (const UPCGNode* Node : Nodes)
+		{
+			if(const UPCGMegaPlantOutputSettings* OutputSettings = Cast<UPCGMegaPlantOutputSettings>(Node->GetSettings()))
+			{
+				if(OutputSettings->ExportSettings.bShouldExport)
+				{
+					NodesToExport.Add(Node);
+				}
+			}
+		}
+
+		FScopedSlowTask Progress(NodesToExport.Num(), NSLOCTEXT("ExportAsset", "ExportAsset_SlowTask", "Exporting Asset(s) ..."));
+		Progress.MakeDialog();
+		
+		for(auto Node : NodesToExport)
+		{
+			if(const UPCGMegaPlantOutputSettings* OutputSettings = Cast<UPCGMegaPlantOutputSettings>(Node->GetSettings()))
+			{
+				FManagedArrayCollection OutCollection;
+				FName PinName = NAME_None;
+				if (GetCollectionFromNode(Node, PinName, OutCollection))
+				{
+					MegaPlants::Export::ExportCollectionAsMesh(MegaPlantBeingEdited, OutCollection, OutCollection, OutputSettings->ExportSettings);
+					Progress.EnterProgressFrame();
+				}
+			}
+		}
+	}
+	
+	return FReply::Handled();
+}
+
+bool FMegaPlantEditor::GetCollectionFromNode(const UPCGNode* InNode, FName& OutPinName, FManagedArrayCollection& OutCollection) const
+{
+	const UPCGPin* Pin = nullptr;
+		
+	const TArray<TObjectPtr<UPCGPin>>& Pins = InNode->GetOutputPins();
+
+	if (Pins.IsValidIndex(0))
+	{
+		Pin = Pins[0];
+	}
+		
+	check(Pin)
+		
+	const FPCGStack* PCGStack = GetStackBeingInspected();
+	check(PCGStack)
+		
+	// Create a temporary stack with Node+Pin to query the exact DataCollection
+	FPCGStack Stack = *PCGStack;
+	TArray<FPCGStackFrame>& StackFrames = Stack.GetStackFramesMutable();
+	StackFrames.Reserve(StackFrames.Num() + 2);
+	StackFrames.Emplace(InNode);
+	StackFrames.Emplace(Pin);
+		
+	const FPCGDataCollection* Data = ExecutionSource->GetExecutionState().GetInspection().GetInspectionData(Stack);
+		
+	if(Data != nullptr)
+	{
+		const FPCGTaggedData* TaggedData = Data->TaggedData.GetData();
+		
+		if(TaggedData != nullptr)
+		{
+			const UPCGMegaPlantData* MegaPlantData = Cast<UPCGMegaPlantData>(TaggedData->Data);
+		
+			if(MegaPlantData != nullptr)
+			{
+				OutPinName = Pin->Properties.Label;
+				OutCollection = MegaPlantData->GetCollection();
+				return true;
+			}
+		}
+	}
+
+	return false;
+}
+
+#if MEGAPLANT_DEBUG_SHOW_COLLECTION_TAB
+TSharedRef<SCollectionSpreadSheetWidget> FMegaPlantEditor::CreateCollectionSpreadSheetWidget()
+{
+	return SNew(SCollectionSpreadSheetWidget);
+}
+
+TSharedRef<SDockTab> FMegaPlantEditor::SpawnTab_CollectionSpreadSheet(const FSpawnTabArgs& SpawnTabArgs)
+{
+	check(SpawnTabArgs.GetTabId() == FMegaPlantEditor::CollectionSpreadSheetTabId);
+
+	TSharedRef<SDockTab> DockableTab = SNew(SDockTab)
+		[
+			CollectionSpreadSheetWidget.ToSharedRef()
+		];
+
+	CollectionSpreadSheetTab = DockableTab;
+
+	return DockableTab;
+}
+
+void FMegaPlantEditor::UpdateCollectionSpreadSheet(const UPCGNode* InSelectedNode) const
+{
+	if (CollectionSpreadSheetTab.IsValid() && CollectionSpreadSheetTab.Pin()->GetVisibility().IsVisible() && CollectionSpreadSheetWidget)
+	{
+		if (InSelectedNode)
+		{
+			FManagedArrayCollection OutCollection;
+			FName OutputName;
+			GetCollectionFromNode(InSelectedNode, OutputName, OutCollection);
+			
+			CollectionSpreadSheetWidget->GetCollectionTable()->GetCollectionInfoMap().Empty();
+			CollectionSpreadSheetWidget->GetCollectionTable()->GetCollectionInfoMap().Add(OutputName.ToString(), { OutCollection });
+			CollectionSpreadSheetWidget->SetData(InSelectedNode->GetNodeTitle(EPCGNodeTitleType::FullTitle).ToString());
+		}
+		else
+		{
+			CollectionSpreadSheetWidget->SetData(FString());
+		}
+
+		CollectionSpreadSheetWidget->RefreshWidget();	
+	}
+}
+#endif
+
+#undef LOCTEXT_NAMESPACE
diff --git a/Engine/Plugins/Experimental/ProceduralVegetationEditor/Source/MegaPlantsEditor/Private/MegaPlantEditor.h b/Engine/Plugins/Experimental/ProceduralVegetationEditor/Source/MegaPlantsEditor/Private/MegaPlantEditor.h
new file mode 100644
--- /dev/null
+++ b/Engine/Plugins/Experimental/ProceduralVegetationEditor/Source/MegaPlantsEditor/Private/MegaPlantEditor.h
@@ -0,0 +1,71 @@
+﻿// Copyright Epic Games, Inc. All Rights Reserved.
+
+#pragma once
+
+#include "PCGEditor.h"
+
+class SCollectionSpreadSheetWidget;
+struct FManagedArrayCollection;
+class UPCGEditorDefaultExecutionSource;
+class UMegaPlant;
+
+#define MEGAPLANT_DEBUG_SHOW_COLLECTION_TAB !UE_BUILD_SHIPPING_WITH_EDITOR 
+
+class FMegaPlantEditor : public FPCGEditor
+{
+public:
+	FMegaPlantEditor();
+
+	void Initialize(const EToolkitMode::Type InMode, const TSharedPtr<class IToolkitHost>& InToolkitHost, UMegaPlant* InMegaPlant);
+	virtual TSubclassOf<UPCGEditorGraphSchema> GetSchemaClass() const override;
+
+	virtual void AddReferencedObjects(FReferenceCollector& Collector) override;
+	
+	// ~Begin FAssetEditorToolkit interface
+	virtual FText GetToolkitName() const override;
+	virtual FName GetToolkitFName() const override;
+	virtual FText GetBaseToolkitName() const override;
+	virtual FText GetToolkitToolTipText() const override;
+	virtual FLinearColor GetWorldCentricTabColorScale() const override;
+	virtual FString GetWorldCentricTabPrefix() const override;
+	virtual void OnClose() override;
+	virtual void RegisterTabSpawners(const TSharedRef<class FTabManager>& InTabManager) override;
+	virtual void UnregisterTabSpawners(const TSharedRef<class FTabManager>& InTabManager) override;
+	// ~End FAssetEditorToolkit interface
+
+	virtual IPCGBaseSubsystem* GetSubsystem() const override;
+
+protected:
+	// ~Begin FPCGEditor interface
+	virtual TAttribute<FGraphAppearanceInfo> GetAppearanceInfo() const override;
+	virtual TSharedRef<FTabManager::FLayout> GetDefaultLayout() const override;
+	virtual void RegisterToolbar() const override;
+	virtual void BindCommands() override;
+	virtual void OnSelectedNodesChanged(const TSet<UObject*>& InNewSelection) override;
+	// ~End FPCGEditor interface
+
+	bool bGenerationInProgress = false;
+	TObjectPtr<UPCGEditorDefaultExecutionSource> ExecutionSource;
+
+private:
+	TObjectPtr<UMegaPlant> MegaPlantBeingEdited = nullptr;
+
+	FDateTime SessionStartTime;
+
+#if MEGAPLANT_DEBUG_SHOW_COLLECTION_TAB
+	static const FName CollectionSpreadSheetTabId;
+
+	TWeakPtr<SDockTab> CollectionSpreadSheetTab;
+	TSharedPtr<SCollectionSpreadSheetWidget> CollectionSpreadSheetWidget;
+
+	TSharedRef<SCollectionSpreadSheetWidget> CreateCollectionSpreadSheetWidget();
+
+	TSharedRef<SDockTab> SpawnTab_CollectionSpreadSheet(const FSpawnTabArgs& SpawnTabArgs);
+	
+	void UpdateCollectionSpreadSheet(const UPCGNode* InSelectedNode) const;
+#endif
+	
+	bool GetCollectionFromNode(const UPCGNode* InNode, FName& OutPinName, FManagedArrayCollection& OutCollection) const;
+	void OnExport_Clicked();
+	FReply ExportInternal();
+};
diff --git a/Engine/Plugins/Experimental/ProceduralVegetationEditor/Source/MegaPlantsEditor/Private/MegaPlantEditorCommands.cpp b/Engine/Plugins/Experimental/ProceduralVegetationEditor/Source/MegaPlantsEditor/Private/MegaPlantEditorCommands.cpp
new file mode 100644
--- /dev/null
+++ b/Engine/Plugins/Experimental/ProceduralVegetationEditor/Source/MegaPlantsEditor/Private/MegaPlantEditorCommands.cpp
@@ -0,0 +1,22 @@
+// Copyright Epic Games, Inc. All Rights Reserved.
+
+#include "MegaPlantEditorCommands.h"
+#include "Styling/AppStyle.h"
+
+#define LOCTEXT_NAMESPACE "MegaPlantPCGEditorCommands"
+
+FMegaPlantEditorCommands::FMegaPlantEditorCommands()
+	: TCommands<FMegaPlantEditorCommands>(
+		"MegaPlantPCGEditor",
+		NSLOCTEXT("Contexts", "MegaPlantPCGEditor", "MegaPlant Editor"),
+		NAME_None,
+		FAppStyle::GetAppStyleSetName())
+{
+}
+
+void FMegaPlantEditorCommands::RegisterCommands()
+{
+	UI_COMMAND(Export, "Export", "Export outputs.", EUserInterfaceActionType::Button, FInputChord(EModifierKey::Control, EKeys::E));
+}
+
+#undef LOCTEXT_NAMESPACE
diff --git a/Engine/Plugins/Experimental/ProceduralVegetationEditor/Source/MegaPlantsEditor/Private/MegaPlantEditorCommands.h b/Engine/Plugins/Experimental/ProceduralVegetationEditor/Source/MegaPlantsEditor/Private/MegaPlantEditorCommands.h
new file mode 100644
--- /dev/null
+++ b/Engine/Plugins/Experimental/ProceduralVegetationEditor/Source/MegaPlantsEditor/Private/MegaPlantEditorCommands.h
@@ -0,0 +1,17 @@
+// Copyright Epic Games, Inc. All Rights Reserved.
+
+#pragma once
+
+#include "Framework/Commands/Commands.h"
+
+class FMegaPlantEditorCommands : public TCommands<FMegaPlantEditorCommands>
+{
+public:
+	FMegaPlantEditorCommands();
+
+	// ~Begin TCommands<> interface
+	virtual void RegisterCommands() override;
+	// ~End TCommands<> interface
+
+	TSharedPtr<FUICommandInfo> Export;
+};
\ No newline at end of file
diff --git a/Engine/Plugins/Experimental/ProceduralVegetationEditor/Source/MegaPlantsEditor/Private/MegaPlantEditorSchema.cpp b/Engine/Plugins/Experimental/ProceduralVegetationEditor/Source/MegaPlantsEditor/Private/MegaPlantEditorSchema.cpp
new file mode 100644
--- /dev/null
+++ b/Engine/Plugins/Experimental/ProceduralVegetationEditor/Source/MegaPlantsEditor/Private/MegaPlantEditorSchema.cpp
@@ -0,0 +1,23 @@
+﻿// Copyright Epic Games, Inc. All Rights Reserved.
+
+#include "MegaPlantEditorSchema.h"
+#include "Nodes/PCGMegaPlantBaseSettings.h"
+
+EPCGElementType UMegaPlantEditorSchema::GetElementTypeFiltering() const
+{
+	return EPCGElementType::Native;
+}
+
+const FPCGGraphEditorCustomization& UMegaPlantEditorSchema::GetGraphEditorCustomization(const UEdGraph* InEdGraph) const
+{
+	static const FPCGGraphEditorCustomization GraphEditorCustomization
+	{
+		.bFilterSettings = true,
+		.FilteredSettingsTypes =
+	{
+			UPCGMegaPlantBaseSettings::StaticClass()
+		}
+	};
+
+	return GraphEditorCustomization;
+}
diff --git a/Engine/Plugins/Experimental/ProceduralVegetationEditor/Source/MegaPlantsEditor/Private/MegaPlantEditorSchema.h b/Engine/Plugins/Experimental/ProceduralVegetationEditor/Source/MegaPlantsEditor/Private/MegaPlantEditorSchema.h
new file mode 100644
--- /dev/null
+++ b/Engine/Plugins/Experimental/ProceduralVegetationEditor/Source/MegaPlantsEditor/Private/MegaPlantEditorSchema.h
@@ -0,0 +1,19 @@
+﻿// Copyright Epic Games, Inc. All Rights Reserved.
+
+#pragma once
+
+#include "PCGEditor.h"
+#include "MegaPlantEditorSchema.generated.h"
+
+UCLASS()
+class UMegaPlantEditorSchema : public UPCGEditorGraphSchema
+{
+	GENERATED_BODY()
+
+public:
+	// Return the element type filtering for the current editor. By default it is all.
+	virtual EPCGElementType GetElementTypeFiltering() const override;
+
+	// Return the graph customization for this editor. By default it is the one provided by the graph.
+	virtual const FPCGGraphEditorCustomization& GetGraphEditorCustomization(const UEdGraph* InEdGraph) const override;
+};
diff --git a/Engine/Plugins/Experimental/ProceduralVegetationEditor/Source/MegaPlantsEditor/Private/MegaPlantFactory.cpp b/Engine/Plugins/Experimental/ProceduralVegetationEditor/Source/MegaPlantsEditor/Private/MegaPlantFactory.cpp
new file mode 100644
--- /dev/null
+++ b/Engine/Plugins/Experimental/ProceduralVegetationEditor/Source/MegaPlantsEditor/Private/MegaPlantFactory.cpp
@@ -0,0 +1,49 @@
+// Copyright Epic Games, Inc. All Rights Reserved.
+
+#include "MegaPlantFactory.h"
+#include "Widgets/SMegaPlantCreateDialog.h"
+#include "MegaPlant.h"
+
+#define LOCTEXT_NAMESPACE "MegaPlantFactory"
+
+UMegaPlantFactory::UMegaPlantFactory(const FObjectInitializer& ObjectInitializer)
+: Super(ObjectInitializer)
+{
+	SupportedClass = UMegaPlant::StaticClass();
+	bCreateNew = true;
+	bEditAfterNew = false;
+}
+
+UObject* UMegaPlantFactory::FactoryCreateNew(UClass* InClass, UObject* InParent, FName InName, EObjectFlags InFlags, UObject* InContext,
+	FFeedbackContext* Warn)
+{
+	UMegaPlant* NewMegaPlant;
+	if (SampleMegaPlant)
+	{
+		NewMegaPlant = NewObject<UMegaPlant>(InParent, InClass, InName, InFlags, SampleMegaPlant);
+
+		TObjectPtr<UMegaPlantGraph> SampleGraph = SampleMegaPlant->GetGraph();
+		NewMegaPlant->CreateGraph(SampleGraph);
+		
+	}
+	else
+	{
+		NewMegaPlant = NewObject<UMegaPlant>(InParent, InClass, InName, InFlags, InContext);
+		NewMegaPlant->CreateGraph();	
+	}
+		
+	return NewMegaPlant;
+}
+
+bool UMegaPlantFactory::ConfigureProperties()
+{
+	SampleMegaPlant = nullptr;
+	
+	const bool bPressedOk = SMegaPlantCreateDialog::OpenCreateModal(
+		LOCTEXT("MegaPlantFactoryConfigureProperties", "Create MegaPlant"),
+		SampleMegaPlant
+	);
+	return bPressedOk;
+}
+
+#undef LOCTEXT_NAMESPACE
diff --git a/Engine/Plugins/Experimental/ProceduralVegetationEditor/Source/MegaPlantsEditor/Private/MegaPlantFactory.h b/Engine/Plugins/Experimental/ProceduralVegetationEditor/Source/MegaPlantsEditor/Private/MegaPlantFactory.h
new file mode 100644
--- /dev/null
+++ b/Engine/Plugins/Experimental/ProceduralVegetationEditor/Source/MegaPlantsEditor/Private/MegaPlantFactory.h
@@ -0,0 +1,26 @@
+// Copyright Epic Games, Inc. All Rights Reserved.
+
+#pragma once
+
+#include "Factories/Factory.h"
+#include "MegaPlantFactory.generated.h"
+
+class UMegaPlant;
+
+UCLASS(hidecategories=Object)
+class UMegaPlantFactory : public UFactory
+{
+	GENERATED_BODY()
+
+public:
+	UMegaPlantFactory(const FObjectInitializer& ObjectInitializer);
+	
+	//~UFactory interface
+	virtual UObject* FactoryCreateNew(UClass* InClass, UObject* InParent, FName InName, EObjectFlags InFlags, UObject* InContext, FFeedbackContext* Warn) override;
+	virtual bool ShouldShowInNewMenu() const override { return true; }
+	virtual bool ConfigureProperties() override;
+
+private:
+	UPROPERTY()
+	TObjectPtr<UMegaPlant> SampleMegaPlant = nullptr;
+};
\ No newline at end of file
diff --git a/Engine/Plugins/Experimental/ProceduralVegetationEditor/Source/MegaPlantsEditor/Private/MegaPlantsEditorModule.cpp b/Engine/Plugins/Experimental/ProceduralVegetationEditor/Source/MegaPlantsEditor/Private/MegaPlantsEditorModule.cpp
new file mode 100644
--- /dev/null
+++ b/Engine/Plugins/Experimental/ProceduralVegetationEditor/Source/MegaPlantsEditor/Private/MegaPlantsEditorModule.cpp
@@ -0,0 +1,263 @@
+// Copyright Epic Games, Inc. All Rights Reserved.
+
+#include "MegaPlantsEditorModule.h"
+#include "ContentBrowserMenuContexts.h"
+#include "FabWorkflowFactoryRegistry.h"
+#include "LevelEditor.h"
+#include "MegaPlantLink.h"
+#include "PCGModule.h"
+#include "Importer/MegaPlantImportWorkflowFactory.h"
+#include "Modules/ModuleManager.h"
+#include "MegaPlantEditorCommands.h"
+#include "PCGMegaPlantData.h"
+#include "Selection.h"
+#include "Animation/SkeletalMeshActor.h"
+#include "Visualizations/PCGMegaPlantDataVisualization.h"
+#include "Engine/StaticMeshActor.h"
+#include "ToolMenus.h"
+#include "ToolMenuContext.h"
+#include "ToolMenuDelegates.h"
+#include "ToolMenuSection.h"
+#include "GameFramework/Actor.h"
+#include "Framework/MultiBox/MultiBox.h"
+#include "Engine/StaticMesh.h"
+#include "Engine/SkeletalMesh.h"
+#include "Animation/Skeleton.h"
+#include "Components/StaticMeshComponent.h"
+#include "Components/SkeletalMeshComponent.h"
+#include "Templates/SharedPointer.h"
+
+#define LOCTEXT_NAMESPACE "MegaPlantsEditorModule"
+
+DEFINE_LOG_CATEGORY(LogMegaPlantsEditor);
+
+void FMegaPlantsEditorModule::StartupModule()
+{
+	// This code will execute after your module is loaded into memory; the exact timing is specified in the .uplugin file per-module
+	RegisterMenus();
+	FMegaPlantEditorCommands::Register();
+	RegisterLevelEditorMenuExtension();
+
+	// Visualizations
+	FPCGDataVisualizationRegistry& DataVisRegistry = FPCGModule::GetMutablePCGDataVisualizationRegistry();
+	DataVisRegistry.RegisterPCGDataVisualization(UPCGMegaPlantData::StaticClass(), MakeUnique<const FPCGMegaPlantDataVisualization>());
+
+	// FAB
+	MegaPlantsImportWorkflowFactory = MakeShared<FMegaPlantImportWorkflowFactory>();
+	FFabWorkflowFactoryRegistry::RegisterFactory(MegaPlantsImportWorkflowFactory);
+}
+
+void FMegaPlantsEditorModule::ShutdownModule()
+{
+	UnregisterLevelEditorMenuExtension();
+	FMegaPlantEditorCommands::Unregister();
+	FFabWorkflowFactoryRegistry::UnregisterFactory(MegaPlantsImportWorkflowFactory);
+}
+
+void FMegaPlantsEditorModule::RegisterMenus()
+{
+	// Allows cleanup when module unloads.
+	FToolMenuOwnerScoped OwnerScoped(this);
+	
+	// Extend the content browser context menu for static meshes and skeletal meshes
+	auto AddToContextMenuSection = [this](FToolMenuSection& Section)
+	{
+		Section.AddDynamicEntry("Open in MegaPlants", FNewToolMenuSectionDelegate::CreateLambda(
+			[this](FToolMenuSection& InSection)
+			{
+				if (const UContentBrowserAssetContextMenuContext* Context = UContentBrowserAssetContextMenuContext::FindContextWithAssets(InSection))
+				{
+					{
+						const TAttribute<FText> Label = LOCTEXT("Mesh_OpenInMegaPlantEditor", "Open in MegaPlants editor");
+						const TAttribute<FText> ToolTip = LOCTEXT("Mesh_OpenInMegaPlantEditor", "Open the mesh in the MegaPlant editor.");
+						const FSlateIcon Icon = FSlateIcon(FAppStyle::GetAppStyleSetName(), "Icons.Edit");
+
+						FToolUIAction UIAction;
+						UIAction.ExecuteAction = FToolMenuExecuteAction::CreateRaw(this, &FMegaPlantsEditorModule::ExecuteOpenInMegaPlants);
+						UIAction.CanExecuteAction = FToolMenuCanExecuteAction::CreateRaw(this, &FMegaPlantsEditorModule::CanExecuteOpenInMegaPlants);
+						UIAction.IsActionVisibleDelegate = FToolMenuIsActionButtonVisible::CreateRaw(this, &FMegaPlantsEditorModule::CanExecuteOpenInMegaPlants);
+						
+						InSection.AddMenuEntry("Mesh_OpenInMegaPlantsEditor", Label, ToolTip, Icon, UIAction);
+					}
+				}
+			}));
+	};
+
+	UToolMenu* StaticMeshMenu = UE::ContentBrowser::ExtendToolMenu_AssetContextMenu(UStaticMesh::StaticClass());
+	FToolMenuSection& StaticMeshSection = StaticMeshMenu->FindOrAddSection("GetAssetActions");
+	AddToContextMenuSection(StaticMeshSection);
+
+	UToolMenu* SkeletalMeshMenu = UE::ContentBrowser::ExtendToolMenu_AssetContextMenu(USkeletalMesh::StaticClass());
+	FToolMenuSection& SkeletalMeshSection = SkeletalMeshMenu->FindOrAddSection("GetAssetActions");
+	AddToContextMenuSection(SkeletalMeshSection);
+
+	UToolMenu* SkeletonMenu = UE::ContentBrowser::ExtendToolMenu_AssetContextMenu(USkeleton::StaticClass());
+	FToolMenuSection& SkeletonSection = SkeletonMenu->FindOrAddSection("GetAssetActions");
+	AddToContextMenuSection(SkeletonSection);
+}
+
+void FMegaPlantsEditorModule::ExecuteOpenInMegaPlants(const FToolMenuContext& InContext)
+{
+	if (const UContentBrowserAssetContextMenuContext* Context = UContentBrowserAssetContextMenuContext::FindContextWithAssets(InContext))
+	{
+		const TArray<UObject*> &Objects = Context->LoadSelectedObjects<UObject>();
+		
+		for (auto Iterator = Objects.CreateConstIterator(); Iterator; ++Iterator)
+		{
+			if (UMegaPlant* MegaPlant = GetMegaPlantFromAsset(*Iterator))
+			{
+				GEditor->GetEditorSubsystem<UAssetEditorSubsystem>()->OpenEditorForAsset( MegaPlant );
+			}
+		}
+	}
+}
+
+bool FMegaPlantsEditorModule::CanExecuteOpenInMegaPlants(const FToolMenuContext& InContext)
+{
+	if (const UContentBrowserAssetContextMenuContext* Context = UContentBrowserAssetContextMenuContext::FindContextWithAssets(InContext))
+	{
+		const TArray<UObject*> &Objects = Context->LoadSelectedObjects<UObject>();
+		
+		for (auto Iterator = Objects.CreateConstIterator(); Iterator; ++Iterator)
+		{
+			if (GetMegaPlantFromAsset(*Iterator))
+			{
+				return true;
+			}
+		}
+	}
+
+	return false;
+}
+
+void FMegaPlantsEditorModule::RegisterLevelEditorMenuExtension()
+{
+	FLevelEditorModule& LevelEditorModule = FModuleManager::LoadModuleChecked< FLevelEditorModule >( "LevelEditor" );
+	TArray< FLevelEditorModule::FLevelViewportMenuExtender_SelectedActors >& CBMenuExtenderDelegates = LevelEditorModule.GetAllLevelViewportContextMenuExtenders();
+
+	CBMenuExtenderDelegates.Add( FLevelEditorModule::FLevelViewportMenuExtender_SelectedActors::CreateStatic( &FMegaPlantsEditorModule::OnExtendLevelEditorActorSelectionMenu ) );
+	LevelEditorExtenderDelegateHandle = CBMenuExtenderDelegates.Last().GetHandle();
+}
+
+void FMegaPlantsEditorModule::UnregisterLevelEditorMenuExtension()
+{
+	if ( LevelEditorExtenderDelegateHandle.IsValid() && FModuleManager::Get().IsModuleLoaded( "LevelEditor" ) )
+	{
+		FLevelEditorModule& LevelEditorModule = FModuleManager::GetModuleChecked< FLevelEditorModule >( "LevelEditor" );
+		TArray< FLevelEditorModule::FLevelViewportMenuExtender_SelectedActors >& CBMenuExtenderDelegates = LevelEditorModule.GetAllLevelViewportContextMenuExtenders();
+
+		CBMenuExtenderDelegates.RemoveAll([ this ]( const FLevelEditorModule::FLevelViewportMenuExtender_SelectedActors& Delegate )
+		{
+			return Delegate.GetHandle() == LevelEditorExtenderDelegateHandle;
+		});
+	}
+}
+
+TSharedRef<FExtender> FMegaPlantsEditorModule::OnExtendLevelEditorActorSelectionMenu(const TSharedRef<FUICommandList> InCommandList, TArray<AActor*> InSelectedActors)
+{
+	TSharedRef<FExtender> Extender = MakeShared<FExtender>();
+
+	bool bShouldExtendActorActions = false;
+
+	for ( AActor* Actor : InSelectedActors )
+	{
+		if (UMegaPlant* MegaPlant = GetMegaPlantFromActor(Actor))
+		{
+			bShouldExtendActorActions = true;
+		}
+
+		if ( bShouldExtendActorActions )
+		{
+			break;
+		}
+	}
+
+	if ( bShouldExtendActorActions )
+	{
+		// Add the MegaPlant actions
+		Extender->AddMenuExtension("ActorGeneral", EExtensionHook::After, nullptr, FMenuExtensionDelegate::CreateLambda(
+			[InSelectedActors](FMenuBuilder& MenuBuilder)
+		{
+			FUIAction Action_OpenInMegaPlant(FExecuteAction::CreateStatic(&FMegaPlantsEditorModule::ExecuteOpenInMegaPlantsOnSelectedActors));
+
+			MenuBuilder.AddMenuEntry(
+				LOCTEXT("MenuExtensionOpenInMegaPlantsEditor", "Open in MegaPlants editor"),
+				LOCTEXT("MenuExtensionOpenInMegaPlantsEditor_Tooltip", "Open mesh in MegaPlants editor"),
+				FSlateIcon(FAppStyle::GetAppStyleSetName(), "Icons.Edit"),
+				Action_OpenInMegaPlant,
+				NAME_None,
+				EUserInterfaceActionType::Button);
+		}));
+	}
+
+	return Extender;
+}
+
+void FMegaPlantsEditorModule::ExecuteOpenInMegaPlantsOnSelectedActors()
+{
+	// Create an array of actors to consider
+	TArray<UObject*> SelectedActors;
+	GEditor->GetSelectedActors()->GetSelectedObjects(AActor::StaticClass(), /*out*/ SelectedActors);
+	
+	for (UObject* SelectedObject : SelectedActors)
+	{
+		if (AActor* SelectedActor = Cast<AActor>(SelectedObject))
+		{
+			if(UMegaPlant* MegaPlant = GetMegaPlantFromActor(SelectedActor))
+			{
+				GEditor->GetEditorSubsystem<UAssetEditorSubsystem>()->OpenEditorForAsset( MegaPlant );
+			}
+		}
+	}
+}
+
+UMegaPlant* FMegaPlantsEditorModule::GetMegaPlantFromAsset(UObject* InObject)
+{
+	if (IInterface_AssetUserData* IAssetUserData = Cast<IInterface_AssetUserData>(InObject))
+	{
+		const UMegaPlantLink* Data = IAssetUserData->GetAssetUserData<UMegaPlantLink>();
+
+		if (Data && Data->Source)
+		{
+			return Data->Source;
+		}
+	}
+
+	return nullptr;
+}
+
+UMegaPlant* FMegaPlantsEditorModule::GetMegaPlantFromActor(AActor* InActor)
+{
+	if (const AStaticMeshActor* StaticMeshActor = Cast<AStaticMeshActor>(InActor))
+	{
+		if (StaticMeshActor && StaticMeshActor->GetStaticMeshComponent())
+		{
+			if (UStaticMesh* StaticMesh = StaticMeshActor->GetStaticMeshComponent()->GetStaticMesh())
+			{
+				if (UMegaPlant* MegaPlant = GetMegaPlantFromAsset(StaticMesh))
+				{
+					return MegaPlant;
+				}	
+			}
+		}
+	}
+	else if (const ASkeletalMeshActor* SkeletalMeshActor = Cast<ASkeletalMeshActor>(InActor))
+	{
+		if (SkeletalMeshActor && SkeletalMeshActor->GetSkeletalMeshComponent())
+		{
+			if (USkeletalMesh* Mesh = SkeletalMeshActor->GetSkeletalMeshComponent()->GetSkeletalMeshAsset())
+			{
+				if (UMegaPlant* MegaPlant = GetMegaPlantFromAsset(Mesh))
+				{
+					return MegaPlant;
+				}	
+			}
+		}
+	}
+
+	return nullptr;
+}
+
+IMPLEMENT_MODULE(FMegaPlantsEditorModule, MegaPlantsEditor)
+
+#undef LOCTEXT_NAMESPACE
diff --git a/Engine/Plugins/Experimental/ProceduralVegetationEditor/Source/MegaPlantsEditor/Private/Visualizations/MegaPlantLineBatchComponent.cpp b/Engine/Plugins/Experimental/ProceduralVegetationEditor/Source/MegaPlantsEditor/Private/Visualizations/MegaPlantLineBatchComponent.cpp
new file mode 100644
--- /dev/null
+++ b/Engine/Plugins/Experimental/ProceduralVegetationEditor/Source/MegaPlantsEditor/Private/Visualizations/MegaPlantLineBatchComponent.cpp
@@ -0,0 +1,109 @@
+// Copyright Epic Games, Inc. All Rights Reserved.
+
+#include "MegaPlantLineBatchComponent.h"
+#include "MeshElementCollector.h"
+#include "PrimitiveDrawInterface.h"
+
+FMegaPlantLineSceneProxy::FMegaPlantLineSceneProxy(const UMegaPlantLineBatchComponent* InComponent)
+	: FPrimitiveSceneProxy(InComponent)
+	, LineComponent(InComponent)
+{}
+
+FPrimitiveViewRelevance FMegaPlantLineSceneProxy::GetViewRelevance(const FSceneView* View) const
+{
+	FPrimitiveViewRelevance ViewRelevance;
+	ViewRelevance.bDrawRelevance = IsShown(View);
+	ViewRelevance.bDynamicRelevance = true;
+	ViewRelevance.bSeparateTranslucency = ViewRelevance.bNormalTranslucency = true;
+	return ViewRelevance;
+}
+
+void FMegaPlantLineSceneProxy::GetDynamicMeshElements(
+	const TArray<const FSceneView*>& Views,
+	const FSceneViewFamily& ViewFamily,
+	uint32 VisibilityMap,
+	FMeshElementCollector& Collector
+) const
+{
+	for (int32 ViewIndex = 0; ViewIndex < Views.Num(); ViewIndex++)
+	{
+		if (VisibilityMap & (1 << ViewIndex))
+		{
+			const FSceneView* View = Views[ViewIndex];
+
+#if UE_ENABLE_DEBUG_DRAWING
+			FPrimitiveDrawInterface* PDI = Collector.GetDebugPDI(ViewIndex);
+#else	
+			FPrimitiveDrawInterface* PDI = Collector.GetPDI(ViewIndex);
+#endif
+
+			LineComponent->Draw(View, PDI);
+		}
+	}
+}
+
+SIZE_T FMegaPlantLineSceneProxy::GetTypeHash() const
+{
+	static size_t UniquePointer;
+	return reinterpret_cast<size_t>(&UniquePointer);
+}
+
+uint32 FMegaPlantLineSceneProxy::GetMemoryFootprint() const { return sizeof *this + GetAllocatedSize(); }
+
+bool FMegaPlantLineSceneProxy::CanBeOccluded() const
+{
+	return false;
+}
+
+UMegaPlantLineBatchComponent::UMegaPlantLineBatchComponent()
+{
+	bAutoActivate = true;
+	bTickInEditor = true;
+	PrimaryComponentTick.bCanEverTick = true;
+	bUseEditorCompositing = true;
+	SetGenerateOverlapEvents(false);
+	bIgnoreStreamingManagerUpdate = true;
+	
+	BBox = FBox(ForceInit);
+	Bounds = FBoxSphereBounds(FVector::ZeroVector, FVector::OneVector, 1);
+}
+
+void UMegaPlantLineBatchComponent::InitBounds()
+{
+	BBox = FBox(ForceInit);
+}
+
+void UMegaPlantLineBatchComponent::AddLine(const FVector& InStartPos, const FVector& InEndPos, const FLinearColor& InColor)
+{
+	Lines.Emplace(InStartPos, InEndPos, InColor);
+
+	BBox += InStartPos;
+	BBox += InEndPos;
+}
+
+void UMegaPlantLineBatchComponent::Draw(const FSceneView* View, FPrimitiveDrawInterface* PDI) const
+{
+	for (const FMegaPlantLineInfo& Line : Lines)
+	{
+		PDI->DrawLine(
+			Line.StartPos,
+			Line.EndPos,
+			Line.Color,
+			0,
+			3.0f,
+			0,
+			true
+		);
+	}
+}
+
+FPrimitiveSceneProxy* UMegaPlantLineBatchComponent::CreateSceneProxy()
+{
+	FMegaPlantLineSceneProxy* Proxy = new FMegaPlantLineSceneProxy(this);
+	return Proxy;
+}
+
+FBoxSphereBounds UMegaPlantLineBatchComponent::CalcBounds(const FTransform& LocalToWorld) const
+{
+	return FBoxSphereBounds(BBox).TransformBy(LocalToWorld); 
+}
diff --git a/Engine/Plugins/Experimental/ProceduralVegetationEditor/Source/MegaPlantsEditor/Private/Visualizations/MegaPlantLineBatchComponent.h b/Engine/Plugins/Experimental/ProceduralVegetationEditor/Source/MegaPlantsEditor/Private/Visualizations/MegaPlantLineBatchComponent.h
new file mode 100644
--- /dev/null
+++ b/Engine/Plugins/Experimental/ProceduralVegetationEditor/Source/MegaPlantsEditor/Private/Visualizations/MegaPlantLineBatchComponent.h
@@ -0,0 +1,66 @@
+// Copyright Epic Games, Inc. All Rights Reserved.
+
+#pragma once
+
+#include "CoreMinimal.h"
+#include "PrimitiveSceneProxy.h"
+#include "PrimitiveViewRelevance.h"
+#include "SceneView.h"
+#include "Components/PrimitiveComponent.h"
+#include "MegaPlantLineBatchComponent.generated.h"
+
+struct FManagedArrayCollection;
+class UMegaPlantLineBatchComponent;
+
+class FMegaPlantLineSceneProxy final : public FPrimitiveSceneProxy
+{
+public:
+	FMegaPlantLineSceneProxy(const UMegaPlantLineBatchComponent* InComponent);
+
+	virtual FPrimitiveViewRelevance GetViewRelevance(const FSceneView* View) const override;
+	virtual void GetDynamicMeshElements(
+		const TArray<const FSceneView*>& Views,
+		const FSceneViewFamily& ViewFamily,
+		uint32 VisibilityMap,
+		FMeshElementCollector& Collector
+	) const override;
+
+	virtual SIZE_T GetTypeHash() const override;
+	virtual uint32 GetMemoryFootprint(void) const override;
+
+	virtual bool CanBeOccluded() const override;
+
+private:
+	const TObjectPtr<const UMegaPlantLineBatchComponent> LineComponent;
+};
+
+USTRUCT()
+struct FMegaPlantLineInfo
+{
+	GENERATED_BODY()
+
+	FVector StartPos;
+	FVector EndPos;
+	FLinearColor Color;
+};
+
+UCLASS(Blueprintable, BlueprintType, meta=(BlueprintSpawnableComponent), MinimalAPI)
+class UMegaPlantLineBatchComponent : public UPrimitiveComponent
+{
+	GENERATED_BODY()
+
+public:
+	UMegaPlantLineBatchComponent();
+
+	void InitBounds();
+	void AddLine(const FVector& InStartPos, const FVector& InEndPos, const FLinearColor& InColor = FLinearColor::White);
+
+	void Draw(const FSceneView* View, FPrimitiveDrawInterface* PDI) const;
+
+	virtual FPrimitiveSceneProxy* CreateSceneProxy() override;
+	virtual FBoxSphereBounds CalcBounds(const FTransform& LocalToWorld) const override;
+
+protected:
+	FBox BBox;
+	TArray<FMegaPlantLineInfo> Lines;
+};
diff --git a/Engine/Plugins/Experimental/ProceduralVegetationEditor/Source/MegaPlantsEditor/Private/Visualizations/MegaPlantSkeletonVisualizerComponent.cpp b/Engine/Plugins/Experimental/ProceduralVegetationEditor/Source/MegaPlantsEditor/Private/Visualizations/MegaPlantSkeletonVisualizerComponent.cpp
new file mode 100644
--- /dev/null
+++ b/Engine/Plugins/Experimental/ProceduralVegetationEditor/Source/MegaPlantsEditor/Private/Visualizations/MegaPlantSkeletonVisualizerComponent.cpp
@@ -0,0 +1,64 @@
+// Copyright Epic Games, Inc. All Rights Reserved.
+
+
+#include "MegaPlantSkeletonVisualizerComponent.h"
+#include "MegaPlantsEditorModule.h"
+#include "PrimitiveDrawInterface.h"
+#include "Facades/MegaPlantBranchFacade.h"
+#include "Facades/MegaPlantPointFacade.h"
+
+UMegaPlantSkeletonVisualizerComponent::UMegaPlantSkeletonVisualizerComponent()
+	: SkeletonCollection(nullptr)
+{
+	bAutoActivate = true;
+	bTickInEditor = true;
+	PrimaryComponentTick.bCanEverTick = true;
+	bUseEditorCompositing = true;
+	SetGenerateOverlapEvents(false);
+	bIgnoreStreamingManagerUpdate = true;
+	
+	BBox = FBox(ForceInit);
+	Bounds = FBoxSphereBounds(FVector::ZeroVector, FVector::OneVector, 1);
+}
+
+const FManagedArrayCollection* UMegaPlantSkeletonVisualizerComponent::GetCollection() const
+{
+	return SkeletonCollection;
+}
+
+void UMegaPlantSkeletonVisualizerComponent::SetCollection(const FManagedArrayCollection* const InSkeletonCollection)
+{
+	if (InSkeletonCollection == nullptr)
+	{
+		return;
+	}
+
+	const MegaPlants::Facades::FBranchFacade BranchFacade(*InSkeletonCollection);
+	const MegaPlants::Facades::FPointFacade PointFacade(*InSkeletonCollection);
+	if (!BranchFacade.IsValid() || !PointFacade.IsValid())
+	{
+		UE_LOG(LogMegaPlantsEditor, Error, TEXT("Failed to setup data viewport, Collection data is not valid for Mega Plants."));
+		return;
+	}
+
+	SkeletonCollection = InSkeletonCollection;
+
+	const TManagedArray<FVector3f>& PointsPositions = PointFacade.GetPositions();
+
+	InitBounds();
+	for (int BranchIndex = 0; BranchIndex < BranchFacade.GetElementCount(); ++BranchIndex)
+	{
+		const TArray<int32>& BranchPoints = BranchFacade.GetPoints(BranchIndex);
+		FVector PreviousPosition = FVector(PointsPositions[BranchPoints[0]]);
+		BBox += PreviousPosition;
+
+		for (int32 BranchPointIndex = 1; BranchPointIndex < BranchPoints.Num(); ++BranchPointIndex)
+		{
+			FVector CurrentPosition = FVector(PointsPositions[BranchPoints[BranchPointIndex]]);
+
+			AddLine(PreviousPosition, CurrentPosition);
+
+			PreviousPosition = CurrentPosition;
+		}
+	}
+}
\ No newline at end of file
diff --git a/Engine/Plugins/Experimental/ProceduralVegetationEditor/Source/MegaPlantsEditor/Private/Visualizations/MegaPlantSkeletonVisualizerComponent.h b/Engine/Plugins/Experimental/ProceduralVegetationEditor/Source/MegaPlantsEditor/Private/Visualizations/MegaPlantSkeletonVisualizerComponent.h
new file mode 100644
--- /dev/null
+++ b/Engine/Plugins/Experimental/ProceduralVegetationEditor/Source/MegaPlantsEditor/Private/Visualizations/MegaPlantSkeletonVisualizerComponent.h
@@ -0,0 +1,24 @@
+// Copyright Epic Games, Inc. All Rights Reserved.
+
+#pragma once
+
+#include "CoreMinimal.h"
+#include "MegaPlantLineBatchComponent.h"
+#include "MegaPlantSkeletonVisualizerComponent.generated.h"
+
+struct FManagedArrayCollection;
+
+UCLASS(Blueprintable, BlueprintType, meta=(BlueprintSpawnableComponent), MinimalAPI)
+class UMegaPlantSkeletonVisualizerComponent : public UMegaPlantLineBatchComponent
+{
+	GENERATED_BODY()
+
+public:
+	UMegaPlantSkeletonVisualizerComponent();
+
+	const FManagedArrayCollection* GetCollection() const;
+	void SetCollection(const FManagedArrayCollection* const InSkeletonCollection);
+
+private:
+	const FManagedArrayCollection* SkeletonCollection;
+};
diff --git a/Engine/Plugins/Experimental/ProceduralVegetationEditor/Source/MegaPlantsEditor/Private/Visualizations/PCGMegaPlantDataVisualization.cpp b/Engine/Plugins/Experimental/ProceduralVegetationEditor/Source/MegaPlantsEditor/Private/Visualizations/PCGMegaPlantDataVisualization.cpp
new file mode 100644
--- /dev/null
+++ b/Engine/Plugins/Experimental/ProceduralVegetationEditor/Source/MegaPlantsEditor/Private/Visualizations/PCGMegaPlantDataVisualization.cpp
@@ -0,0 +1,184 @@
+// Copyright Epic Games, Inc. All Rights Reserved.
+
+#include "PCGMegaPlantDataVisualization.h"
+
+#include "AdvancedPreviewScene.h"
+#include "EditorViewportClient.h"
+#include "GeometryCollectionToDynamicMesh.h"
+#include "PCGComponent.h"
+#include "PCGContext.h"
+#include "PCGDataVisualization.h"
+#include "PCGEditorSettings.h"
+#include "PlanarCut.h"
+#include "Components/DynamicMeshComponent.h"
+#include "Components/InstancedStaticMeshComponent.h"
+#include "Components/LineBatchComponent.h"
+#include "Facades/MegaPlantFoliageFacade.h"
+#include "GeometryCollection/GeometryCollectionObject.h"
+#include "GeometryCollection/Facades/CollectionRenderingFacade.h"
+#include "Helpers/MegaPlantUtilities.h"
+#include "PCGMegaPlantData.h"
+#include "Visualizations/MegaPlantSkeletonVisualizerComponent.h"
+
+#define LOCTEXT_NAMESPACE "PCGManagedArrayCollectionDataVisualization"
+
+DEFINE_LOG_CATEGORY(LogMegaPlantDataVisualization);
+
+FPCGMegaPlantDataVisualization::FPCGMegaPlantDataVisualization()
+{
+	UPCGEditorSettings* PCGEditorSettings = GetMutableDefault<UPCGEditorSettings>();
+	PCGEditorSettings->bAutoFocusViewport = false;
+
+	RegisterVisualizations();
+}
+
+void FPCGMegaPlantDataVisualization::RegisterVisualizations()
+{
+	RenderMap.Add(EMegaPlantRenderType::Skeleton, FMegaPlantRenderCallback::CreateRaw(this, &FPCGMegaPlantDataVisualization::SkeletonRenderer));
+	RenderMap.Add(EMegaPlantRenderType::Mesh, FMegaPlantRenderCallback::CreateRaw(this, &FPCGMegaPlantDataVisualization::MeshRenderer));
+	RenderMap.Add(EMegaPlantRenderType::Foliage, FMegaPlantRenderCallback::CreateRaw(this, &FPCGMegaPlantDataVisualization::FoliageRenderer));
+}
+
+void FPCGMegaPlantDataVisualization::ExecuteDebugDisplay(FPCGContext* Context, const UPCGSettingsInterface* SettingsInterface, const UPCGData* Data,
+                                                         AActor* TargetActor) const
+{
+	const UPCGMegaPlantData* ManagedArrayCollectionData = CastChecked<UPCGMegaPlantData>(Data);
+	UPCGComponent* SourcePCGComponent = Cast<UPCGComponent>(Context->ExecutionSource.Get());
+	if (!ManagedArrayCollectionData || !SourcePCGComponent)
+	{
+		return;
+	}
+}
+
+FPCGSetupSceneFunc FPCGMegaPlantDataVisualization::GetViewportSetupFunc(const UPCGSettingsInterface* SettingsInterface, const UPCGData* Data) const
+{
+	return [this, WeakData=TWeakObjectPtr<const UPCGMegaPlantData>(Cast<UPCGMegaPlantData>(Data))](
+		FPCGSceneSetupParams& InOutParams)
+	{
+		check(InOutParams.Scene);
+
+		if (!WeakData.IsValid())
+		{
+			UE_LOG(LogMegaPlantDataVisualization, Error, TEXT("Failed to setup data viewport, the data was lost or invalid."));
+			return;
+		}
+
+		const FManagedArrayCollection& Collection = WeakData->GetCollection();
+
+		FBoxSphereBounds Bounds(ForceInit);
+
+		for (const EMegaPlantRenderType& RenderType : WeakData->GetRenderType())
+		{
+			if (RenderMap.Contains(RenderType))
+			{
+				InOutParams.Scene->GetLineBatcher()->Flush();
+				RenderMap[RenderType].Execute(Collection, InOutParams, Bounds);
+			}
+		}
+
+		if (!InOutParams.Scene->IsGridEnabled())
+		{
+			InOutParams.Scene->HandleToggleGrid();
+		}
+		if (!InOutParams.Scene->IsEnvironmentEnabled())
+		{
+			InOutParams.Scene->HandleToggleEnvironment();
+		}
+		if (!InOutParams.Scene->IsPostProcessingEnabled())
+		{
+			InOutParams.Scene->HandleTogglePostProcessing();
+		}
+
+		UStaticMesh* MannequinMesh = LoadObject<UStaticMesh>(
+			nullptr,
+			TEXT("/MegaPlants/Mannequin/Viewport_Mannequin_T_Pose.Viewport_Mannequin_T_Pose")
+		);
+
+		constexpr bool bShowMannequin = true;
+		if (bShowMannequin)
+		{
+			UStaticMeshComponent* MannequinComponent = NewObject<UStaticMeshComponent>(GetTransientPackage(), NAME_None, RF_Transient);
+			MannequinComponent->SetStaticMesh(MannequinMesh);
+
+			InOutParams.ManagedResources.Add(MannequinComponent);
+			InOutParams.Scene->AddComponent(MannequinComponent, FTransform(FVector(-Bounds.BoxExtent.X * 2.0, 0, 0)));
+		}
+
+		InOutParams.Scene->SetFloorVisibility(true);
+		InOutParams.Scene->SetFloorOffset(1.0f);
+
+		InOutParams.FocusBounds = FBoxSphereBounds(Bounds.Origin, Bounds.BoxExtent, Bounds.SphereRadius);
+	};
+}
+
+void FPCGMegaPlantDataVisualization::MeshRenderer(const FManagedArrayCollection& InCollection, FPCGSceneSetupParams& InOutParams,
+                                                  FBoxSphereBounds& OutBounds)
+{
+	TRACE_CPUPROFILER_EVENT_SCOPE(FPCGMegaPlantDataVisualization::MeshRenderer);
+
+	const TManagedArray<FString>& MaterialPaths = InCollection.GetAttribute<FString>("MaterialPath", FGeometryCollection::MaterialGroup);
+	TArray<UMaterialInterface*> Materials;
+	for (const FString& MaterialPath : MaterialPaths.GetConstArray())
+	{
+		Materials.Add(LoadObject<UMaterialInterface>(nullptr, MaterialPath));
+	}
+
+	UE::Geometry::FGeometryCollectionToDynamicMeshes Converter;
+	Converter.Init(InCollection, {});
+	for (const auto& MeshInfo : Converter.Meshes)
+	{
+		FDynamicMesh3* DynamicMesh = MeshInfo.Mesh.Get();
+		if (DynamicMesh->VertexCount())
+		{
+			TObjectPtr<UDynamicMeshComponent> DynamicMeshComponent = NewObject<UDynamicMeshComponent>(GetTransientPackage(), NAME_None, RF_Transient);
+			DynamicMeshComponent->SetMesh(MoveTemp(*DynamicMesh));
+			DynamicMeshComponent->SetTangentsType(EDynamicMeshComponentTangentsMode::AutoCalculated);
+			DynamicMeshComponent->ConfigureMaterialSet(Materials);
+			DynamicMeshComponent->UpdateBounds();
+
+			InOutParams.ManagedResources.Add(DynamicMeshComponent);
+			InOutParams.Scene->AddComponent(DynamicMeshComponent, FTransform::Identity);
+
+			OutBounds = OutBounds + DynamicMeshComponent->CalcLocalBounds();
+		}
+	}
+}
+
+void FPCGMegaPlantDataVisualization::SkeletonRenderer(const FManagedArrayCollection& InCollection, FPCGSceneSetupParams& InOutParams,
+                                                      FBoxSphereBounds& OutBounds)
+{
+	TRACE_CPUPROFILER_EVENT_SCOPE(FPCGMegaPlantDataVisualization::SkeletonRenderer);
+
+	const TObjectPtr<UMegaPlantSkeletonVisualizerComponent> VisualizerComponent = NewObject<UMegaPlantSkeletonVisualizerComponent>(
+		GetTransientPackage(), NAME_None, RF_Transient);
+	VisualizerComponent->SetCollection(&InCollection);
+	VisualizerComponent->UpdateBounds();
+
+	InOutParams.ManagedResources.Add(VisualizerComponent);
+	InOutParams.Scene->AddComponent(VisualizerComponent, FTransform::Identity);
+
+	OutBounds = OutBounds + VisualizerComponent->CalcLocalBounds();
+}
+
+void FPCGMegaPlantDataVisualization::FoliageRenderer(const FManagedArrayCollection& InCollection, FPCGSceneSetupParams& InOutParams,
+                                                     FBoxSphereBounds& OutBounds)
+{
+	TRACE_CPUPROFILER_EVENT_SCOPE(FPCGMegaPlantDataVisualization::FoliageRenderer);
+
+	MegaPlants::Facades::FFoliageFacade FoliageFacade(InCollection);
+
+	auto OnComponentCreated = [&InOutParams](UMeshComponent* InInstancedComponent)
+	{
+		InOutParams.ManagedResources.Add(InInstancedComponent);
+		InInstancedComponent->RegisterComponentWithWorld(InOutParams.Scene->GetWorld());
+		InOutParams.Scene->AddComponent(InInstancedComponent, FTransform::Identity);
+	};
+
+	TMap<FString, TObjectPtr<UMeshComponent>> OutInstancedComponentMap;
+
+	MegaPlants::Utilities::AddFoliageInstances(FoliageFacade, nullptr,
+	                                           MegaPlants::Utilities::FFoliageComponentCreatedCallback::CreateLambda(OnComponentCreated),
+	                                           OutInstancedComponentMap);
+}
+
+#undef LOCTEXT_NAMESPACE
diff --git a/Engine/Plugins/Experimental/ProceduralVegetationEditor/Source/MegaPlantsEditor/Private/Visualizations/PCGMegaPlantDataVisualization.h b/Engine/Plugins/Experimental/ProceduralVegetationEditor/Source/MegaPlantsEditor/Private/Visualizations/PCGMegaPlantDataVisualization.h
new file mode 100644
--- /dev/null
+++ b/Engine/Plugins/Experimental/ProceduralVegetationEditor/Source/MegaPlantsEditor/Private/Visualizations/PCGMegaPlantDataVisualization.h
@@ -0,0 +1,35 @@
+// Copyright Epic Games, Inc. All Rights Reserved.
+
+#pragma once
+
+#include "PCGDataVisualization.h"
+#include "DataVisualizations/PCGSpatialDataVisualization.h"
+#include "PCGMegaPlantData.h"
+
+class AActor;
+class UPCGData;
+struct FPCGContext;
+
+DECLARE_DELEGATE_ThreeParams(FMegaPlantRenderCallback, const FManagedArrayCollection&, FPCGSceneSetupParams&, FBoxSphereBounds&);
+
+class FPCGMegaPlantDataVisualization : public IPCGSpatialDataVisualization
+{
+public:
+	FPCGMegaPlantDataVisualization();
+	
+	// ~Begin IPCGDataVisualization interface
+	virtual void ExecuteDebugDisplay(FPCGContext* Context, const UPCGSettingsInterface* SettingsInterface, const UPCGData* Data, AActor* TargetActor) const override;
+	virtual FPCGSetupSceneFunc GetViewportSetupFunc(const UPCGSettingsInterface* SettingsInterface, const UPCGData* Data) const override;
+	// ~End IPCGDataVisualization interface
+
+protected:
+	TMap<EMegaPlantRenderType, FMegaPlantRenderCallback> RenderMap;
+	
+	void RegisterVisualizations();
+
+	void MeshRenderer(const FManagedArrayCollection& InCollection, FPCGSceneSetupParams& InOutParams, FBoxSphereBounds& OutBounds);
+	void SkeletonRenderer(const FManagedArrayCollection& InCollection, FPCGSceneSetupParams& InOutParams, FBoxSphereBounds& OutBounds);
+	void FoliageRenderer(const FManagedArrayCollection& InCollection, FPCGSceneSetupParams& InOutParams, FBoxSphereBounds& OutBounds);
+};
+
+DECLARE_LOG_CATEGORY_EXTERN(LogMegaPlantDataVisualization, Log, All);
\ No newline at end of file
diff --git a/Engine/Plugins/Experimental/ProceduralVegetationEditor/Source/MegaPlantsEditor/Private/Widgets/SMegaPlantCreateDialog.cpp b/Engine/Plugins/Experimental/ProceduralVegetationEditor/Source/MegaPlantsEditor/Private/Widgets/SMegaPlantCreateDialog.cpp
new file mode 100644
--- /dev/null
+++ b/Engine/Plugins/Experimental/ProceduralVegetationEditor/Source/MegaPlantsEditor/Private/Widgets/SMegaPlantCreateDialog.cpp
@@ -0,0 +1,202 @@
+﻿// Copyright Epic Games, Inc. All Rights Reserved.
+
+#include "SMegaPlantCreateDialog.h"
+
+#include "ContentBrowserModule.h"
+#include "Editor.h"
+#include "IContentBrowserDataModule.h"
+#include "IContentBrowserSingleton.h"
+#include "MegaPlant.h"
+#include "SPrimaryButton.h"
+#include "Widgets/Layout/SExpandableArea.h"
+#include "Widgets/Layout/SUniformGridPanel.h"
+
+#define LOCTEXT_NAMESPACE "SMegaPlantCreateDialog"
+
+void SMegaPlantCreateDialog::Construct(const FArguments& InArgs)
+{
+	WeakParentWindow = InArgs._ParentWindow;
+
+	bPressedOk = false;
+	bCreateNew = false;
+
+	TSharedRef<SVerticalBox> Container = SNew(SVerticalBox);
+
+	FName SampleTreesVirtualPath;
+	IContentBrowserDataModule::Get().GetSubsystem()->ConvertInternalPathToVirtual("/MegaPlants/SampleTrees", SampleTreesVirtualPath);
+
+	const FContentBrowserModule& ContentBrowserModule = FModuleManager::LoadModuleChecked<FContentBrowserModule>("ContentBrowser");
+
+	FAssetPickerConfig AssetPickerConfig;
+	AssetPickerConfig.Filter.ClassPaths.Add(UMegaPlant::StaticClass()->GetClassPathName());
+	AssetPickerConfig.Filter.PackagePaths.Add(SampleTreesVirtualPath);
+	AssetPickerConfig.Filter.bRecursivePaths = true;
+
+	AssetPickerConfig.InitialAssetViewType = EAssetViewType::Tile;
+	AssetPickerConfig.SelectionMode = ESelectionMode::Single;
+
+	AssetPickerConfig.bForceShowEngineContent = true;
+	AssetPickerConfig.bForceShowPluginContent = true;
+	AssetPickerConfig.bCanShowReadOnlyFolders = false;
+	AssetPickerConfig.bCanShowFolders = false;
+	AssetPickerConfig.bCanShowDevelopersFolder = false;
+	AssetPickerConfig.bCanShowClasses = false;
+	AssetPickerConfig.bAddFilterUI = false;
+	AssetPickerConfig.bAllowDragging = false;
+	AssetPickerConfig.bAllowRename = false;
+	AssetPickerConfig.OnAssetSelected.BindLambda(
+		[this](const FAssetData& AssetData)
+		{
+			SelectedMegaPlant = AssetData;
+		}
+	);
+
+	Container
+		->AddSlot()
+		.AutoHeight()
+		[
+			SNew(SExpandableArea)
+			.BorderImage(FStyleDefaults::GetNoBrush())
+			.BodyBorderImage(FAppStyle::Get().GetBrush("Brushes.Recessed"))
+			.HeaderPadding(FMargin(5.0f, 3.0f))
+			.AllowAnimatedTransition(false)
+			.HeaderContent()
+			[
+				SNew(STextBlock)
+				.Text(NSLOCTEXT("SMegaPlantCreateDialog", "SampleMegaPlantsAreaTitle", "Sample MegaPlants"))
+				.TextStyle(FAppStyle::Get(), "ButtonText")
+				.Font(FAppStyle::Get().GetFontStyle("NormalFontBold"))
+			]
+			.BodyContent()
+			[
+				ContentBrowserModule.Get().CreateAssetPicker(AssetPickerConfig)
+			]
+		];
+
+	Container
+		->AddSlot()
+		.HAlign(HAlign_Right)
+		.VAlign(VAlign_Bottom)
+		.Padding(8.f)
+		[
+			SNew(SHorizontalBox)
+			+ SHorizontalBox::Slot()
+			.HAlign(HAlign_Left)
+			.AutoWidth()
+			[
+				SNew(SButton)
+				.Text(NSLOCTEXT("SMegaPlantCreateDialog", "MegaPlantCreateNewButton", "Create New MegaPlant"))
+				.HAlign(HAlign_Center)
+				.OnClicked(this, &SMegaPlantCreateDialog::OnMegaPlantCreateNew)
+			]
+			+ SHorizontalBox::Slot()
+			.HAlign(HAlign_Right)
+			[
+				SNew(SUniformGridPanel)
+				.SlotPadding(FAppStyle::GetMargin("StandardDialog.SlotPadding"))
+				+ SUniformGridPanel::Slot(0, 0)
+				[
+					SNew(SPrimaryButton)
+					.Text(NSLOCTEXT("SMegaPlantCreateDialog", "MegaPlantCreateButton", "Create"))
+					.IsEnabled(this, &SMegaPlantCreateDialog::GetCreateButtonEnabled)
+					.OnClicked(this, &SMegaPlantCreateDialog::OnMegaPlantCreate)
+				]
+				+ SUniformGridPanel::Slot(1, 0)
+				[
+					SNew(SButton)
+					.Text(NSLOCTEXT("SMegaPlantCreateDialog", "MegaPlantCancelButton", "Cancel"))
+					.HAlign(HAlign_Center)
+					.OnClicked(this, &SMegaPlantCreateDialog::OnMegaPlantCancelled)
+				]
+			]
+		];
+
+	ChildSlot
+	[
+		SNew(SBorder)
+		.BorderImage(FAppStyle::GetBrush("Brushes.Panel"))
+		[
+			SNew(SBox)
+			.WidthOverride(610.0f)
+			[
+				Container
+			]
+		]
+	];
+}
+
+bool SMegaPlantCreateDialog::OpenCreateModal(const FText& TitleText, TObjectPtr<UMegaPlant>& SampleMegaPlant)
+{
+	TSharedRef<SWindow> CreateWindow = SNew(SWindow)
+	.Title(TitleText)
+	.SizingRule(ESizingRule::Autosized)
+	.ClientSize(FVector2D(0.f, 610.f))
+	.SupportsMaximize(false)
+	.SupportsMinimize(false);
+
+	TSharedRef<SMegaPlantCreateDialog> CreateDialog = SNew(SMegaPlantCreateDialog)
+		.ParentWindow(CreateWindow);
+
+	CreateWindow->SetContent(CreateDialog);
+
+	GEditor->EditorAddModalWindow(CreateWindow);
+
+	if (CreateDialog->bPressedOk)
+	{
+		if (CreateDialog->bCreateNew)
+		{
+			SampleMegaPlant = nullptr;
+			return true;
+		}
+		if (CreateDialog->SelectedMegaPlant.IsValid())
+		{
+			SampleMegaPlant = Cast<UMegaPlant>(CreateDialog->SelectedMegaPlant.GetAsset());
+			return true;
+		}
+	}
+
+	return false;
+}
+
+bool SMegaPlantCreateDialog::GetCreateButtonEnabled() const
+{
+	bool ButtonEnabled = false;
+	if (SelectedMegaPlant.IsValid())
+	{
+		ButtonEnabled = true;
+	}
+	return ButtonEnabled;
+}
+
+FReply SMegaPlantCreateDialog::OnMegaPlantCreate()
+{
+	bPressedOk = true;
+	bCreateNew = false;
+	if (WeakParentWindow.IsValid())
+	{
+		WeakParentWindow.Pin()->RequestDestroyWindow();
+	}
+	return FReply::Handled();
+}
+
+FReply SMegaPlantCreateDialog::OnMegaPlantCreateNew()
+{
+	bPressedOk = true;
+	bCreateNew = true;
+	if (WeakParentWindow.IsValid())
+	{
+		WeakParentWindow.Pin()->RequestDestroyWindow();
+	}
+	return FReply::Handled();
+}
+
+FReply SMegaPlantCreateDialog::OnMegaPlantCancelled()
+{
+	if (WeakParentWindow.IsValid())
+	{
+		WeakParentWindow.Pin()->RequestDestroyWindow();
+	}
+	return FReply::Handled();
+}
+
+#undef LOCTEXT_NAMESPACE
diff --git a/Engine/Plugins/Experimental/ProceduralVegetationEditor/Source/MegaPlantsEditor/Private/Widgets/SMegaPlantCreateDialog.h b/Engine/Plugins/Experimental/ProceduralVegetationEditor/Source/MegaPlantsEditor/Private/Widgets/SMegaPlantCreateDialog.h
new file mode 100644
--- /dev/null
+++ b/Engine/Plugins/Experimental/ProceduralVegetationEditor/Source/MegaPlantsEditor/Private/Widgets/SMegaPlantCreateDialog.h
@@ -0,0 +1,37 @@
+﻿// Copyright Epic Games, Inc. All Rights Reserved.
+
+#pragma once
+
+#include "AssetRegistry/AssetData.h"
+#include "Widgets/SCompoundWidget.h"
+#include "Widgets/SWindow.h"
+
+class UMegaPlant;
+
+class SMegaPlantCreateDialog : public SCompoundWidget
+{
+public:
+	SLATE_BEGIN_ARGS(SMegaPlantCreateDialog)
+		{}
+		SLATE_ARGUMENT(TSharedPtr<SWindow>, ParentWindow)
+	SLATE_END_ARGS()
+
+private:
+	bool GetCreateButtonEnabled() const;
+
+	FReply OnMegaPlantCreate();
+	FReply OnMegaPlantCreateNew();
+	FReply OnMegaPlantCancelled();
+
+public:
+	void Construct(const FArguments& InArgs);
+
+	static bool OpenCreateModal(const FText& TitleText, TObjectPtr<UMegaPlant>& SampleMegaPlant);
+private:
+	TWeakPtr<SWindow> WeakParentWindow;
+
+	bool bPressedOk = false; 
+	bool bCreateNew = false;
+
+	FAssetData SelectedMegaPlant;
+};
diff --git a/Engine/Plugins/Experimental/ProceduralVegetationEditor/Source/MegaPlantsEditor/Private/Widgets/SMegaPlantExportIem.cpp b/Engine/Plugins/Experimental/ProceduralVegetationEditor/Source/MegaPlantsEditor/Private/Widgets/SMegaPlantExportIem.cpp
new file mode 100644
--- /dev/null
+++ b/Engine/Plugins/Experimental/ProceduralVegetationEditor/Source/MegaPlantsEditor/Private/Widgets/SMegaPlantExportIem.cpp
@@ -0,0 +1,50 @@
+// Copyright Epic Games, Inc. All Rights Reserved.
+#include "SMegaPlantExportItem.h"
+#include "Widgets/Layout/SBox.h"
+#include <Widgets/Text/STextBlock.h>
+
+#define LOCTEXT_NAMESPACE "SMegaPlantExportItem"
+
+void SMegaPlantExportItem::Construct(const FArguments& InArgs)
+{
+	OnOutputSelectionChanged = InArgs._OnOutputSelectionChanged;
+	bIsSelected = InArgs._bIsSelected;
+	auto& Name = InArgs._Name;
+
+	auto OnSelectionChanged = [this, Name](ECheckBoxState NewState)
+	{
+		bIsSelected = NewState == ECheckBoxState::Checked;
+		OnOutputSelectionChanged.ExecuteIfBound(Name.ToString(), NewState);
+	};
+
+	ChildSlot
+	[
+		SNew(SHorizontalBox)
+
+		+ SHorizontalBox::Slot()
+		[
+			SNew(SCheckBox)
+			.IsChecked(this, &SMegaPlantExportItem::OnCheckBoxState)
+			.OnCheckStateChanged_Lambda(OnSelectionChanged)
+			.Content()
+			[
+				SNew(SHorizontalBox)
+
+				+ SHorizontalBox::Slot()
+				.VAlign(VAlign_Center)
+				.Padding(5)
+				[
+					SNew(STextBlock)
+					.Text(InArgs._Name)
+				]
+			]
+		]
+	];
+}
+
+ECheckBoxState SMegaPlantExportItem::OnCheckBoxState() const
+{
+	return bIsSelected ? ECheckBoxState::Checked : ECheckBoxState::Unchecked;
+}
+
+#undef LOCTEXT_NAMESPACE
\ No newline at end of file
diff --git a/Engine/Plugins/Experimental/ProceduralVegetationEditor/Source/MegaPlantsEditor/Private/Widgets/SMegaPlantExportItem.h b/Engine/Plugins/Experimental/ProceduralVegetationEditor/Source/MegaPlantsEditor/Private/Widgets/SMegaPlantExportItem.h
new file mode 100644
--- /dev/null
+++ b/Engine/Plugins/Experimental/ProceduralVegetationEditor/Source/MegaPlantsEditor/Private/Widgets/SMegaPlantExportItem.h
@@ -0,0 +1,27 @@
+// Copyright Epic Games, Inc. All Rights Reserved.
+#pragma once
+#include "CoreMinimal.h"
+#include "Widgets/SBoxPanel.h"
+#include "Input/Reply.h"
+#include <Widgets/Input/SCheckBox.h>
+#include "Widgets/SCompoundWidget.h"
+
+class SMegaPlantExportItem : public SCompoundWidget
+{
+	DECLARE_DELEGATE_TwoParams(FOnOutputSelectionChanged, const FString ItemName, ECheckBoxState NewState);
+
+	SLATE_BEGIN_ARGS(SMegaPlantExportItem)
+	{}
+	SLATE_ARGUMENT(FText, Name)
+	SLATE_ARGUMENT(bool, bIsSelected)
+	SLATE_EVENT(FOnOutputSelectionChanged, OnOutputSelectionChanged)
+	SLATE_END_ARGS()
+
+public:
+	void Construct(const FArguments& InArgs);
+	ECheckBoxState OnCheckBoxState() const;
+
+	FOnOutputSelectionChanged OnOutputSelectionChanged;
+private:
+	bool bIsSelected;
+};
diff --git a/Engine/Plugins/Experimental/ProceduralVegetationEditor/Source/MegaPlantsEditor/Private/Widgets/SMegaPlantExportSelectionDialog.cpp b/Engine/Plugins/Experimental/ProceduralVegetationEditor/Source/MegaPlantsEditor/Private/Widgets/SMegaPlantExportSelectionDialog.cpp
new file mode 100644
--- /dev/null
+++ b/Engine/Plugins/Experimental/ProceduralVegetationEditor/Source/MegaPlantsEditor/Private/Widgets/SMegaPlantExportSelectionDialog.cpp
@@ -0,0 +1,131 @@
+// Copyright Epic Games, Inc. All Rights Reserved.
+#include "SMegaPlantExportSelectionDialog.h"
+#include "Widgets/Layout/SSeparator.h"
+#include "SPrimaryButton.h"
+#include "Widgets/Layout/SUniformGridPanel.h"
+#include "Editor.h"
+#include "SMegaPlantExportItem.h"
+#include "MegaPlant.h"
+#include "Nodes/PCGMegaPlantOutput.h"
+#include "Widgets/Layout/SScrollBox.h"
+
+#define LOCTEXT_NAMESPACE "SMegaPlantExportSelectionDialog"
+
+void SMegaPlantExportSelectionDialog::Construct(const FArguments& InArgs)
+{
+	Graph = InArgs._Graph;
+	
+	SWindow::Construct(SWindow::FArguments()
+		.Title(InArgs._Title)
+		.SupportsMinimize(false)
+		.SupportsMaximize(false)
+		.ClientSize(FVector2D(350, 450))
+		[
+			SNew(SVerticalBox)
+			+ SVerticalBox::Slot() // Add user input block
+			.Padding(2, 2, 2, 4)
+			[
+				SNew(SBorder)
+				.BorderImage(FAppStyle::GetBrush("ToolPanel.GroupBorder"))
+				[
+					SNew(SVerticalBox)
+					+ SVerticalBox::Slot()
+					.FillHeight(1)
+					[
+						SAssignNew(ScrollBox,SScrollBox)
+					]
+				]
+			]
+			+ SVerticalBox::Slot()
+			.AutoHeight()
+			.HAlign(HAlign_Right)
+			.VAlign(VAlign_Bottom)
+			.Padding(8.f, 16.f)
+			[
+				SNew(SUniformGridPanel)
+				.MinDesiredSlotWidth(FAppStyle::GetFloat("StandardDialog.MinDesiredSlotWidth"))
+				.MinDesiredSlotHeight(FAppStyle::GetFloat("StandardDialog.MinDesiredSlotHeight"))
+				.SlotPadding(FAppStyle::GetMargin("StandardDialog.SlotPadding"))
+				+ SUniformGridPanel::Slot(0, 0)
+				[
+					SNew(SPrimaryButton)
+					.Text(LOCTEXT("Export", "Export"))
+					.OnClicked(this, &SMegaPlantExportSelectionDialog::OnButtonClick, EAppReturnType::Ok)
+				]
+				+ SUniformGridPanel::Slot(1, 0)
+				[
+					SNew(SButton)
+					.Text(LOCTEXT("Cancel", "Cancel"))
+					.ContentPadding(FAppStyle::GetMargin("StandardDialog.ContentPadding"))
+					.OnClicked(this, &SMegaPlantExportSelectionDialog::OnButtonClick, EAppReturnType::Cancel)
+				]
+			]
+		]);
+	AddExportItems();
+}
+void SMegaPlantExportSelectionDialog::AddExportItems()
+{
+	ScrollBox->ClearChildren();
+	
+	Graph->ForEachNode([this](UPCGNode* Node)
+	{
+		const UPCGMegaPlantOutputSettings* OutputSettings = Cast<UPCGMegaPlantOutputSettings>(Node->GetSettings());
+
+		if (OutputSettings)
+		{
+			FText OutputName = Node->GetNodeTitle(EPCGNodeTitleType::FullTitle);
+	
+			ScrollBox->AddSlot()
+			.Padding(5)
+			[
+				SNew(SMegaPlantExportItem)
+				.Name(OutputName)
+				.OnOutputSelectionChanged(this, &SMegaPlantExportSelectionDialog::OnOutputSelectionChanged)
+				.bIsSelected(OutputSettings->ExportSettings.bShouldExport)
+			];
+			ScrollBox->AddSlot()
+			[
+				SNew(SSeparator)
+				.Thickness(1)
+			];
+		}
+
+		return true;
+	});
+}
+FReply SMegaPlantExportSelectionDialog::OnButtonClick(EAppReturnType::Type ButtonID)
+{
+	UserResponse = ButtonID;
+	if (ButtonID == EAppReturnType::Cancel || ButtonID == EAppReturnType::Ok)
+	{
+		// Only close the window if canceling or if the ok
+		RequestDestroyWindow();
+	}
+	else
+	{
+		// reset the user response in case the window is closed using 'x'.
+		UserResponse = EAppReturnType::Cancel;
+	}
+	return FReply::Handled();
+}
+EAppReturnType::Type SMegaPlantExportSelectionDialog::ShowModal()
+{
+	GEditor->EditorAddModalWindow(SharedThis(this));
+	return UserResponse;
+}
+void SMegaPlantExportSelectionDialog::OnOutputSelectionChanged(const FString ItemName, ECheckBoxState NewState)
+{
+	Graph->ForEachNode([this, &ItemName, &NewState](UPCGNode* Node)
+	{
+		if(ItemName == Node->GetNodeTitle(EPCGNodeTitleType::FullTitle).ToString())
+		{
+			if (UPCGMegaPlantOutputSettings* OutputSettings = Cast<UPCGMegaPlantOutputSettings>(Node->GetSettings()))
+			{
+				OutputSettings->ExportSettings.bShouldExport = NewState ==ECheckBoxState::Checked;
+			}
+		}
+		
+		return true;
+	});
+}
+#undef LOCTEXT_NAMESPACE
\ No newline at end of file
diff --git a/Engine/Plugins/Experimental/ProceduralVegetationEditor/Source/MegaPlantsEditor/Private/Widgets/SMegaPlantExportSelectionDialog.h b/Engine/Plugins/Experimental/ProceduralVegetationEditor/Source/MegaPlantsEditor/Private/Widgets/SMegaPlantExportSelectionDialog.h
new file mode 100644
--- /dev/null
+++ b/Engine/Plugins/Experimental/ProceduralVegetationEditor/Source/MegaPlantsEditor/Private/Widgets/SMegaPlantExportSelectionDialog.h
@@ -0,0 +1,35 @@
+// Copyright Epic Games, Inc. All Rights Reserved.
+#pragma once
+#include "CoreMinimal.h"
+#include "Input/Reply.h"
+#include "Widgets/SWindow.h"
+
+class SScrollBox;
+class UMegaPlantGraph;
+
+class SMegaPlantExportSelectionDialog : public SWindow
+{
+	SLATE_BEGIN_ARGS(SMegaPlantExportSelectionDialog)
+	{}
+	SLATE_ARGUMENT(FText, Title)
+	SLATE_ARGUMENT(const UMegaPlantGraph*, Graph)
+	SLATE_ARGUMENT(TArray<FText>, ExportItems)
+	SLATE_END_ARGS()
+
+public:
+	void Construct(const FArguments& InArgs);
+
+	/** Displays the dialog in a blocking fashion */
+	EAppReturnType::Type ShowModal();
+
+protected:
+	FReply OnButtonClick(EAppReturnType::Type ButtonID);
+	EAppReturnType::Type UserResponse = EAppReturnType::Type::Cancel;
+
+private:
+	TSharedPtr<SScrollBox> ScrollBox;
+	const UMegaPlantGraph* Graph = nullptr;
+
+	void AddExportItems();
+	void OnOutputSelectionChanged(const FString ItemName, ECheckBoxState NewState);
+};
diff --git a/Engine/Plugins/Experimental/ProceduralVegetationEditor/Source/MegaPlantsEditor/Public/MegaPlantsEditorModule.h b/Engine/Plugins/Experimental/ProceduralVegetationEditor/Source/MegaPlantsEditor/Public/MegaPlantsEditorModule.h
new file mode 100644
--- /dev/null
+++ b/Engine/Plugins/Experimental/ProceduralVegetationEditor/Source/MegaPlantsEditor/Public/MegaPlantsEditorModule.h
@@ -0,0 +1,36 @@
+// Copyright Epic Games, Inc. All Rights Reserved.
+#pragma once
+#include "CoreMinimal.h"
+#include "Modules/ModuleInterface.h"
+
+class UMegaPlant;
+class AActor;
+class FUICommandList;
+class FExtender;
+struct FToolMenuContext;
+
+DECLARE_LOG_CATEGORY_EXTERN(LogMegaPlantsEditor, Log, All);
+
+class MEGAPLANTSEDITOR_API FMegaPlantsEditorModule : public IModuleInterface
+{
+public:
+	/** IModuleInterface implementation */
+	virtual void StartupModule() override;
+	virtual void ShutdownModule() override;
+	/** IModuleInterface implementation end */
+
+private:
+	TSharedPtr<class IFabWorkflowFactory> MegaPlantsImportWorkflowFactory;
+	FDelegateHandle LevelEditorExtenderDelegateHandle;
+	
+	void RegisterMenus();
+	void ExecuteOpenInMegaPlants(const FToolMenuContext& InContext);
+	bool CanExecuteOpenInMegaPlants(const FToolMenuContext& InContext);
+
+	void RegisterLevelEditorMenuExtension();
+	void UnregisterLevelEditorMenuExtension();
+	static TSharedRef<FExtender> OnExtendLevelEditorActorSelectionMenu(const TSharedRef<FUICommandList> InCommandList, TArray<AActor*> InSelectedActors);
+	static void ExecuteOpenInMegaPlantsOnSelectedActors();
+	static UMegaPlant* GetMegaPlantFromAsset(UObject* InObject);
+	static UMegaPlant* GetMegaPlantFromActor(AActor* InActor);
+};
