Fix crash if visualizing the "Foliage Distributor" node without first plugging in a Mesh Builder node

diff --git a/Engine/Plugins/Experimental/ProceduralVegetationEditor/Source/MegaPlantsEditor/Private/Visualizations/PCGMegaPlantDataVisualization.cpp b/Engine/Plugins/Experimental/ProceduralVegetationEditor/Source/MegaPlantsEditor/Private/Visualizations/PCGMegaPlantDataVisualization.cpp
--- a/Engine/Plugins/Experimental/ProceduralVegetationEditor/Source/MegaPlantsEditor/Private/Visualizations/PCGMegaPlantDataVisualization.cpp
+++ b/Engine/Plugins/Experimental/ProceduralVegetationEditor/Source/MegaPlantsEditor/Private/Visualizations/PCGMegaPlantDataVisualization.cpp
@@ -116,11 +116,14 @@ void FPCGMegaPlantDataVisualization::MeshRenderer(const FManagedArrayCollection&
 {
 	TRACE_CPUPROFILER_EVENT_SCOPE(FPCGMegaPlantDataVisualization::MeshRenderer);
 
-	const TManagedArray<FString>& MaterialPaths = InCollection.GetAttribute<FString>("MaterialPath", FGeometryCollection::MaterialGroup);
 	TArray<UMaterialInterface*> Materials;
-	for (const FString& MaterialPath : MaterialPaths.GetConstArray())
+	if (InCollection.HasAttribute("MaterialPath", FGeometryCollection::MaterialGroup))
 	{
-		Materials.Add(LoadObject<UMaterialInterface>(nullptr, MaterialPath));
+		const TManagedArray<FString>& MaterialPaths = InCollection.GetAttribute<FString>("MaterialPath", FGeometryCollection::MaterialGroup);
+		for (const FString& MaterialPath : MaterialPaths.GetConstArray())
+		{
+			Materials.Add(LoadObject<UMaterialInterface>(nullptr, MaterialPath));
+		}
 	}
 
 	UE::Geometry::FGeometryCollectionToDynamicMeshes Converter;
