[ProceduralVegetationEditor] Bone Reduction and visualization

diff --git a/Engine/Plugins/Experimental/ProceduralVegetationEditor/Source/MegaPlants/Private/Facades/MegaPlantBoneFacade.cpp b/Engine/Plugins/Experimental/ProceduralVegetationEditor/Source/MegaPlants/Private/Facades/MegaPlantBoneFacade.cpp
--- a/Engine/Plugins/Experimental/ProceduralVegetationEditor/Source/MegaPlants/Private/Facades/MegaPlantBoneFacade.cpp
+++ b/Engine/Plugins/Experimental/ProceduralVegetationEditor/Source/MegaPlants/Private/Facades/MegaPlantBoneFacade.cpp
@@ -5,52 +5,63 @@
 #include "Facades/MegaPlantAttributesNames.h"
 #include "Facades/MegaPlantBranchFacade.h"
 #include "Facades/MegaPlantPointFacade.h"
+#include "GeometryCollection/Facades/CollectionConstraintOverrideFacade.h"
+#include "PCGMegaPlantData.h"
+
 
 class FGeometryCollection;
 
 namespace MegaPlants::Facades
 {
+	struct FConnectingSegment
+	{
+		TArray<FBoneNode> Nodes;
+	};
+	
 	FBoneFacade::FBoneFacade(FManagedArrayCollection& InCollection)
 		: BoneName(InCollection, AttributeNames::BoneName, GroupNames::BonesGroup),
 		BoneParentIndex(InCollection, AttributeNames::BoneParentIndex, GroupNames::BonesGroup),
+		BoneId(InCollection, AttributeNames::BoneId, GroupNames::BonesGroup),
 		BonePose(InCollection, AttributeNames::BonePose, GroupNames::BonesGroup),
 		BonePointIndex(InCollection, AttributeNames::BonePointIndex, GroupNames::BonesGroup),
-		BoneParentPointIndex(InCollection, AttributeNames::BoneParentPointIndex, GroupNames::BonesGroup),
+		BoneAbsolutePosition(InCollection, AttributeNames::BoneAbsolutePosition, GroupNames::BonesGroup),
 		BoneBranchIndex(InCollection, AttributeNames::BoneBranchIndex, GroupNames::BonesGroup),
-		NjordPixelIds(InCollection, AttributeNames::NjordPixelId, GroupNames::VerticesGroup)
+		VertexPointIds(InCollection, AttributeNames::VertexPointIds, GroupNames::VerticesGroup)
 	{}
 
 	FBoneFacade::FBoneFacade(const FManagedArrayCollection& InCollection)
 		: BoneName(InCollection, AttributeNames::BoneName, GroupNames::BonesGroup),
 		BoneParentIndex(InCollection, AttributeNames::BoneParentIndex, GroupNames::BonesGroup),
+		BoneId(InCollection, AttributeNames::BoneId, GroupNames::BonesGroup),
 		BonePose(InCollection, AttributeNames::BonePose, GroupNames::BonesGroup),
 		BonePointIndex(InCollection, AttributeNames::BonePointIndex, GroupNames::BonesGroup),
-		BoneParentPointIndex(InCollection, AttributeNames::BoneParentPointIndex, GroupNames::BonesGroup),
+		BoneAbsolutePosition(InCollection, AttributeNames::BoneAbsolutePosition, GroupNames::BonesGroup),
 		BoneBranchIndex(InCollection, AttributeNames::BoneBranchIndex, GroupNames::BonesGroup),
-		NjordPixelIds(InCollection, AttributeNames::NjordPixelId, GroupNames::VerticesGroup)
+		VertexPointIds(InCollection, AttributeNames::VertexPointIds, GroupNames::VerticesGroup)
 	{}
 
 	void FBoneFacade::DefineSchema(FManagedArrayCollection& InCollection)
 	{
+		InCollection.AddAttribute<int32>(AttributeNames::BoneId, GroupNames::BonesGroup);
 		InCollection.AddAttribute<FString>(AttributeNames::BoneName, GroupNames::BonesGroup);
 		InCollection.AddAttribute<int32>(AttributeNames::BoneParentIndex, GroupNames::BonesGroup);
 		InCollection.AddAttribute<FTransform>(AttributeNames::BonePose, GroupNames::BonesGroup);
 		InCollection.AddAttribute<int32>(AttributeNames::BonePointIndex, GroupNames::BonesGroup);
-		InCollection.AddAttribute<int32>(AttributeNames::BoneParentPointIndex, GroupNames::BonesGroup);
 		InCollection.AddAttribute<int32>(AttributeNames::BoneBranchIndex, GroupNames::BonesGroup);
-		InCollection.AddAttribute<float>(AttributeNames::NjordPixelId, GroupNames::VerticesGroup);
+		InCollection.AddAttribute<int32>(AttributeNames::VertexPointIds, GroupNames::VerticesGroup);
+		InCollection.AddAttribute<FVector3f>(AttributeNames::BoneAbsolutePosition, GroupNames::BonesGroup);
 	}
 
 	bool FBoneFacade::IsValid() const
 	{
-		return NjordPixelIds.IsValid();
+		return VertexPointIds.IsValid();
 	}
 
-	void FBoneFacade::GetNjordPixelIds(TArray<float>& VertexNjordIds) const
+	void FBoneFacade::GetPointIds(TArray<int32>& OutVertexPointIds) const
 	{
-		for (auto VertexBoneID : NjordPixelIds.Get())
+		for (auto VertexBoneID : VertexPointIds.Get())
 		{
-			VertexNjordIds.Add(VertexBoneID);
+			OutVertexPointIds.Add(VertexBoneID);
 		}
 	}
 
@@ -69,14 +80,14 @@ namespace MegaPlants::Facades
 		return BonePose.Get();
 	}
 
-	const TManagedArray<float>& FBoneFacade::GetNjordPixelIds() const
+	const TManagedArray<int32>& FBoneFacade::GetVertexPointIds() const
 	{
-		return NjordPixelIds.Get();
+		return VertexPointIds.Get();
 	}
 
-	TManagedArray<float>& FBoneFacade::ModifyNjordPixelIds()
+	TManagedArray<int32>& FBoneFacade::ModifyVertexPointIds()
 	{
-		return NjordPixelIds.Modify();
+		return VertexPointIds.Modify();
 	}
 
 	TArray<FBoneNode> GetParentBranchBones(const  TArray<FBoneNode>& InBoneNodes, int32 InParentIndex)
@@ -118,12 +129,60 @@ namespace MegaPlants::Facades
 
 		return FoundBone;
 	}
-	
-	void FBoneFacade::CreateBoneData(TArray<FBoneNode>& Bones) const
+
+	void ApplyBoneReduction(TArray<FConnectingSegment>& ConnectingSegments , TArray<FBoneNode>& OutReducedBones, float ReductionStrength)
 	{
-		FPointFacade PointFacade(NjordPixelIds.GetConstCollection());
-		FBranchFacade BranchFacade(NjordPixelIds.GetConstCollection());
+		for (auto& Segment : ConnectingSegments)
+		{
+			check(Segment.Nodes.Num() > 0);
+			float NumBones = FMath::Lerp(Segment.Nodes.Num(), 1.0f, FMath::Clamp(ReductionStrength,0,1));
+
+			for (int32 j = 0; j < NumBones; j++)
+			{
+				OutReducedBones.Emplace(Segment.Nodes[j]);
+			}
+
+			if (NumBones != Segment.Nodes.Num())
+			{
+				int32 LastIndex = Segment.Nodes.Num() - 1;
+				FBoneNode LastBone = Segment.Nodes[LastIndex];
+				
+				int LastOutBoneIndex = OutReducedBones.Num() - 1;
+				FBoneNode& LastOutBone = OutReducedBones[LastOutBoneIndex];
+				FVector RelativePosition = LastOutBone.BoneTransform.GetLocation() + (LastBone.AbsolutePosition - LastOutBone.AbsolutePosition);
+				int32 ParentIndex = LastOutBone.ParentBoneIndex;
+				LastOutBone = LastBone;
+				LastOutBone.BoneTransform = FTransform(RelativePosition);
+				LastOutBone.ParentBoneIndex = ParentIndex;
+			}
+		}
+	}
+
+	void ReindexBones(TArray<FBoneNode>& OutReducedBones)
+	{
+		TMap<int32, int32> OldToNewIndexMap;
+		for (int i=0; i < OutReducedBones.Num(); i++)
+		{
+			FBoneNode& Bone = OutReducedBones[i];
+			OldToNewIndexMap.Add(Bone.BoneIndex, i);
+			Bone.BoneIndex = i;
+		}
 
+		for (auto& Bone : OutReducedBones)
+		{
+			if (int32* NewParentID = OldToNewIndexMap.Find(Bone.ParentBoneIndex))
+			{
+				Bone.ParentBoneIndex = *NewParentID;
+				Bone.BoneName = FName(FString::Printf(TEXT("Bone_%i"), Bone.BoneIndex));
+			}
+		}
+	}
+	
+	void FBoneFacade::CreateBoneData(TArray<FBoneNode>& Bones,const float ReductionStrength)
+	{
+		FPointFacade PointFacade(VertexPointIds.GetConstCollection());
+		FBranchFacade BranchFacade(VertexPointIds.GetConstCollection());
+		
 		const TManagedArray<FVector3f>& Positions = PointFacade.GetPositions();
 		TArray<int32> SortedBranchIndices;
 		BranchFacade.GetSortedBranchIndicesByHierarchy(SortedBranchIndices);
@@ -135,6 +194,7 @@ namespace MegaPlants::Facades
 		}
 
 		int32 RootBranch = SortedBranchIndices[0];
+		
 		check(BranchFacade.GetPoints(RootBranch).Num() > 0);
 		const int32 RootPointIndex = BranchFacade.GetPoints(RootBranch)[0];
 		
@@ -171,6 +231,9 @@ namespace MegaPlants::Facades
 			PreviousPosition = FoundBoneNode ? FVector(FoundBoneNode->AbsolutePosition) : RootPointPosition;
 
 			TArray<int32> ConnectingPoints;
+
+			TArray<FConnectingSegment> ConnectingSegments;
+			FConnectingSegment ConnectingSegment;
 			
 			for (int Idx = 0; Idx < BranchPoints.Num(); Idx++)
 			{
@@ -228,15 +291,117 @@ namespace MegaPlants::Facades
 					FVector Position = FVector(PointPosition) - PreviousPosition;
 					BoneNode.AbsolutePosition = PointPosition;
 					BoneNode.BoneTransform = FTransform{Position};
-					Bones.Emplace(MoveTemp(BoneNode));
+					//Bones.Emplace(MoveTemp(BoneNode));
 					ParentBoneIndex = BoneIndex;
 					PreviousPosition = FVector(PointPosition);
 					BoneIndex++;
 
+					int AddedNodeIndex = ConnectingSegment.Nodes.Emplace(MoveTemp(BoneNode));
+					if (IsConnectingPoint || IsLastPoint)
+					{
+						int AddedIndex = ConnectingSegments.Add(MoveTemp(ConnectingSegment));
+						ConnectingSegment = FConnectingSegment();
+					}
 					//UE_LOG(LogMegaPlants, Warning, TEXT("Creating Bone NjordPixelId, Id{%f} BoneId{%i} PointIndex{%i} ParentBoneIndex{%i} BranchID{%i}"),
 					//	BoneNode.NjordPixelID, BoneNode.BoneIndex, BoneNode.PointIndex,BoneNode.ParentBoneIndex,BoneNode.BranchIndex);
 				}
 			}
+			ApplyBoneReduction(ConnectingSegments, Bones, ReductionStrength);
+		}
+
+		ReindexBones(Bones);
+	}
+
+	FBoneNode* FBoneFacade::FindClosestBone(const FManagedArrayCollection& Collection, TArray<FBoneNode>& Bones, const int32 PointID)
+	{
+		FPointFacade PointFacade(Collection);
+		FBranchFacade BranchFacade(Collection);
+
+		const float NjordValue = PointFacade.GetNjordPixelIndex(PointID);
+		
+		FBoneNode* BoneNode = Bones.FindByPredicate([&NjordValue](const FBoneNode& Node)
+		{
+			return FMath::FloorToInt(Node.NjordPixelID) == FMath::FloorToInt(NjordValue);
+		});
+
+		if (BoneNode)
+		{
+			return BoneNode;
+		}
+
+		int32 BranchIndex = BranchFacade.GetBranchIndexFromPointIndex(PointID);
+
+		if (BranchIndex == -1)
+		{
+			UE_LOG(LogMegaPlants, Warning, TEXT("Unable to find the branch for point index %i"), PointID);
+			return nullptr;
+		}
+		
+		for (FBoneNode& Bone : Bones)
+		{
+			if (Bone.BranchIndex == BranchIndex && Bones.IsValidIndex(Bone.ParentBoneIndex))
+			{
+				FBoneNode ParentBone = Bones[Bone.ParentBoneIndex];
+				if (NjordValue > ParentBone.NjordPixelID && NjordValue <= Bone.NjordPixelID)
+				{
+					return &Bone;
+				}
+			}
+		}
+		
+		return BoneNode;
+	}
+
+	TArray<FBoneNode> FBoneFacade::GetBoneDataFromCollection() const
+	{
+		TArray<FBoneNode> Bones;
+		
+		if (!IsValid())
+		{
+			UE_LOG(LogMegaPlants, Warning, TEXT("Trying to get bones with invalid BoneFacade."));
+			return Bones;
+		}
+
+		const FPointFacade PointFacade = FPointFacade(BoneName.GetConstCollection());
+
+		for (int32 i = 0; i < BoneName.Num(); i++)
+		{
+			FBoneNode BoneNode;
+			BoneNode.BoneName = FName(BoneName[i]);
+			BoneNode.ParentBoneIndex = BoneParentIndex[i];
+			BoneNode.PointIndex = BonePointIndex[i];
+			BoneNode.BranchIndex = BoneBranchIndex[i];
+			BoneNode.NjordPixelID = PointFacade.GetNjordPixelIndex(BoneNode.PointIndex);
+			BoneNode.BoneTransform = BonePose[i];
+			BoneNode.AbsolutePosition = FVector(BoneAbsolutePosition[i]);
+			BoneNode.BoneIndex = BoneId[i];
+			Bones.Add(BoneNode);
+		}
+		
+		return Bones;
+	}
+
+	void FBoneFacade::SetBoneDataToCollection(TArray<FBoneNode>& Bones)
+	{
+		if (!IsValid())
+		{
+			UE_LOG(LogMegaPlants, Warning, TEXT("Cannot Set Bones, FBoneFacade is not valid, Schema not defined."));
+			return;
+		}
+		
+		int32 NumElements = BoneName.AddElements(Bones.Num());
+		
+		int BoneIndex = 0;
+		for (const auto& Bone : Bones)
+		{
+			BoneName.ModifyAt(BoneIndex,Bone.BoneName.ToString());
+			BoneParentIndex.ModifyAt(BoneIndex,Bone.ParentBoneIndex);
+			BonePointIndex.ModifyAt(BoneIndex,Bone.PointIndex);
+			BoneBranchIndex.ModifyAt(BoneIndex,Bone.BranchIndex);
+			BoneId.ModifyAt(BoneIndex,Bone.BoneIndex);
+			BonePose.ModifyAt(BoneIndex,Bone.BoneTransform);
+			BoneAbsolutePosition.ModifyAt(BoneIndex,FVector3f(Bone.AbsolutePosition));
+			BoneIndex++;
 		}
 	}
 }
diff --git a/Engine/Plugins/Experimental/ProceduralVegetationEditor/Source/MegaPlants/Private/Facades/MegaPlantBranchFacade.cpp b/Engine/Plugins/Experimental/ProceduralVegetationEditor/Source/MegaPlants/Private/Facades/MegaPlantBranchFacade.cpp
--- a/Engine/Plugins/Experimental/ProceduralVegetationEditor/Source/MegaPlants/Private/Facades/MegaPlantBranchFacade.cpp
+++ b/Engine/Plugins/Experimental/ProceduralVegetationEditor/Source/MegaPlants/Private/Facades/MegaPlantBranchFacade.cpp
@@ -194,7 +194,7 @@ namespace MegaPlants::Facades
 
 		return INDEX_NONE;
 	}
-
+	
 	void FBranchFacade::SetTrunkMaterialPath(const FString& InPath)
 	{
 		check(!IsConst());
@@ -269,4 +269,18 @@ namespace MegaPlants::Facades
 			return Index1Hierarchy < Index2Hierarchy;
 		});
 	}
+
+	int32 FBranchFacade::GetBranchIndexFromPointIndex(const int32 PointIndex) const
+	{
+		for (int BranchIndex = 0; BranchIndex < GetElementCount(); BranchIndex++)
+		{
+			TArray<int32> PointInexes = GetPoints(BranchIndex);
+			if (PointInexes.Contains(PointIndex))
+			{
+				return BranchIndex;
+			}
+		}
+
+		return INDEX_NONE;
+	}
 }
diff --git a/Engine/Plugins/Experimental/ProceduralVegetationEditor/Source/MegaPlants/Private/Helpers/MegaPlantExportHelper.cpp b/Engine/Plugins/Experimental/ProceduralVegetationEditor/Source/MegaPlants/Private/Helpers/MegaPlantExportHelper.cpp
--- a/Engine/Plugins/Experimental/ProceduralVegetationEditor/Source/MegaPlants/Private/Helpers/MegaPlantExportHelper.cpp
+++ b/Engine/Plugins/Experimental/ProceduralVegetationEditor/Source/MegaPlants/Private/Helpers/MegaPlantExportHelper.cpp
@@ -631,25 +631,35 @@ namespace MegaPlants::Export::Internal
 		TArray<TObjectPtr<UMaterialInterface>> Materials = GetMaterialsFromCollection(ExportedCollection);
 
 		TArray<Facades::FBoneNode> BoneNodes;
-		BoneFacade.CreateBoneData(BoneNodes);
+		BoneNodes = BoneFacade.GetBoneDataFromCollection();
+
+		//if bones already not created with bone reduction node create bones with full density
+		if (BoneNodes.IsEmpty())
+		{
+			BoneFacade.CreateBoneData(BoneNodes, 0);
+		}
+		
 		AssignBoneIDsToFoliage(BoneNodes, ExportedCollection);
 		
 		int32 BoneCount = BoneNodes.Num();
 		MeshAttributes->EnableBones(BoneCount);
 
-		//NjordPixel id is set per vertex
-		TArray<float> NjordPixelIDs;
-		BoneFacade.GetNjordPixelIds(NjordPixelIDs);
+		TArray<int32> VertexPointIDs;
+		BoneFacade.GetPointIds(VertexPointIDs);
 
 		TArray<int32> VertexBoneIDs;
-		//Assign the Bone id to every vertex in the base mesh through NjordPixelId
-		for (const float& NjordPixelId : NjordPixelIDs)
+		//Assign the Bone id to every vertex in the base mesh through VertexPointIDs
+		for (const int32& PointID : VertexPointIDs)
 		{
-			Facades::FBoneNode* BoneNode = BoneNodes.FindByPredicate([&NjordPixelId](const Facades::FBoneNode& Node)
+			if (Facades::FBoneNode* BoneNode = Facades::FBoneFacade::FindClosestBone(ExportedCollection, BoneNodes, PointID))
 			{
-				return FMath::FloorToInt(Node.NjordPixelID) == FMath::FloorToInt(NjordPixelId);
-			});
-			VertexBoneIDs.Add(BoneNode->BoneIndex);
+				VertexBoneIDs.Add(BoneNode->BoneIndex);
+			}
+			else
+			{
+				VertexBoneIDs.Add(INDEX_NONE);
+				UE_LOG(LogMegaPlants, Log, TEXT("Invalid Bone Assigned to PointID %i") , PointID);
+			}
 		}
 		int FoliageVertexStart = VertexBoneIDs.Num();
 
@@ -721,7 +731,7 @@ namespace MegaPlants::Export::Internal
 		const FName ProfileName = SkinProfile.GetProfileName();
 		
 		//Here we make the correction to the weights
-		RemoveUnwantedSkinWeights(DynamicMesh, ProfileName, VertexBoneIDs, NjordPixelIDs, FoliageVertexStart,BoneNodes);
+		RemoveUnwantedSkinWeights(DynamicMesh, ProfileName, VertexBoneIDs, FoliageVertexStart,BoneNodes);
 		
 		UGeometryScriptLibrary_StaticMeshFunctions::CopySkinWeightProfileToSkeletalMesh(
 			DynamicMesh,
@@ -819,7 +829,7 @@ namespace MegaPlants::Export::Internal
 		}
 	}
 
-	void RemoveUnwantedSkinWeights(TObjectPtr<UDynamicMesh> DynamicMesh, const FName ProfileName, const TArray<int32>& VertexBoneIDs, const TArray<float>& NjordIDs, int32 FoliageVertexStart, const TArray<Facades::FBoneNode>& BoneNodes)
+	void RemoveUnwantedSkinWeights(TObjectPtr<UDynamicMesh> DynamicMesh, const FName ProfileName, const TArray<int32>& VertexBoneIDs, int32 FoliageVertexStart, const TArray<Facades::FBoneNode>& BoneNodes)
 	{
 		auto InMesh = DynamicMesh->GetMeshPtr();
 		check(InMesh)
@@ -832,42 +842,50 @@ namespace MegaPlants::Export::Internal
 			ParallelFor(NumVertices, [&](const int32 VertexIdx)
 			{
 				int32 BoneID = VertexBoneIDs[VertexIdx];
-				int32 ParentBoneID = BoneNodes[BoneID].ParentBoneIndex;
-			
-				//UE_LOG(LogMegaPlants, Warning, TEXT("Remove Bones, VertexId{%i} BoneID{%i} BoneParentIndex{%i}"), VertexIdx, BoneID, ParentBoneID);
-
-				UE::AnimationCore::FBoneWeights Weights;
-				SkinWeights->GetValue(VertexIdx, Weights);
 
-				auto NextBone = BoneNodes.FindByPredicate([&](const Facades::FBoneNode& Node)
+				if(BoneID != INDEX_NONE)
 				{
-					return Node.BoneIndex == BoneID + 1 && Node.BranchIndex == BoneNodes[BoneID].BranchIndex;
-				});
+					int32 ParentBoneID = BoneNodes[BoneID].ParentBoneIndex;
+				
+					//UE_LOG(LogMegaPlants, Warning, TEXT("Remove Bones, VertexId{%i} BoneID{%i} BoneParentIndex{%i}"), VertexIdx, BoneID, ParentBoneID);
 
-				for (int32 WeightID = Weights.Num() - 1 ; WeightID >= 0; WeightID--)
-				{
-					auto Weight = Weights[WeightID];
+					UE::AnimationCore::FBoneWeights Weights;
+					SkinWeights->GetValue(VertexIdx, Weights);
 
-					bool bValidBone = (Weight.GetBoneIndex() == ParentBoneID);
-					bValidBone |= (Weight.GetBoneIndex() == BoneID);
-					bValidBone |= NextBone && (Weight.GetBoneIndex() == NextBone->BoneIndex);
-					
-					if (!bValidBone || VertexIdx > FoliageVertexStart || BoneID == 0)
+					auto NextBone = BoneNodes.FindByPredicate([&](const Facades::FBoneNode& Node)
+					{
+						return Node.BoneIndex == BoneID + 1 && Node.BranchIndex == BoneNodes[BoneID].BranchIndex;
+					});
+
+					for (int32 WeightID = Weights.Num() - 1 ; WeightID >= 0; WeightID--)
+					{
+						auto Weight = Weights[WeightID];
+
+						bool bValidBone = (Weight.GetBoneIndex() == ParentBoneID);
+						bValidBone |= (Weight.GetBoneIndex() == BoneID);
+						bValidBone |= NextBone && (Weight.GetBoneIndex() == NextBone->BoneIndex);
+						
+						if (!bValidBone || VertexIdx > FoliageVertexStart || BoneID == 0)
+						{
+							UE::AnimationCore::FBoneWeightsSettings Settings;
+							Settings.SetNormalizeType(UE::AnimationCore::EBoneWeightNormalizeType::Always);
+							Weights.RemoveBoneWeight(Weight.GetBoneIndex(),Settings);
+						}
+					}
+
+					if (VertexIdx > FoliageVertexStart || BoneID == 0 || Weights.Num() == 0)
 					{
 						UE::AnimationCore::FBoneWeightsSettings Settings;
 						Settings.SetNormalizeType(UE::AnimationCore::EBoneWeightNormalizeType::Always);
-						Weights.RemoveBoneWeight(Weight.GetBoneIndex(),Settings);
+						Weights.SetBoneWeight(UE::AnimationCore::FBoneWeight(BoneID, 1.0f), Settings);
 					}
+					
+					SkinWeights->SetNewValue(VertexIdx, Weights);
 				}
-
-				if (VertexIdx > FoliageVertexStart || BoneID == 0 || Weights.Num() == 0)
+				else
 				{
-					UE::AnimationCore::FBoneWeightsSettings Settings;
-					Settings.SetNormalizeType(UE::AnimationCore::EBoneWeightNormalizeType::Always);
-					Weights.SetBoneWeight(UE::AnimationCore::FBoneWeight(BoneID, 1.0f), Settings);
+					UE_LOG(LogMegaPlants, Log, TEXT("Invalid bone index found while removing bones for vertex %i"), VertexIdx);
 				}
-				
-				SkinWeights->SetNewValue(VertexIdx, Weights);
 			});
 				
 		}, EDynamicMeshChangeType::GeneralEdit, EDynamicMeshAttributeChangeFlags::Unknown, false);
diff --git a/Engine/Plugins/Experimental/ProceduralVegetationEditor/Source/MegaPlants/Private/Implementations/MegaPlantMesher.cpp b/Engine/Plugins/Experimental/ProceduralVegetationEditor/Source/MegaPlants/Private/Implementations/MegaPlantMesher.cpp
--- a/Engine/Plugins/Experimental/ProceduralVegetationEditor/Source/MegaPlants/Private/Implementations/MegaPlantMesher.cpp
+++ b/Engine/Plugins/Experimental/ProceduralVegetationEditor/Source/MegaPlants/Private/Implementations/MegaPlantMesher.cpp
@@ -35,7 +35,7 @@ void FMegaPlantMesher::GenerateGeometryCollection(const FManagedArrayCollection&
 		TriangleCount += Triangles.Num();
 	}
 
-	TManagedArray<float>& NjordPixelIds = BoneFacade.ModifyNjordPixelIds();
+	TManagedArray<int32>& PointIDs = BoneFacade.ModifyVertexPointIds();
 
 	OutGeometryCollection.AddElements(1, OutGeometryCollection.TransformGroup);
 	OutGeometryCollection.AddElements(1, OutGeometryCollection.GeometryGroup);
@@ -93,7 +93,7 @@ void FMegaPlantMesher::GenerateGeometryCollection(const FManagedArrayCollection&
 			Vertices[VertexIndex] = static_cast<FVector3f>(Vertex.Position);
 			Normals[VertexIndex] = static_cast<FVector3f>(Vertex.Normal);
 			UVs[VertexIndex] = Vertex.UV;
-			NjordPixelIds[VertexIndex] = Vertex.NjordPixelId;
+			PointIDs[VertexIndex] = Vertex.NjordPixelId;
 			++VertexIndex;
 		}
 
@@ -606,7 +606,7 @@ void FMegaPlantMesher::GenerateBranchMeshData(const bool bPrimitiveIsTrunk, cons
 	for (int32 Index = 0; Index <= MaxIndex; ++Index)
 	{
 		const int32& PointIndex = PrimitivePoints[Index];
-		const float NjordPixelId = PointFacade.GetNjordPixelIndex(PointIndex);
+		const float NjordPixelId = PointIndex;//PointFacade.GetNjordPixelIndex(PointIndex);
 
 		float TextureCoordV = PointFacade.GetTextureCoordV(PointIndex);
 		float TextureCoordUOffset = PointFacade.GetTextureCoordUOffset(PointIndex);
diff --git a/Engine/Plugins/Experimental/ProceduralVegetationEditor/Source/MegaPlants/Private/Nodes/PCGMegaPlantBoneReduction.cpp b/Engine/Plugins/Experimental/ProceduralVegetationEditor/Source/MegaPlants/Private/Nodes/PCGMegaPlantBoneReduction.cpp
new file mode 100644
--- /dev/null
+++ b/Engine/Plugins/Experimental/ProceduralVegetationEditor/Source/MegaPlants/Private/Nodes/PCGMegaPlantBoneReduction.cpp
@@ -0,0 +1,58 @@
+﻿// Copyright Epic Games, Inc. All Rights Reserved.
+
+#include "PCGMegaPlantBoneReduction.h"
+#include "PCGMegaPlantData.h"
+#include "Data/PCGBasePointData.h"
+#include "Facades/MegaPlantBoneFacade.h"
+#include "GeometryCollection/ManagedArrayCollection.h"
+#include "Helpers/PCGSettingsHelpers.h"
+
+FPCGElementPtr UPCGMegaPlantBoneReductionSettings::CreateElement() const
+{
+	return MakeShared<FPCGMegaPlantBoneReductionElement>();
+}
+
+bool FPCGMegaPlantBoneReductionElement::ExecuteInternal(FPCGContext* InContext) const
+{
+	TRACE_CPUPROFILER_EVENT_SCOPE(FPCGMegaPlantBoneReductionElement::ExecuteInternal);
+
+	check(InContext);
+
+	const UPCGMegaPlantBoneReductionSettings* InputSettings = InContext->GetInputSettings<UPCGMegaPlantBoneReductionSettings>();
+	check(InputSettings);
+
+	for (const FPCGTaggedData& Input : InContext->InputData.GetInputsByPin(PCGPinConstants::DefaultInputLabel))
+	{
+		if (const UPCGMegaPlantData* InputData = Cast<UPCGMegaPlantData>(Input.Data))
+		{
+			FManagedArrayCollection SourceCollection = InputData->GetCollection();
+			MegaPlants::Facades::FBoneFacade BoneFacadeConst = MegaPlants::Facades::FBoneFacade(SourceCollection);
+			if (!BoneFacadeConst.IsValid())
+			{
+				PCGLog::LogErrorOnGraph(FText::FromString(TEXT("Can not reduce bones. Mesh data not available.")), InContext);
+				return true;
+			}
+
+			FManagedArrayCollection OutCollection;
+			SourceCollection.CopyTo(&OutCollection);
+
+			MegaPlants::Facades::FBoneFacade BoneFacade = MegaPlants::Facades::FBoneFacade(OutCollection);
+			UPCGMegaPlantData* OutManagedArrayCollectionData = FPCGContext::NewObject_AnyThread<UPCGMegaPlantData>(InContext);
+
+			TArray<MegaPlants::Facades::FBoneNode> BoneNodes;
+			float ReductionStrength = InputSettings->Strength;
+			BoneFacade.CreateBoneData(BoneNodes, ReductionStrength);
+			BoneFacade.SetBoneDataToCollection(BoneNodes);
+			
+			OutManagedArrayCollectionData->Initialize(MoveTemp(OutCollection));
+			InContext->OutputData.TaggedData.Emplace(OutManagedArrayCollectionData);
+		}
+		else
+		{
+			PCGLog::InputOutput::LogInvalidInputDataError(InContext);
+			return true;
+		}
+	}
+
+	return true;
+}
diff --git a/Engine/Plugins/Experimental/ProceduralVegetationEditor/Source/MegaPlants/Private/Nodes/PCGMegaPlantBoneReduction.h b/Engine/Plugins/Experimental/ProceduralVegetationEditor/Source/MegaPlants/Private/Nodes/PCGMegaPlantBoneReduction.h
new file mode 100644
--- /dev/null
+++ b/Engine/Plugins/Experimental/ProceduralVegetationEditor/Source/MegaPlants/Private/Nodes/PCGMegaPlantBoneReduction.h
@@ -0,0 +1,37 @@
+﻿// Copyright Epic Games, Inc. All Rights Reserved.
+#pragma once
+#include "CoreMinimal.h"
+#include "Nodes/PCGMegaPlantBaseSettings.h"
+#include "PCGMegaPlantData.h"
+#include "PCGMegaPlantBoneReduction.generated.h"
+
+UCLASS(BlueprintType, ClassGroup = (Procedural))
+class UPCGMegaPlantBoneReductionSettings : public UPCGMegaPlantBaseSettings
+{
+	GENERATED_BODY()
+
+	friend class FPCGMegaPlantBoneReductionElement;
+
+public:
+#if WITH_EDITOR
+	virtual FName GetDefaultNodeName() const override { return FName(TEXT("MegaPlantBoneReduction")); }
+	virtual FText GetDefaultNodeTitle() const override { return NSLOCTEXT("PCGMegaBoneReduction", "NodeTitle", "Bone Reduction"); }
+	virtual FText GetNodeTooltipText() const override { return NSLOCTEXT("PCGMegaBoneReduction", "NodeTooltip", "Controls the number of bones created for skeletal mesh"); }
+	virtual FLinearColor GetNodeTitleColor() const override { return FLinearColor::Yellow; }
+
+	virtual TArray<EMegaPlantRenderType> GetRenderType() const override { return TArray{EMegaPlantRenderType::Bones, EMegaPlantRenderType::Mesh}; }
+#endif
+
+protected:
+	virtual FPCGElementPtr CreateElement() const override;
+
+private:
+	UPROPERTY(EditAnywhere, Category = "Bone Reduction", meta = (ShowOnlyInnerProperties, ClampMin ="0", ClampMax = "1", UIMin = "0", UIMax = "1"))
+	float Strength;
+};
+
+class FPCGMegaPlantBoneReductionElement : public FPCGMegaPlantBaseElement
+{
+protected:
+	virtual bool ExecuteInternal(FPCGContext* InContext) const override;
+};
diff --git a/Engine/Plugins/Experimental/ProceduralVegetationEditor/Source/MegaPlants/Public/Facades/MegaPlantAttributesNames.h b/Engine/Plugins/Experimental/ProceduralVegetationEditor/Source/MegaPlants/Public/Facades/MegaPlantAttributesNames.h
--- a/Engine/Plugins/Experimental/ProceduralVegetationEditor/Source/MegaPlants/Public/Facades/MegaPlantAttributesNames.h
+++ b/Engine/Plugins/Experimental/ProceduralVegetationEditor/Source/MegaPlants/Public/Facades/MegaPlantAttributesNames.h
@@ -77,8 +77,10 @@ namespace MegaPlants::AttributeNames
 	inline static const FName BoneParentIndex = FName(TEXT("BoneParentIndex"));
 	inline static const FName BonePose = FName(TEXT("BonePose"));
 	inline static const FName BonePointIndex = FName(TEXT("BonePointIndex"));
-	inline static const FName BoneParentPointIndex = FName(TEXT("BoneParentPointIndex"));
+	inline static const FName BoneAbsolutePosition = FName(TEXT("BoneAbsolutePosition"));
 	inline static const FName BoneBranchIndex = FName(TEXT("BoneBranchIndex"));
 	inline static const FName NjordPixelId = FName(TEXT("NjordPixelId"));
+	inline static const FName VertexPointIds = FName(TEXT("VertexPointIds"));
+	inline static const FName BoneId = FName(TEXT("BoneId"));
 }
 
diff --git a/Engine/Plugins/Experimental/ProceduralVegetationEditor/Source/MegaPlants/Public/Facades/MegaPlantBoneFacade.h b/Engine/Plugins/Experimental/ProceduralVegetationEditor/Source/MegaPlants/Public/Facades/MegaPlantBoneFacade.h
--- a/Engine/Plugins/Experimental/ProceduralVegetationEditor/Source/MegaPlants/Public/Facades/MegaPlantBoneFacade.h
+++ b/Engine/Plugins/Experimental/ProceduralVegetationEditor/Source/MegaPlants/Public/Facades/MegaPlantBoneFacade.h
@@ -38,7 +38,7 @@ namespace MegaPlants::Facades
 		
 		bool IsValid() const;
 
-		void GetNjordPixelIds(TArray<float>& VertexNjordIds) const;
+		void GetPointIds(TArray<int32>& OutVertexPointIds) const;
 		
 		const TManagedArray<FString>& GetBoneNames() const;
 		
@@ -46,19 +46,26 @@ namespace MegaPlants::Facades
 		
 		const TManagedArray<FTransform>& GetBonePoses() const;
 
-		const TManagedArray<float>& GetNjordPixelIds() const;
+		const TManagedArray<int32>& GetVertexPointIds() const;
 		
-		TManagedArray<float>& ModifyNjordPixelIds();
+		TManagedArray<int32>& ModifyVertexPointIds();
+		
+		void CreateBoneData(TArray<FBoneNode>& Bones,const float ReductionStrength);
+
+		static FBoneNode* FindClosestBone(const FManagedArrayCollection& Collection, TArray<FBoneNode>& Bones, const int32 PointID);
+
+		TArray<FBoneNode> GetBoneDataFromCollection() const;
 
-		void CreateBoneData(TArray<FBoneNode>& Bones) const;
+		void SetBoneDataToCollection(TArray<FBoneNode>& Bones);
 		
 	private:
 		TManagedArrayAccessor<FString> BoneName;
 		TManagedArrayAccessor<int32> BoneParentIndex;
+		TManagedArrayAccessor<int32> BoneId;
 		TManagedArrayAccessor<FTransform> BonePose;
 		TManagedArrayAccessor<int32> BonePointIndex;
-		TManagedArrayAccessor<int32> BoneParentPointIndex;
+		TManagedArrayAccessor<FVector3f> BoneAbsolutePosition;
 		TManagedArrayAccessor<int32> BoneBranchIndex;
-		TManagedArrayAccessor<float> NjordPixelIds;
+		TManagedArrayAccessor<int32> VertexPointIds;
 	};
 }
diff --git a/Engine/Plugins/Experimental/ProceduralVegetationEditor/Source/MegaPlants/Public/Facades/MegaPlantBranchFacade.h b/Engine/Plugins/Experimental/ProceduralVegetationEditor/Source/MegaPlants/Public/Facades/MegaPlantBranchFacade.h
--- a/Engine/Plugins/Experimental/ProceduralVegetationEditor/Source/MegaPlants/Public/Facades/MegaPlantBranchFacade.h
+++ b/Engine/Plugins/Experimental/ProceduralVegetationEditor/Source/MegaPlants/Public/Facades/MegaPlantBranchFacade.h
@@ -64,6 +64,8 @@ namespace MegaPlants::Facades
 		virtual void RemoveEntries(const int32 NumEntries, const int32 StartIndex) override;
 
 		void GetSortedBranchIndicesByHierarchy(TArray<int32>& OutSortedIndices) const;
+		
+		int32 GetBranchIndexFromPointIndex(int32 PointIndex) const;
 
 	private:
 		FManagedArrayCollection* Collection = nullptr;
diff --git a/Engine/Plugins/Experimental/ProceduralVegetationEditor/Source/MegaPlants/Public/Helpers/MegaPlantExportHelper.h b/Engine/Plugins/Experimental/ProceduralVegetationEditor/Source/MegaPlants/Public/Helpers/MegaPlantExportHelper.h
--- a/Engine/Plugins/Experimental/ProceduralVegetationEditor/Source/MegaPlants/Public/Helpers/MegaPlantExportHelper.h
+++ b/Engine/Plugins/Experimental/ProceduralVegetationEditor/Source/MegaPlants/Public/Helpers/MegaPlantExportHelper.h
@@ -117,7 +117,7 @@ namespace MegaPlants::Export::Internal
 	
 	void AssignBoneIDsToFoliage(const TArray<Facades::FBoneNode>& Bones, FManagedArrayCollection& Collection);
 
-	void RemoveUnwantedSkinWeights(TObjectPtr<UDynamicMesh> DynamicMesh, const FName ProfileName, const TArray<int32>& VertexBoneIDs, const TArray<float>& NjordIDs,
+	void RemoveUnwantedSkinWeights(TObjectPtr<UDynamicMesh> DynamicMesh, const FName ProfileName, const TArray<int32>& VertexBoneIDs,
 		int32 FoliageVertexStart, const TArray<Facades::FBoneNode>& BoneNodes);
 
 	void ConvertToDefaultSkeletalMesh(USkeletalMesh* SkeletalMesh,
diff --git a/Engine/Plugins/Experimental/ProceduralVegetationEditor/Source/MegaPlants/Public/PCGMegaPlantData.h b/Engine/Plugins/Experimental/ProceduralVegetationEditor/Source/MegaPlants/Public/PCGMegaPlantData.h
--- a/Engine/Plugins/Experimental/ProceduralVegetationEditor/Source/MegaPlants/Public/PCGMegaPlantData.h
+++ b/Engine/Plugins/Experimental/ProceduralVegetationEditor/Source/MegaPlants/Public/PCGMegaPlantData.h
@@ -15,6 +15,7 @@ enum class EMegaPlantRenderType : uint8
 	Skeleton 	= 1 << 0,
 	Mesh 		= 1 << 1,
 	Foliage 	= 1 << 2,
+	Bones 		= 1 << 3,
 
 	// Add new renter type above
 	Count	
diff --git a/Engine/Plugins/Experimental/ProceduralVegetationEditor/Source/MegaPlantsEditor/Private/Visualizations/MegaPlantBoneComponent.cpp b/Engine/Plugins/Experimental/ProceduralVegetationEditor/Source/MegaPlantsEditor/Private/Visualizations/MegaPlantBoneComponent.cpp
new file mode 100644
--- /dev/null
+++ b/Engine/Plugins/Experimental/ProceduralVegetationEditor/Source/MegaPlantsEditor/Private/Visualizations/MegaPlantBoneComponent.cpp
@@ -0,0 +1,133 @@
+// Copyright Epic Games, Inc. All Rights Reserved.
+
+#include "MegaPlantBoneComponent.h"
+#include "MeshElementCollector.h"
+#include "SkeletalDebugRendering.h"
+#include "Facades/MegaPlantBoneFacade.h"
+#include "SkeletalMesh/SkinWeightsPaintTool.h"
+
+FMegaPlantBoneSceneProxy::FMegaPlantBoneSceneProxy(const UMegaPlantBoneComponent* InComponent)
+	: FPrimitiveSceneProxy(InComponent)
+	, BoneComponent(InComponent)
+{}
+
+FPrimitiveViewRelevance FMegaPlantBoneSceneProxy::GetViewRelevance(const FSceneView* View) const
+{
+	FPrimitiveViewRelevance ViewRelevance;
+	ViewRelevance.bDrawRelevance = IsShown(View);
+	ViewRelevance.bDynamicRelevance = true;
+	ViewRelevance.bSeparateTranslucency = ViewRelevance.bNormalTranslucency = true;
+	return ViewRelevance;
+}
+
+void FMegaPlantBoneSceneProxy::GetDynamicMeshElements(
+	const TArray<const FSceneView*>& Views,
+	const FSceneViewFamily& ViewFamily,
+	uint32 VisibilityMap,
+	FMeshElementCollector& Collector
+) const
+{
+	for (int32 ViewIndex = 0; ViewIndex < Views.Num(); ViewIndex++)
+	{
+		if (VisibilityMap & (1 << ViewIndex))
+		{
+			const FSceneView* View = Views[ViewIndex];
+
+#if UE_ENABLE_DEBUG_DRAWING
+			FPrimitiveDrawInterface* PDI = Collector.GetDebugPDI(ViewIndex);
+#else	
+			FPrimitiveDrawInterface* PDI = Collector.GetPDI(ViewIndex);
+#endif
+
+			BoneComponent->Draw(View, PDI);
+		}
+	}
+}
+
+SIZE_T FMegaPlantBoneSceneProxy::GetTypeHash() const
+{
+	static size_t UniquePointer;
+	return reinterpret_cast<size_t>(&UniquePointer);
+}
+
+uint32 FMegaPlantBoneSceneProxy::GetMemoryFootprint() const { return sizeof *this + GetAllocatedSize(); }
+
+bool FMegaPlantBoneSceneProxy::CanBeOccluded() const
+{
+	return false;
+}
+
+UMegaPlantBoneComponent::UMegaPlantBoneComponent()
+{
+	bAutoActivate = true;
+	bTickInEditor = true;
+	PrimaryComponentTick.bCanEverTick = true;
+	bUseEditorCompositing = true;
+	SetGenerateOverlapEvents(false);
+	bIgnoreStreamingManagerUpdate = true;
+	
+	BBox = FBox(ForceInit);
+	Bounds = FBoxSphereBounds(FVector::ZeroVector, FVector::OneVector, 1);
+}
+
+void UMegaPlantBoneComponent::InitBounds()
+{
+	BBox = FBox(ForceInit);
+}
+
+void UMegaPlantBoneComponent::AddBone(const FVector& InStartPos, const FVector& InEndPos, const FLinearColor& InColor)
+{
+	Bones.Emplace(InStartPos, InEndPos, InColor);
+
+	BBox += InStartPos;
+	BBox += InEndPos;
+}
+
+void UMegaPlantBoneComponent::Draw(const FSceneView* View, FPrimitiveDrawInterface* PDI) const
+{
+	for (const FMegaPlantBoneInfo& Bone : Bones)
+	{
+		SkeletalDebugRendering::DrawWireBone(PDI,Bone.StartPos, Bone.EndPos, Bone.Color,SDPG_MAX,2);
+	}
+}
+
+FPrimitiveSceneProxy* UMegaPlantBoneComponent::CreateSceneProxy()
+{
+	FMegaPlantBoneSceneProxy* Proxy = new FMegaPlantBoneSceneProxy(this);
+	return Proxy;
+}
+
+FBoxSphereBounds UMegaPlantBoneComponent::CalcBounds(const FTransform& LocalToWorld) const
+{
+	return FBoxSphereBounds(BBox).TransformBy(LocalToWorld); 
+}
+
+void UMegaPlantBoneComponent::SetCollection(const FManagedArrayCollection* const InSkeletonCollection)
+{
+	if (InSkeletonCollection == nullptr)
+	{
+		return;
+	}
+
+	Bones.Empty();
+	InitBounds();
+	
+	const MegaPlants::Facades::FBoneFacade BoneFacade = MegaPlants::Facades::FBoneFacade(*InSkeletonCollection);
+	TArray<MegaPlants::Facades::FBoneNode> BoneNodes = BoneFacade.GetBoneDataFromCollection();
+
+	for (int i=0; i < BoneNodes.Num(); i++)
+	{
+		MegaPlants::Facades::FBoneNode* ParentNode = BoneNodes.FindByPredicate([&](const MegaPlants::Facades::FBoneNode& BoneNode)
+		{
+			return BoneNode.BoneIndex == BoneNodes[i].ParentBoneIndex;
+		});
+
+		FVector StartPos = ParentNode ? ParentNode->AbsolutePosition : BoneNodes[i].AbsolutePosition;
+		AddBone(StartPos,BoneNodes[i].AbsolutePosition, SkeletalDebugRendering::GetSemiRandomColorForBone(i));
+	}
+}
+
+int UMegaPlantBoneComponent::GetBoneCount()
+{
+	return Bones.Num();
+}
diff --git a/Engine/Plugins/Experimental/ProceduralVegetationEditor/Source/MegaPlantsEditor/Private/Visualizations/MegaPlantBoneComponent.h b/Engine/Plugins/Experimental/ProceduralVegetationEditor/Source/MegaPlantsEditor/Private/Visualizations/MegaPlantBoneComponent.h
new file mode 100644
--- /dev/null
+++ b/Engine/Plugins/Experimental/ProceduralVegetationEditor/Source/MegaPlantsEditor/Private/Visualizations/MegaPlantBoneComponent.h
@@ -0,0 +1,69 @@
+// Copyright Epic Games, Inc. All Rights Reserved.
+
+#pragma once
+
+#include "CoreMinimal.h"
+#include "PrimitiveSceneProxy.h"
+#include "PrimitiveViewRelevance.h"
+#include "SceneView.h"
+#include "Components/PrimitiveComponent.h"
+#include "MegaPlantBoneComponent.generated.h"
+
+struct FManagedArrayCollection;
+class UMegaPlantBoneComponent;
+
+class FMegaPlantBoneSceneProxy final : public FPrimitiveSceneProxy
+{
+public:
+	FMegaPlantBoneSceneProxy(const UMegaPlantBoneComponent* InComponent);
+
+	virtual FPrimitiveViewRelevance GetViewRelevance(const FSceneView* View) const override;
+	virtual void GetDynamicMeshElements(
+		const TArray<const FSceneView*>& Views,
+		const FSceneViewFamily& ViewFamily,
+		uint32 VisibilityMap,
+		FMeshElementCollector& Collector
+	) const override;
+
+	virtual SIZE_T GetTypeHash() const override;
+	virtual uint32 GetMemoryFootprint(void) const override;
+
+	virtual bool CanBeOccluded() const override;
+private:
+	const TObjectPtr<const UMegaPlantBoneComponent> BoneComponent;
+};
+
+USTRUCT()
+struct FMegaPlantBoneInfo
+{
+	GENERATED_BODY()
+
+	FVector StartPos;
+	FVector EndPos;
+	FLinearColor Color;
+};
+
+UCLASS(Blueprintable, BlueprintType, meta=(BlueprintSpawnableComponent), MinimalAPI)
+class UMegaPlantBoneComponent : public UPrimitiveComponent
+{
+	GENERATED_BODY()
+
+public:
+	UMegaPlantBoneComponent();
+
+	void InitBounds();
+	void AddBone(const FVector& InStartPos, const FVector& InEndPos, const FLinearColor& InColor = FLinearColor::White);
+
+	void Draw(const FSceneView* View, FPrimitiveDrawInterface* PDI) const;
+
+	virtual FPrimitiveSceneProxy* CreateSceneProxy() override;
+	virtual FBoxSphereBounds CalcBounds(const FTransform& LocalToWorld) const override;
+
+	void SetCollection(const FManagedArrayCollection* const InSkeletonCollection);
+	int GetBoneCount();
+
+protected:
+	FBox BBox;
+	TArray<FMegaPlantBoneInfo> Bones;
+	const FManagedArrayCollection* SkeletonCollection;
+};
diff --git a/Engine/Plugins/Experimental/ProceduralVegetationEditor/Source/MegaPlantsEditor/Private/Visualizations/PCGMegaPlantDataVisualization.cpp b/Engine/Plugins/Experimental/ProceduralVegetationEditor/Source/MegaPlantsEditor/Private/Visualizations/PCGMegaPlantDataVisualization.cpp
--- a/Engine/Plugins/Experimental/ProceduralVegetationEditor/Source/MegaPlantsEditor/Private/Visualizations/PCGMegaPlantDataVisualization.cpp
+++ b/Engine/Plugins/Experimental/ProceduralVegetationEditor/Source/MegaPlantsEditor/Private/Visualizations/PCGMegaPlantDataVisualization.cpp
@@ -5,6 +5,7 @@
 #include "AdvancedPreviewScene.h"
 #include "EditorViewportClient.h"
 #include "GeometryCollectionToDynamicMesh.h"
+#include "MegaPlantBoneComponent.h"
 #include "PCGComponent.h"
 #include "PCGContext.h"
 #include "PCGDataVisualization.h"
@@ -18,10 +19,13 @@
 #include "GeometryCollection/Facades/CollectionRenderingFacade.h"
 #include "Helpers/MegaPlantUtilities.h"
 #include "PCGMegaPlantData.h"
+#include "Components/TextRenderComponent.h"
 #include "Visualizations/MegaPlantSkeletonVisualizerComponent.h"
 
 #define LOCTEXT_NAMESPACE "PCGManagedArrayCollectionDataVisualization"
 
+class UMegaPlantBoneComponent;
+
 DEFINE_LOG_CATEGORY(LogMegaPlantDataVisualization);
 
 FPCGMegaPlantDataVisualization::FPCGMegaPlantDataVisualization()
@@ -37,6 +41,7 @@ void FPCGMegaPlantDataVisualization::RegisterVisualizations()
 	RenderMap.Add(EMegaPlantRenderType::Skeleton, FMegaPlantRenderCallback::CreateRaw(this, &FPCGMegaPlantDataVisualization::SkeletonRenderer));
 	RenderMap.Add(EMegaPlantRenderType::Mesh, FMegaPlantRenderCallback::CreateRaw(this, &FPCGMegaPlantDataVisualization::MeshRenderer));
 	RenderMap.Add(EMegaPlantRenderType::Foliage, FMegaPlantRenderCallback::CreateRaw(this, &FPCGMegaPlantDataVisualization::FoliageRenderer));
+	RenderMap.Add(EMegaPlantRenderType::Bones, FMegaPlantRenderCallback::CreateRaw(this, &FPCGMegaPlantDataVisualization::BonesRenderer));
 }
 
 void FPCGMegaPlantDataVisualization::ExecuteDebugDisplay(FPCGContext* Context, const UPCGSettingsInterface* SettingsInterface, const UPCGData* Data,
@@ -117,6 +122,7 @@ void FPCGMegaPlantDataVisualization::MeshRenderer(const FManagedArrayCollection&
 	TRACE_CPUPROFILER_EVENT_SCOPE(FPCGMegaPlantDataVisualization::MeshRenderer);
 
 	TArray<UMaterialInterface*> Materials;
+	
 	if (InCollection.HasAttribute("MaterialPath", FGeometryCollection::MaterialGroup))
 	{
 		const TManagedArray<FString>& MaterialPaths = InCollection.GetAttribute<FString>("MaterialPath", FGeometryCollection::MaterialGroup);
@@ -184,4 +190,37 @@ void FPCGMegaPlantDataVisualization::FoliageRenderer(const FManagedArrayCollecti
 	                                           OutInstancedComponentMap);
 }
 
+void FPCGMegaPlantDataVisualization::BonesRenderer(const FManagedArrayCollection& InCollection, FPCGSceneSetupParams& InOutParams,
+	FBoxSphereBounds& OutBounds)
+{
+	TRACE_CPUPROFILER_EVENT_SCOPE(FPCGMegaPlantDataVisualization::BonesRenderer);
+
+	const TObjectPtr<UMegaPlantBoneComponent> BoneComponent = NewObject<UMegaPlantBoneComponent>(
+		GetTransientPackage(), NAME_None, RF_Transient);
+	BoneComponent->SetCollection(&InCollection);
+	BoneComponent->UpdateBounds();
+
+	InOutParams.ManagedResources.Add(BoneComponent);
+	InOutParams.Scene->AddComponent(BoneComponent, FTransform::Identity);
+
+	FString InTextToDraw = FString::Format(TEXT("{0} Bones"),{BoneComponent->GetBoneCount()});
+	const FVector3f& Pos = (FVector3f::ForwardVector * 100);
+
+	DrawDebugString(BoneComponent->GetWorld(), FVector::Zero(), InTextToDraw);
+	
+	TObjectPtr<UTextRenderComponent> Component = NewObject<UTextRenderComponent>(GetTransientPackage(), NAME_None, RF_Transient);
+	Component->SetText(FText::FromString(InTextToDraw));
+	Component->SetTextRenderColor(FColor::Blue);
+	Component->SetWorldSize(40);
+	Component->SetGenerateOverlapEvents(false);
+						
+	InOutParams.ManagedResources.Add(Component);
+	FTransform TextTransform = FTransform(FVector(Pos.X, Pos.Y, Pos.Z));
+	FRotator TextRotator = FRotator(0,90,0);
+	TextTransform.SetRotation(TextRotator.Quaternion());
+	InOutParams.Scene->AddComponent(Component, TextTransform);
+
+	OutBounds = OutBounds + BoneComponent->CalcLocalBounds();
+}
+
 #undef LOCTEXT_NAMESPACE
diff --git a/Engine/Plugins/Experimental/ProceduralVegetationEditor/Source/MegaPlantsEditor/Private/Visualizations/PCGMegaPlantDataVisualization.h b/Engine/Plugins/Experimental/ProceduralVegetationEditor/Source/MegaPlantsEditor/Private/Visualizations/PCGMegaPlantDataVisualization.h
--- a/Engine/Plugins/Experimental/ProceduralVegetationEditor/Source/MegaPlantsEditor/Private/Visualizations/PCGMegaPlantDataVisualization.h
+++ b/Engine/Plugins/Experimental/ProceduralVegetationEditor/Source/MegaPlantsEditor/Private/Visualizations/PCGMegaPlantDataVisualization.h
@@ -30,6 +30,7 @@ protected:
 	void MeshRenderer(const FManagedArrayCollection& InCollection, FPCGSceneSetupParams& InOutParams, FBoxSphereBounds& OutBounds);
 	void SkeletonRenderer(const FManagedArrayCollection& InCollection, FPCGSceneSetupParams& InOutParams, FBoxSphereBounds& OutBounds);
 	void FoliageRenderer(const FManagedArrayCollection& InCollection, FPCGSceneSetupParams& InOutParams, FBoxSphereBounds& OutBounds);
+	void BonesRenderer(const FManagedArrayCollection& InCollection, FPCGSceneSetupParams& InOutParams, FBoxSphereBounds& OutBounds);
 };
 
 DECLARE_LOG_CATEGORY_EXTERN(LogMegaPlantDataVisualization, Log, All);
\ No newline at end of file
