[ProceduralVegetationEditor] Fix for default preset

diff --git a/Engine/Plugins/Experimental/ProceduralVegetationEditor/Source/MegaPlantsEditor/Private/Widgets/SMegaPlantCreateDialog.cpp b/Engine/Plugins/Experimental/ProceduralVegetationEditor/Source/MegaPlantsEditor/Private/Widgets/SMegaPlantCreateDialog.cpp
--- a/Engine/Plugins/Experimental/ProceduralVegetationEditor/Source/MegaPlantsEditor/Private/Widgets/SMegaPlantCreateDialog.cpp
+++ b/Engine/Plugins/Experimental/ProceduralVegetationEditor/Source/MegaPlantsEditor/Private/Widgets/SMegaPlantCreateDialog.cpp
@@ -23,7 +23,7 @@ void SMegaPlantCreateDialog::Construct(const FArguments& InArgs)
 	TSharedRef<SVerticalBox> Container = SNew(SVerticalBox);
 
 	FName SampleTreesVirtualPath;
-	IContentBrowserDataModule::Get().GetSubsystem()->ConvertInternalPathToVirtual("/MegaPlants/SampleTrees", SampleTreesVirtualPath);
+	IContentBrowserDataModule::Get().GetSubsystem()->ConvertInternalPathToVirtual("/ProceduralVegetationEditor/SampleTrees", SampleTreesVirtualPath);
 
 	const FContentBrowserModule& ContentBrowserModule = FModuleManager::LoadModuleChecked<FContentBrowserModule>("ContentBrowser");
 
