Fix materials breaking due to pluing re-name

diff --git a/Engine/Plugins/Experimental/ProceduralVegetationEditor/Config/BaseProceduralVegetationEditor.ini b/Engine/Plugins/Experimental/ProceduralVegetationEditor/Config/BaseProceduralVegetationEditor.ini
new file mode 100644
--- /dev/null
+++ b/Engine/Plugins/Experimental/ProceduralVegetationEditor/Config/BaseProceduralVegetationEditor.ini
@@ -0,0 +1,2 @@
+[CoreRedirects]
++PackageRedirects=(OldName="/MegaPlants/...", NewName="/ProceduralVegetationEditor/", MatchWildcard=true)
