[ProceduralVegetationEditor] Setting the dynamic wind data to skeletal mesh

diff --git a/Engine/Plugins/Experimental/DynamicWind/Source/DynamicWindEditor/Public/DynamicWindImportData.h b/Engine/Plugins/Experimental/DynamicWind/Source/DynamicWindEditor/Public/DynamicWindImportData.h
--- a/Engine/Plugins/Experimental/DynamicWind/Source/DynamicWindEditor/Public/DynamicWindImportData.h
+++ b/Engine/Plugins/Experimental/DynamicWind/Source/DynamicWindEditor/Public/DynamicWindImportData.h
@@ -9,7 +9,7 @@
 class USkeletalMesh;
 
 USTRUCT()
-struct FDynamicWindJointImportData
+struct DYNAMICWINDEDITOR_API FDynamicWindJointImportData
 {
 	GENERATED_USTRUCT_BODY()
 
@@ -21,7 +21,7 @@ struct FDynamicWindJointImportData
 };
 
 USTRUCT(BlueprintType)
-struct FDynamicWindSkeletalImportData
+struct DYNAMICWINDEDITOR_API FDynamicWindSkeletalImportData
 {
 	GENERATED_USTRUCT_BODY()
 
@@ -42,7 +42,7 @@ namespace DynamicWind
 {
 
 /** Attempts to create dynamic wind asset user data for a skeletal mesh with the given import data */
-UDynamicWindSkeletalData* ImportSkeletalData(
+DYNAMICWINDEDITOR_API UDynamicWindSkeletalData* ImportSkeletalData(
 	USkeletalMesh& TargetSkeletalMesh,
 	const FDynamicWindSkeletalImportData& ImportData
 );
diff --git a/Engine/Plugins/Experimental/ProceduralVegetationEditor/ProceduralVegetationEditor.uplugin b/Engine/Plugins/Experimental/ProceduralVegetationEditor/ProceduralVegetationEditor.uplugin
--- a/Engine/Plugins/Experimental/ProceduralVegetationEditor/ProceduralVegetationEditor.uplugin
+++ b/Engine/Plugins/Experimental/ProceduralVegetationEditor/ProceduralVegetationEditor.uplugin
@@ -53,6 +53,10 @@
 		{
 			"Name": "Fab",
 			"Enabled": true
+		},
+		{
+			"Name": "DynamicWind",
+			"Enabled": true
 		}
 	]
 }
\ No newline at end of file
diff --git a/Engine/Plugins/Experimental/ProceduralVegetationEditor/Source/MegaPlants/MegaPlants.Build.cs b/Engine/Plugins/Experimental/ProceduralVegetationEditor/Source/MegaPlants/MegaPlants.Build.cs
--- a/Engine/Plugins/Experimental/ProceduralVegetationEditor/Source/MegaPlants/MegaPlants.Build.cs
+++ b/Engine/Plugins/Experimental/ProceduralVegetationEditor/Source/MegaPlants/MegaPlants.Build.cs
@@ -32,9 +32,19 @@ namespace UnrealBuildTool.Rules
 					"SkeletalMeshDescription",
 					"StaticMeshDescription",
 					"ModelingComponents",
-					"PCG" // ... add private dependencies that you statically link with here ...	
+					"PCG",
+					"DynamicWind"// ... add private dependencies that you statically link with here ...	
 				}
 			);
+			
+			if (Target.Type == TargetRules.TargetType.Editor)
+			{
+				PrivateDependencyModuleNames.AddRange(
+					new string[]
+					{
+						"DynamicWindEditor"
+					});
+			}
 		}
 	}
 }
diff --git a/Engine/Plugins/Experimental/ProceduralVegetationEditor/Source/MegaPlants/Private/Facades/MegaPlantBoneFacade.cpp b/Engine/Plugins/Experimental/ProceduralVegetationEditor/Source/MegaPlants/Private/Facades/MegaPlantBoneFacade.cpp
--- a/Engine/Plugins/Experimental/ProceduralVegetationEditor/Source/MegaPlants/Private/Facades/MegaPlantBoneFacade.cpp
+++ b/Engine/Plugins/Experimental/ProceduralVegetationEditor/Source/MegaPlants/Private/Facades/MegaPlantBoneFacade.cpp
@@ -6,8 +6,6 @@
 #include "Facades/MegaPlantBranchFacade.h"
 #include "Facades/MegaPlantPointFacade.h"
 #include "GeometryCollection/Facades/CollectionConstraintOverrideFacade.h"
-#include "PCGMegaPlantData.h"
-
 
 class FGeometryCollection;
 
@@ -177,11 +175,68 @@ namespace MegaPlants::Facades
 			}
 		}
 	}
+
+	void SetBranchSimulationGroupIndex(const FPointFacade& PointFacade, FBranchFacade& BranchFacade)
+	{
+		TArray<float> BranchRootPScales;
+		int32 NumBranches = BranchFacade.GetElementCount();
+		BranchRootPScales.Reserve(NumBranches);
+
+		float Min = FLT_MAX;
+		float Max = UE_KINDA_SMALL_NUMBER;
+
+		for (int32 BranchIdx = 0; BranchIdx < NumBranches; ++BranchIdx)
+		{
+			const TArray<int32>& BranchPoints = BranchFacade.GetPoints(BranchIdx);
+			check(BranchPoints.Num() > 0);
+			const int32 BranchRootIndex = BranchPoints[0];
+
+			float Pscale = PointFacade.GetPointScale(BranchRootIndex);
+			BranchRootPScales.Emplace(Pscale);
+
+			Min = FMath::Min(Min, Pscale);
+			Max = FMath::Max(Max, Pscale);
+		}
+
+		for (auto& Pscale : BranchRootPScales)
+		{
+			const UE::Math::TVector2<float> InputRange(Min, Max);
+			const UE::Math::TVector2<float> OutputRange(0, 1);
+			Pscale = FMath::GetMappedRangeValueClamped(InputRange, OutputRange, Pscale);
+		}
+
+		for (int32 BranchIdx = 0; BranchIdx < NumBranches; ++BranchIdx)
+		{
+			TArray<int32> ParentIndices = BranchFacade.GetParentBranchIndices(BranchIdx);
+			int32 ParentIndex = ParentIndices.Num() > 0 ? ParentIndices.Last() : -1;
+
+			if (BranchRootPScales.IsValidIndex(ParentIndex))
+			{
+				BranchRootPScales[ParentIndex] = FMath::Min(BranchRootPScales[BranchIdx], BranchRootPScales[ParentIndex]);
+			}
+		}
+
+		for (int32 BranchIdx = 0; BranchIdx < NumBranches; ++BranchIdx)
+		{
+			int32 LogicalDepth = FMath::RoundToInt32(BranchRootPScales[BranchIdx]*2);
+			BranchFacade.SetBranchSimulationGroupIndex(BranchIdx, LogicalDepth);
+		}
+	}
+
+	void FBoneFacade::SetWindSimulationData(FManagedArrayCollection& Collection)
+	{
+		FPointFacade PointFacade = FPointFacade(*BoneName.GetCollection());
+		FBranchFacade BranchFacade = FBranchFacade(*BoneName.GetCollection());
+		
+		SetBranchSimulationGroupIndex(PointFacade, BranchFacade);
+	}
 	
 	void FBoneFacade::CreateBoneData(TArray<FBoneNode>& Bones,const float ReductionStrength)
 	{
 		FPointFacade PointFacade(VertexPointIds.GetConstCollection());
 		FBranchFacade BranchFacade(VertexPointIds.GetConstCollection());
+
+		SetWindSimulationData(*VertexPointIds.GetCollection());
 		
 		const TManagedArray<FVector3f>& Positions = PointFacade.GetPositions();
 		TArray<int32> SortedBranchIndices;
diff --git a/Engine/Plugins/Experimental/ProceduralVegetationEditor/Source/MegaPlants/Private/Facades/MegaPlantBranchFacade.cpp b/Engine/Plugins/Experimental/ProceduralVegetationEditor/Source/MegaPlants/Private/Facades/MegaPlantBranchFacade.cpp
--- a/Engine/Plugins/Experimental/ProceduralVegetationEditor/Source/MegaPlants/Private/Facades/MegaPlantBranchFacade.cpp
+++ b/Engine/Plugins/Experimental/ProceduralVegetationEditor/Source/MegaPlants/Private/Facades/MegaPlantBranchFacade.cpp
@@ -15,6 +15,7 @@ namespace MegaPlants::Facades
 		, BranchFoliageIDs(InCollection, AttributeNames::BranchFoliageIDs, GroupNames::BranchGroup)
 		, BranchUVMaterial(InCollection, AttributeNames::BranchUVMaterial, GroupNames::BranchGroup)
 		, BranchHierarchyNumber(InCollection, AttributeNames::BranchHierarchyNumber, GroupNames::BranchGroup)
+		, BranchSimulationGroupIndex(InCollection, AttributeNames::BranchSimulationGroupIndex, GroupNames::BranchGroup)
 		, TrunkMaterialPathAttribute(InCollection, AttributeNames::TrunkMaterialPath, GroupNames::DetailsGroup)
 	{
 	}
@@ -29,6 +30,7 @@ namespace MegaPlants::Facades
 		, BranchFoliageIDs(InCollection, AttributeNames::BranchFoliageIDs, GroupNames::BranchGroup)
 		, BranchUVMaterial(InCollection, AttributeNames::BranchUVMaterial, GroupNames::BranchGroup)
 		, BranchHierarchyNumber(InCollection, AttributeNames::BranchHierarchyNumber, GroupNames::BranchGroup)
+		, BranchSimulationGroupIndex(InCollection, AttributeNames::BranchSimulationGroupIndex, GroupNames::BranchGroup)
 		, TrunkMaterialPathAttribute(InCollection, AttributeNames::TrunkMaterialPath, GroupNames::DetailsGroup)
 	{
 		
@@ -194,6 +196,28 @@ namespace MegaPlants::Facades
 
 		return INDEX_NONE;
 	}
+
+	int32 FBranchFacade::GetBranchSimulationGroupIndex(const int32 Index) const
+	{
+		if (BranchSimulationGroupIndex.IsValid() && BranchSimulationGroupIndex.IsValidIndex(Index))
+		{
+			return BranchSimulationGroupIndex[Index];
+		}
+
+		return INDEX_NONE;
+	}
+
+	void FBranchFacade::SetBranchSimulationGroupIndex(const int32 Index, const int32 InSimulationGroupIndex)
+	{
+		if (!BranchSimulationGroupIndex.IsValid())
+		{
+			BranchSimulationGroupIndex.Add();
+		}
+		if (BranchSimulationGroupIndex.IsValid() && BranchSimulationGroupIndex.IsValidIndex(Index))
+		{
+			BranchSimulationGroupIndex.ModifyAt(Index, InSimulationGroupIndex);
+		}
+	}
 	
 	void FBranchFacade::SetTrunkMaterialPath(const FString& InPath)
 	{
@@ -238,6 +262,10 @@ namespace MegaPlants::Facades
 			{
 				BranchUVMaterial.ModifyAt(ToIndex, BranchUVMaterial[FromIndex]);
 			}
+			if (BranchSimulationGroupIndex.IsValid())
+			{
+				BranchSimulationGroupIndex.ModifyAt(ToIndex, BranchSimulationGroupIndex[FromIndex]);
+			}
 		}
 	}
 
diff --git a/Engine/Plugins/Experimental/ProceduralVegetationEditor/Source/MegaPlants/Private/Helpers/MegaPlantExportHelper.cpp b/Engine/Plugins/Experimental/ProceduralVegetationEditor/Source/MegaPlants/Private/Helpers/MegaPlantExportHelper.cpp
--- a/Engine/Plugins/Experimental/ProceduralVegetationEditor/Source/MegaPlants/Private/Helpers/MegaPlantExportHelper.cpp
+++ b/Engine/Plugins/Experimental/ProceduralVegetationEditor/Source/MegaPlants/Private/Helpers/MegaPlantExportHelper.cpp
@@ -3,6 +3,10 @@
 #include "Helpers/MegaPlantExportHelper.h"
 
 #include "DynamicMeshEditor.h"
+#if WITH_EDITOR
+#include "DynamicWindImportData.h"
+#endif
+#include "DynamicWindSkeletalData.h"
 #include "MegaPlantLink.h"
 #include "MegaPlantsModule.h"
 #include "MeshDescriptionToDynamicMesh.h"
@@ -609,6 +613,19 @@ namespace MegaPlants::Export::Internal
 #endif
 	}
 
+#if WITH_EDITOR
+	void BuildWindImportData(FDynamicWindSkeletalImportData& OutImportData, const FManagedArrayCollection& Collection, const TArray<MegaPlants::Facades::FBoneNode>& Bones)
+	{
+		Facades::FBranchFacade BranchFacade(Collection);
+
+		for (const auto& Bone : Bones)
+		{
+			int32 SimulationGroupIndex = BranchFacade.GetBranchSimulationGroupIndex(Bone.BranchIndex);
+			OutImportData.Joints.Add({Bone.BoneName, SimulationGroupIndex});
+		}
+	}
+#endif
+
 	void ExportCollectionToSkeletalMesh(
 		TObjectPtr<USkeletalMesh> ExportMesh,
 		const FManagedArrayCollection& Collection,
@@ -750,6 +767,15 @@ namespace MegaPlants::Export::Internal
 			{
 				//Build the nanite assemblies
 				BuildNaniteAssemblyData(ExportedCollection, ExportMesh);
+
+				FDynamicWindSkeletalImportData ImportData;
+				BuildWindImportData(ImportData,ExportedCollection, BoneNodes);
+				
+				if (!DynamicWind::ImportSkeletalData(*ExportMesh, ImportData))
+				{
+					UE_LOG(LogTemp, Warning, TEXT("Failed to build wind data"));
+				}
+				
 				ExportMesh->Build();
 			}
 			else
diff --git a/Engine/Plugins/Experimental/ProceduralVegetationEditor/Source/MegaPlants/Public/Facades/MegaPlantAttributesNames.h b/Engine/Plugins/Experimental/ProceduralVegetationEditor/Source/MegaPlants/Public/Facades/MegaPlantAttributesNames.h
--- a/Engine/Plugins/Experimental/ProceduralVegetationEditor/Source/MegaPlants/Public/Facades/MegaPlantAttributesNames.h
+++ b/Engine/Plugins/Experimental/ProceduralVegetationEditor/Source/MegaPlants/Public/Facades/MegaPlantAttributesNames.h
@@ -28,6 +28,7 @@ namespace MegaPlants::AttributeNames
 	inline static const FName BranchUVMaterial = FName(TEXT("branchUVMat"));
 	inline static const FName TrunkMaterialPath = FName(TEXT("TrunkMaterialPath"));
 	inline static const FName BranchHierarchyNumber = FName(TEXT("branchHierarchyNumber"));
+	inline static const FName BranchSimulationGroupIndex = FName(TEXT("branchSimulationGroupIndex"));
 
 	//Point Attributes
 	inline static const FName PointPosition = FName(TEXT("Position"));
diff --git a/Engine/Plugins/Experimental/ProceduralVegetationEditor/Source/MegaPlants/Public/Facades/MegaPlantBoneFacade.h b/Engine/Plugins/Experimental/ProceduralVegetationEditor/Source/MegaPlants/Public/Facades/MegaPlantBoneFacade.h
--- a/Engine/Plugins/Experimental/ProceduralVegetationEditor/Source/MegaPlants/Public/Facades/MegaPlantBoneFacade.h
+++ b/Engine/Plugins/Experimental/ProceduralVegetationEditor/Source/MegaPlants/Public/Facades/MegaPlantBoneFacade.h
@@ -49,6 +49,8 @@ namespace MegaPlants::Facades
 		const TManagedArray<int32>& GetVertexPointIds() const;
 		
 		TManagedArray<int32>& ModifyVertexPointIds();
+
+		void SetWindSimulationData(FManagedArrayCollection& Collection);
 		
 		void CreateBoneData(TArray<FBoneNode>& Bones,const float ReductionStrength);
 
diff --git a/Engine/Plugins/Experimental/ProceduralVegetationEditor/Source/MegaPlants/Public/Facades/MegaPlantBranchFacade.h b/Engine/Plugins/Experimental/ProceduralVegetationEditor/Source/MegaPlants/Public/Facades/MegaPlantBranchFacade.h
--- a/Engine/Plugins/Experimental/ProceduralVegetationEditor/Source/MegaPlants/Public/Facades/MegaPlantBranchFacade.h
+++ b/Engine/Plugins/Experimental/ProceduralVegetationEditor/Source/MegaPlants/Public/Facades/MegaPlantBranchFacade.h
@@ -55,6 +55,10 @@ namespace MegaPlants::Facades
 		
 		int32 GetBranchHierarchyNumber(int32 Index) const;
 
+		int32 GetBranchSimulationGroupIndex(int32 Index) const;
+		
+		void SetBranchSimulationGroupIndex(int32 Index, int32 InLogicalDepth);
+
 		void SetTrunkMaterialPath(const FString& InPath);
 
 		FString GetTrunkMaterialPath() const;
@@ -78,6 +82,7 @@ namespace MegaPlants::Facades
 		TManagedArrayAccessor<TArray<int32>> BranchFoliageIDs;
 		TManagedArrayAccessor<int32> BranchUVMaterial;
 		TManagedArrayAccessor<int32> BranchHierarchyNumber;
+		TManagedArrayAccessor<int32> BranchSimulationGroupIndex;
 		TManagedArrayAccessor<FString> TrunkMaterialPathAttribute;
 	};
 }
