Rename MegaPlant to ProceduralVegetation
* MegaPlant, and MegaPlantEditor module has been renamed to ProceduralVegetation and ProceduralVegetationEditor
* MegaPlant prefix has been replaced with "PV" prefix
* MegaPlantRecipie has been replaced with "MegaSeed"

diff --git a/Engine/Plugins/Experimental/ProceduralVegetationEditor/Config/BaseProceduralVegetationEditor.ini b/Engine/Plugins/Experimental/ProceduralVegetationEditor/Config/BaseProceduralVegetationEditor.ini
--- a/Engine/Plugins/Experimental/ProceduralVegetationEditor/Config/BaseProceduralVegetationEditor.ini
+++ b/Engine/Plugins/Experimental/ProceduralVegetationEditor/Config/BaseProceduralVegetationEditor.ini
@@ -1,2 +1,39 @@
 [CoreRedirects]
 +PackageRedirects=(OldName="/MegaPlants/...", NewName="/ProceduralVegetationEditor/", MatchWildcard=true)
++PackageRedirects=(OldName="/Script/MegaPlants", NewName="/Script/ProceduralVegetation")
++PackageRedirects=(OldName="/Script/MegaPlantsEditor", NewName="/Script/ProceduralVegetationEditor")
+
++ClassRedirects=(OldName="MegaPlant",NewName="/Script/ProceduralVegetation.ProceduralVegetation")
++ClassRedirects=(OldName="MegaPlantGraph",NewName="/Script/ProceduralVegetation.ProceduralVegetationGraph")
++ClassRedirects=(OldName="MegaPlantLink",NewName="/Script/ProceduralVegetation.ProceduralVegetationLink")
+
++ClassRedirects=(OldName="PCGMegaPlantRemoveBranchesSettings",NewName="/Script/ProceduralVegetation.PVPCGRemoveBranchesSettings")
++ClassRedirects=(OldName="PCGMegaPlantScaleSettings",NewName="/Script/ProceduralVegetation.PVPCGScaleSettings")
++ClassRedirects=(OldName="PCGMegaPlantRecipeLoadSettings",NewName="/Script/ProceduralVegetation.PVPCGMegaSeedLoadSettings")
++PropertyRedirects=(OldName="PCGMegaPlantRecipeLoadSettings.RecipeDataAsset",NewName="PVPCGMegaSeedLoadSettings.MegaSeedDataAsset")
++ClassRedirects=(OldName="PCGMegaPlantOutputSettings",NewName="/Script/ProceduralVegetation.PVPCGOutputSettings")
++ClassRedirects=(OldName="PCGMegaPlantMesherSettings",NewName="/Script/ProceduralVegetation.PVPCGMesherSettings")
++ClassRedirects=(OldName="PCGMegaPlantImporterSettings",NewName="/Script/ProceduralVegetation.PVPCGImporterSettings")
++ClassRedirects=(OldName="PCGMegaPlantGravitySettings",NewName="/Script/ProceduralVegetation.PVPCGGravitySettings")
++ClassRedirects=(OldName="PCGMegaPlantFoliageDistributorSettings",NewName="/Script/ProceduralVegetation.PVPCGFoliageDistributorSettings")
++ClassRedirects=(OldName="PCGMegaPlantFoliageSettings",NewName="/Script/ProceduralVegetation.PVPCGFoliageSettings")
++ClassRedirects=(OldName="PCGMegaPlantCarveSettings",NewName="/Script/ProceduralVegetation.PVPCGCarveSettings")
++ClassRedirects=(OldName="PCGMegaPlantBaseSettings",NewName="/Script/ProceduralVegetation.PVPCGBaseSettings")
++ClassRedirects=(OldName="PCGMegaPlantBoneReductionSettings",NewName="/Script/ProceduralVegetation.PVPCGBoneReductionSettings")
+
++ClassRedirects=(OldName="MegaPlantRecipeDataAsset",NewName="/Script/ProceduralVegetation.MegaSeedDataAsset")
+
++StructRedirects=(OldName="MegaPlantMesherSettings",NewName="/Script/ProceduralVegetation.PVMesherSettings")
++StructRedirects=(OldName="MegaPlantMaterialSettings",NewName="/Script/ProceduralVegetation.PVMaterialSettings")
++StructRedirects=(OldName="MegaPlantGravitySettings",NewName="/Script/ProceduralVegetation.PVGravitySettings")
++StructRedirects=(OldName="MegaPlantCarveSettings",NewName="/Script/ProceduralVegetation.PVCarveSettings")
++StructRedirects=(OldName="MegaPlantExportSettings",NewName="/Script/ProceduralVegetation.MegaPlantExportSettings")
+
++EnumRedirects=(OldName="MegaPlantExportMeshType",NewName="/Script/ProceduralVegetation.PVExportMeshType")
+
++ClassRedirects=(OldName="MegaPlantEditorSchema",NewName="/Script/ProceduralVegetationEditor.PVEditorSchema")
++ClassRedirects=(OldName="MegaPlantBoneComponent",NewName="/Script/ProceduralVegetationEditor.PVBoneComponent")
++StructRedirects=(OldName="MegaPlantBoneInfo",NewName="/Script/ProceduralVegetationEditor.PVBoneInfo")
++ClassRedirects=(OldName="MegaPlantLineBatchComponent",NewName="/Script/ProceduralVegetationEditor.PVLineBatchComponent")
++StructRedirects=(OldName="MegaPlantLineInfo",NewName="/Script/ProceduralVegetationEditor.PVLineInfo")
++ClassRedirects=(OldName="MegaPlantSkeletonVisualizerComponent",NewName="/Script/ProceduralVegetationEditor.PVSkeletonVisualizerComponent")
diff --git a/Engine/Plugins/Experimental/ProceduralVegetationEditor/ProceduralVegetationEditor.uplugin b/Engine/Plugins/Experimental/ProceduralVegetationEditor/ProceduralVegetationEditor.uplugin
--- a/Engine/Plugins/Experimental/ProceduralVegetationEditor/ProceduralVegetationEditor.uplugin
+++ b/Engine/Plugins/Experimental/ProceduralVegetationEditor/ProceduralVegetationEditor.uplugin
@@ -17,7 +17,7 @@
 	"EnabledByDefault": false,
 	"Modules": [
 		{
-			"Name": "MegaPlants",
+			"Name": "ProceduralVegetation",
 			"Type": "Runtime",
 			"LoadingPhase": "Default",
 			"PlatformAllowList": [
@@ -27,7 +27,7 @@
 			]
 		},
 		{
-			"Name": "MegaPlantsEditor",
+			"Name": "ProceduralVegetationEditor",
 			"Type": "Editor",
 			"LoadingPhase": "Default",
 			"PlatformAllowList": [
diff --git a/Engine/Plugins/Experimental/ProceduralVegetationEditor/Source/ProceduralVegetation/Private/Facades/PVFoliageFacade.cpp b/Engine/Plugins/Experimental/ProceduralVegetationEditor/Source/ProceduralVegetation/Private/Facades/PVFoliageFacade.cpp
new file mode 100644
--- /dev/null
+++ b/Engine/Plugins/Experimental/ProceduralVegetationEditor/Source/ProceduralVegetation/Private/Facades/PVFoliageFacade.cpp
@@ -0,0 +1,364 @@
+﻿// Copyright Epic Games, Inc. All Rights Reserved.
+
+#include "Facades/PVFoliageFacade.h"
+
+#include "ProceduralVegetationModule.h"
+#include "Facades/PVAttributesNames.h"
+
+namespace PV::Facades
+{
+	FFoliageFacade::FFoliageFacade(FManagedArrayCollection& InCollection, const int32 InitialSize)
+		: ConstCollection(InCollection)
+		, Collection(&InCollection)
+		, NamesAttribute(InCollection, AttributeNames::FoliageName, GroupNames::FoliageNamesGroup)
+		, NameIdsAttribute(InCollection, AttributeNames::FoliageNameID, GroupNames::FoliageGroup)
+		, BranchIdsAttribute(InCollection, AttributeNames::FoliageBranchID, GroupNames::FoliageGroup)
+		, PivotPointsAttribute(InCollection, AttributeNames::FoliagePivotPoint, GroupNames::FoliageGroup)
+		, UpVectorsAttribute(InCollection, AttributeNames::FoliageUPVector, GroupNames::FoliageGroup)
+		, NormalVectorsAttribute(InCollection, AttributeNames::FoliageNormalVector, GroupNames::FoliageGroup)
+		, ScalesAttribute(InCollection, AttributeNames::FoliageScale, GroupNames::FoliageGroup)
+		, LengthFromRootAttribute(InCollection, AttributeNames::FoliageLengthFromRoot, GroupNames::FoliageGroup)
+		, ParentBoneIdsAttribute(InCollection, AttributeNames::FoliageParentBoneID, GroupNames::FoliageGroup)
+		, FoliageIdsAttribute(InCollection, AttributeNames::BranchFoliageIDs, GroupNames::BranchGroup)
+		, FoliagePathAttribute(InCollection, AttributeNames::FoliagePath, GroupNames::DetailsGroup)
+	{
+		DefineSchema(InitialSize);
+	}
+
+	FFoliageFacade::FFoliageFacade(const FManagedArrayCollection& InCollection)
+		: ConstCollection(InCollection)
+		, Collection(nullptr)
+		, NamesAttribute(InCollection, AttributeNames::FoliageName, GroupNames::FoliageNamesGroup)
+		, NameIdsAttribute(InCollection, AttributeNames::FoliageNameID, GroupNames::FoliageGroup)
+		, BranchIdsAttribute(InCollection, AttributeNames::FoliageBranchID, GroupNames::FoliageGroup)
+		, PivotPointsAttribute(InCollection, AttributeNames::FoliagePivotPoint, GroupNames::FoliageGroup)
+		, UpVectorsAttribute(InCollection, AttributeNames::FoliageUPVector, GroupNames::FoliageGroup)
+		, NormalVectorsAttribute(InCollection, AttributeNames::FoliageNormalVector, GroupNames::FoliageGroup)
+		, ScalesAttribute(InCollection, AttributeNames::FoliageScale, GroupNames::FoliageGroup)
+		, LengthFromRootAttribute(InCollection, AttributeNames::FoliageLengthFromRoot, GroupNames::FoliageGroup)
+		, ParentBoneIdsAttribute(InCollection, AttributeNames::FoliageParentBoneID, GroupNames::FoliageGroup)
+		, FoliageIdsAttribute(InCollection, AttributeNames::BranchFoliageIDs, GroupNames::BranchGroup)
+		, FoliagePathAttribute(InCollection, AttributeNames::FoliagePath, GroupNames::DetailsGroup)
+	{
+	}
+
+	void FFoliageFacade::DefineSchema(const int32 InitialSize)
+	{
+		check(!IsConst());
+
+		if (!Collection->HasGroup(GroupNames::FoliageGroup))
+		{
+			Collection->AddGroup(GroupNames::FoliageGroup);
+			if (InitialSize > 0)
+			{
+				Collection->AddElements(InitialSize, GroupNames::FoliageGroup);
+			}
+		}
+
+		if (!Collection->HasGroup(GroupNames::FoliageNamesGroup))
+		{
+			Collection->AddGroup(GroupNames::FoliageNamesGroup);
+		}
+
+		if (!Collection->HasGroup(GroupNames::DetailsGroup))
+		{
+			Collection->AddGroup(GroupNames::DetailsGroup);
+		}
+
+		NamesAttribute.Add();
+		NameIdsAttribute.Add();
+		BranchIdsAttribute.Add();
+		PivotPointsAttribute.Add();
+		UpVectorsAttribute.Add();
+		NormalVectorsAttribute.Add();
+		ScalesAttribute.Add();
+		LengthFromRootAttribute.Add();
+		ParentBoneIdsAttribute.Add();
+		FoliageIdsAttribute.Add();
+		FoliagePathAttribute.Add();
+	}
+
+	bool FFoliageFacade::IsValid() const
+	{
+		return NamesAttribute.IsValid() && NameIdsAttribute.IsValid() && BranchIdsAttribute.IsValid() && PivotPointsAttribute.IsValid() &&
+			UpVectorsAttribute.IsValid() && NormalVectorsAttribute.IsValid() && ScalesAttribute.IsValid() && LengthFromRootAttribute.IsValid() &&
+			FoliageIdsAttribute.IsValid() && ParentBoneIdsAttribute.IsValid();
+	}
+
+	FFoliageEntryData FFoliageFacade::GetFoliageEntry(const int32 Index) const
+	{
+		FFoliageEntryData ReturnData = {};
+		if (IsValid() && Index > -1)
+		{
+			if (NameIdsAttribute.IsValidIndex(Index))
+			{
+				ReturnData.NameId = NameIdsAttribute.Get()[Index];
+			}
+			if (BranchIdsAttribute.IsValidIndex(Index))
+			{
+				ReturnData.BranchId = BranchIdsAttribute.Get()[Index];
+			}
+			if (PivotPointsAttribute.IsValidIndex(Index))
+			{
+				ReturnData.PivotPoint = PivotPointsAttribute.Get()[Index];
+			}
+			if (UpVectorsAttribute.IsValidIndex(Index))
+			{
+				ReturnData.UpVector = UpVectorsAttribute.Get()[Index];
+			}
+			if (NormalVectorsAttribute.IsValidIndex(Index))
+			{
+				ReturnData.NormalVector = NormalVectorsAttribute.Get()[Index];
+			}
+			if (ScalesAttribute.IsValidIndex(Index))
+			{
+				ReturnData.Scale = ScalesAttribute.Get()[Index];
+			}
+			if (LengthFromRootAttribute.IsValidIndex(Index))
+			{
+				ReturnData.LengthFromRoot = LengthFromRootAttribute.Get()[Index];
+			}
+			if (ParentBoneIdsAttribute.IsValidIndex(Index))
+			{
+				ReturnData.ParentBoneID = ParentBoneIdsAttribute.Get()[Index];
+			}
+		}
+		return ReturnData;
+	}
+
+	const float& FFoliageFacade::GetLengthFromRoot(const int32 Index) const
+	{
+		if (LengthFromRootAttribute.IsValid() && LengthFromRootAttribute.IsValidIndex(Index))
+		{
+			return LengthFromRootAttribute[Index];
+		}
+
+		static constexpr float DefaultLengthFromRoot = 0.0;
+		return DefaultLengthFromRoot;
+	}
+
+	void FFoliageFacade::SetLengthFromRoot(const int32 Index, const float Input)
+	{
+		check(!IsConst());
+		if (IsValid() && LengthFromRootAttribute.IsValidIndex(Index))
+		{
+			LengthFromRootAttribute.ModifyAt(Index, Input);
+		}
+	}
+
+	int32 FFoliageFacade::GetParentBoneID(const int32 Index) const
+	{
+		if (ParentBoneIdsAttribute.IsValid() && ParentBoneIdsAttribute.IsValidIndex(Index))
+		{
+			return ParentBoneIdsAttribute[Index];
+		}
+
+		return INDEX_NONE;
+	}
+
+	void FFoliageFacade::SetParentBoneID(const int32 Index, const int32 Input)
+	{
+		check(!IsConst());
+		if (IsValid() && ParentBoneIdsAttribute.IsValidIndex(Index))
+		{
+			ParentBoneIdsAttribute.ModifyAt(Index, Input);
+		}
+	}
+
+	const TArray<int32>& FFoliageFacade::GetFoliageEntryIdsForBranch(const int32 Index) const
+	{
+		if (FoliageIdsAttribute.IsValid() && FoliageIdsAttribute.IsValidIndex(Index))
+		{
+			return FoliageIdsAttribute[Index];
+		}
+
+		static const TArray<int32> EmptyArray;
+		return EmptyArray;
+	}
+
+	int32 FFoliageFacade::AddFoliageEntry(const FFoliageEntryData& InputData)
+	{
+		check(!IsConst());
+		if (IsValid())
+		{
+			const int32 NewIndex = NameIdsAttribute.AddElements(1);
+			NameIdsAttribute.Modify()[NewIndex] = InputData.NameId;
+			BranchIdsAttribute.Modify()[NewIndex] = InputData.BranchId;
+			PivotPointsAttribute.Modify()[NewIndex] = InputData.PivotPoint;
+			UpVectorsAttribute.Modify()[NewIndex] = InputData.UpVector;
+			NormalVectorsAttribute.Modify()[NewIndex] = InputData.NormalVector;
+			ScalesAttribute.Modify()[NewIndex] = InputData.Scale;
+			LengthFromRootAttribute.Modify()[NewIndex] = InputData.LengthFromRoot;
+			ParentBoneIdsAttribute.Modify()[NewIndex] = InputData.ParentBoneID;
+
+			return NewIndex;
+		}
+		return INDEX_NONE;
+	}
+
+	void FFoliageFacade::SetFoliageEntry(const int32 Index, const FFoliageEntryData& InputData)
+	{
+		check(!IsConst());
+		if (IsValid() && Index >= 0)
+		{
+			NameIdsAttribute.ModifyAt(Index, InputData.NameId);
+			BranchIdsAttribute.ModifyAt(Index, InputData.BranchId);
+			PivotPointsAttribute.ModifyAt(Index, InputData.PivotPoint);
+			UpVectorsAttribute.ModifyAt(Index, InputData.UpVector);
+			NormalVectorsAttribute.ModifyAt(Index, InputData.NormalVector);
+			ScalesAttribute.ModifyAt(Index, InputData.Scale);
+			LengthFromRootAttribute.ModifyAt(Index, InputData.LengthFromRoot);
+			LengthFromRootAttribute.ModifyAt(Index, InputData.LengthFromRoot);
+			LengthFromRootAttribute.ModifyAt(Index, InputData.LengthFromRoot);
+			ParentBoneIdsAttribute.ModifyAt(Index, InputData.ParentBoneID);
+		}
+	}
+
+	void FFoliageFacade::SetFoliageIdsArray(const int32 Index, const TArray<int32>& InputIds)
+	{
+		check(!IsConst());
+		if (IsValid() && Index >= 0)
+		{
+			FoliageIdsAttribute.Modify()[Index] = InputIds;
+		}
+	}
+
+	void FFoliageFacade::SetFoliageBranchId(const int32 Index, const int32 InputId)
+	{
+		check(!IsConst());
+		if (IsValid() && BranchIdsAttribute.IsValidIndex(Index))
+		{
+			BranchIdsAttribute.ModifyAt(Index, InputId);
+		}
+	}
+
+	FString FFoliageFacade::GetFoliageName(const int32 Index) const
+	{
+		if (NamesAttribute.IsValid() && NamesAttribute.IsValidIndex(Index))
+		{
+			return NamesAttribute[Index];
+		}
+
+		return FString();
+	}
+
+	const FVector3f& FFoliageFacade::GetPivotPoint(const int32 Index) const
+	{
+		if (PivotPointsAttribute.IsValid() && PivotPointsAttribute.IsValidIndex(Index))
+		{
+			return PivotPointsAttribute[Index];
+		}
+
+		static const FVector3f DefaultPivotPoint = FVector3f::Zero();
+		return DefaultPivotPoint;
+	}
+
+	void FFoliageFacade::SetPivotPoint(const int32 Index, const FVector3f& Input)
+	{
+		check(!IsConst());
+		if (IsValid() && PivotPointsAttribute.IsValidIndex(Index))
+		{
+			PivotPointsAttribute.ModifyAt(Index, Input);
+		}
+	}
+
+	void FFoliageFacade::SetScale(const int32 Index, const float& Input)
+	{
+		check(!IsConst());
+		if (IsValid() && ScalesAttribute.IsValidIndex(Index))
+		{
+			ScalesAttribute.ModifyAt(Index, Input);
+		}
+	}
+	
+	void FFoliageFacade::SetFoliageName(const int32 Index, const FString& Input)
+	{
+		check(!IsConst());
+		if (IsValid() && NamesAttribute.IsValidIndex(Index))
+		{
+			NamesAttribute.ModifyAt(Index, Input);
+		}
+	}
+	
+	void FFoliageFacade::SetFoliageNames(const TArray<FString>& InputNames)
+	{
+		check(!IsConst());
+		Collection->EmptyGroup(GroupNames::FoliageNamesGroup);
+		Collection->AddElements(InputNames.Num(), GroupNames::FoliageNamesGroup);
+		for (int32 Index = 0; Index < InputNames.Num(); ++Index)
+		{
+			NamesAttribute.Modify()[Index] = InputNames[Index];
+		}
+	}
+
+	void FFoliageFacade::SetFoliagePath(const FString InPath)
+	{
+		check(!IsConst());
+
+		int32 NumElements = Collection->NumElements(GroupNames::DetailsGroup);
+
+		if (NumElements == 0)
+		{
+			Collection->AddElements(1, GroupNames::DetailsGroup);
+		}
+
+		FoliagePathAttribute.Modify()[0] = InPath;
+	}
+
+	FString FFoliageFacade::GetFoliagePath() const
+	{
+		if (FoliagePathAttribute.IsValid() && FoliagePathAttribute.IsValidIndex(0))
+		{
+			return FoliagePathAttribute[0];
+		}
+
+		return FString();
+	}
+
+	int32 FFoliageFacade::GetElementCount() const
+	{
+		return NameIdsAttribute.Num();
+	}
+
+	void FFoliageFacade::CopyEntry(const int32 FromIndex, const int32 ToIndex)
+	{
+		if (IsValid() && NameIdsAttribute.IsValidIndex(FromIndex) && NameIdsAttribute.IsValidIndex(ToIndex))
+		{
+			NameIdsAttribute.ModifyAt(ToIndex, NameIdsAttribute[FromIndex]);
+			BranchIdsAttribute.ModifyAt(ToIndex, BranchIdsAttribute[FromIndex]);
+			PivotPointsAttribute.ModifyAt(ToIndex, PivotPointsAttribute[FromIndex]);
+			UpVectorsAttribute.ModifyAt(ToIndex, UpVectorsAttribute[FromIndex]);
+			NormalVectorsAttribute.ModifyAt(ToIndex, NormalVectorsAttribute[FromIndex]);
+			ScalesAttribute.ModifyAt(ToIndex, ScalesAttribute[FromIndex]);
+			LengthFromRootAttribute.ModifyAt(ToIndex, LengthFromRootAttribute[FromIndex]);
+			ParentBoneIdsAttribute.ModifyAt(ToIndex, ParentBoneIdsAttribute[FromIndex]);
+		}
+	}
+
+	void FFoliageFacade::RemoveEntries(const int32 NumEntries, const int32 StartIndex)
+	{
+		if (IsValid() && NameIdsAttribute.IsValidIndex(StartIndex) && StartIndex + NumEntries <= NameIdsAttribute.Num())
+		{
+			NameIdsAttribute.RemoveElements(NumEntries, StartIndex);
+		}
+	}
+
+	FTransform FFoliageFacade::GetFoliageTransform(int32 Id) const
+	{
+		FFoliageEntryData Data = GetFoliageEntry(Id);
+		FString FoliageName = GetFoliageName(Data.NameId);
+
+		const FQuat UpVectorRotationQuat = FQuat::FindBetweenNormals(FVector::UpVector, FVector(Data.UpVector));
+		const FVector RotatedVectorFace = UpVectorRotationQuat.RotateVector(FVector::RightVector);
+		float Dot = FVector::DotProduct(RotatedVectorFace, FVector(Data.NormalVector));
+		float Angle = FMath::Acos(FMath::Clamp(Dot, -1.0f, 1.0f));
+		const FQuat FaceRotationQuat = FQuat(FVector(Data.UpVector), Angle);
+
+		const FTransform FaceRotationTransform(FaceRotationQuat);
+		const FTransform UpVectorRotationTransform(UpVectorRotationQuat);
+		const FTransform TranslationTransform(FRotator::ZeroRotator, FVector(Data.PivotPoint));
+		const FTransform FinalTransform = UpVectorRotationTransform * FaceRotationTransform * TranslationTransform;
+
+		return FinalTransform;
+	}
+}
diff --git a/Engine/Plugins/Experimental/ProceduralVegetationEditor/Source/ProceduralVegetation/Private/Facades/PVPointFacade.cpp b/Engine/Plugins/Experimental/ProceduralVegetationEditor/Source/ProceduralVegetation/Private/Facades/PVPointFacade.cpp
new file mode 100644
--- /dev/null
+++ b/Engine/Plugins/Experimental/ProceduralVegetationEditor/Source/ProceduralVegetation/Private/Facades/PVPointFacade.cpp
@@ -0,0 +1,401 @@
+﻿// Copyright Epic Games, Inc. All Rights Reserved.
+
+#include "Facades/PVPointFacade.h"
+#include "Facades/PVAttributesNames.h"
+
+namespace PV::Facades
+{
+	FPointFacade::FPointFacade(FManagedArrayCollection& InCollection)
+		: Positions(InCollection, AttributeNames::PointPosition, GroupNames::PointGroup)
+		, LengthFromRoot(InCollection, AttributeNames::LengthFromRoot, GroupNames::PointGroup)
+		, PointScaleGradient(InCollection, AttributeNames::PointScaleGradient, GroupNames::PointGroup)
+		, HullGradient(InCollection, AttributeNames::HullGradient, GroupNames::PointGroup)
+		, MainTrunkGradient(InCollection, AttributeNames::MainTrunkGradient, GroupNames::PointGroup)
+		, GroundGradient(InCollection, AttributeNames::GroundGradient, GroupNames::PointGroup)
+		, PointScale(InCollection, AttributeNames::PointScale, GroupNames::PointGroup)
+		, BudDirections(InCollection, AttributeNames::BudDirection, GroupNames::PointGroup)
+		, BranchGradients(InCollection, AttributeNames::BranchGradient, GroupNames::PointGroup)
+		, BudHormoneLevels(InCollection, AttributeNames::BudHormoneLevels, GroupNames::PointGroup)
+		, PlantGradients(InCollection, AttributeNames::PlantGradient, GroupNames::PointGroup)
+		, LengthFromSeed(InCollection, AttributeNames::LengthFromSeed, GroupNames::PointGroup)
+		, BudNumber(InCollection, AttributeNames::BudNumber, GroupNames::PointGroup)
+		, NjordPixelIndex(InCollection, AttributeNames::NjordPixelIndex, GroupNames::PointGroup)
+		, BudLightDetected(InCollection, AttributeNames::BudLightDetected, GroupNames::PointGroup)
+		, BudDevelopment(InCollection, AttributeNames::BudDevelopment, GroupNames::PointGroup)
+		, TextureCoordV(InCollection, AttributeNames::TextureCoordV, GroupNames::PointGroup)
+		, TextureCoordUOffset(InCollection, AttributeNames::TextureCoordUOffset, GroupNames::PointGroup)
+	{
+	}
+
+	FPointFacade::FPointFacade(const FManagedArrayCollection& InCollection)
+		: Positions(InCollection, AttributeNames::PointPosition, GroupNames::PointGroup)
+		, LengthFromRoot(InCollection, AttributeNames::LengthFromRoot, GroupNames::PointGroup)
+		, PointScaleGradient(InCollection, AttributeNames::PointScaleGradient, GroupNames::PointGroup)
+		, HullGradient(InCollection, AttributeNames::HullGradient, GroupNames::PointGroup)
+		, MainTrunkGradient(InCollection, AttributeNames::MainTrunkGradient, GroupNames::PointGroup)
+		, GroundGradient(InCollection, AttributeNames::GroundGradient, GroupNames::PointGroup)
+		, PointScale(InCollection, AttributeNames::PointScale, GroupNames::PointGroup)
+		, BudDirections(InCollection, AttributeNames::BudDirection, GroupNames::PointGroup)
+		, BranchGradients(InCollection, AttributeNames::BranchGradient, GroupNames::PointGroup)
+		, BudHormoneLevels(InCollection, AttributeNames::BudHormoneLevels, GroupNames::PointGroup)
+		, PlantGradients(InCollection, AttributeNames::PlantGradient, GroupNames::PointGroup)
+		, LengthFromSeed(InCollection, AttributeNames::LengthFromSeed, GroupNames::PointGroup)
+		, BudNumber(InCollection, AttributeNames::BudNumber, GroupNames::PointGroup)
+		, NjordPixelIndex(InCollection, AttributeNames::NjordPixelIndex, GroupNames::PointGroup)
+		, BudLightDetected(InCollection, AttributeNames::BudLightDetected, GroupNames::PointGroup)
+		, BudDevelopment(InCollection, AttributeNames::BudDevelopment, GroupNames::PointGroup)
+		, TextureCoordV(InCollection, AttributeNames::TextureCoordV, GroupNames::PointGroup)
+		, TextureCoordUOffset(InCollection, AttributeNames::TextureCoordUOffset, GroupNames::PointGroup)
+	{
+	}
+
+	bool FPointFacade::IsValid() const
+	{
+		return Positions.IsValid()
+			&& LengthFromRoot.IsValid()
+			&& LengthFromSeed.IsValid()
+			&& PointScaleGradient.IsValid()
+			&& HullGradient.IsValid()
+			&& MainTrunkGradient.IsValid()
+			&& GroundGradient.IsValid()
+			&& PointScale.IsValid()
+			&& BudDirections.IsValid()
+			&& BranchGradients.IsValid()
+			&& BudHormoneLevels.IsValid()
+			&& PlantGradients.IsValid()
+			&& BudNumber.IsValid()
+			&& NjordPixelIndex.IsValid()
+			&& BudLightDetected.IsValid()
+			&& BudDevelopment.IsValid();
+	}
+
+	int32 FPointFacade::GetElementCount() const
+	{
+		return Positions.Num();
+	}
+
+	const FVector3f& FPointFacade::GetPosition(const int32 Index) const
+	{
+		if (Positions.IsValid() && Positions.IsValidIndex(Index))
+		{
+			return Positions[Index];
+		}
+
+		static const FVector3f DefaultPosition = FVector3f::ZeroVector;
+		return DefaultPosition;
+	}
+
+	const float& FPointFacade::GetLengthFromRoot(const int32 Index) const
+	{
+		if (LengthFromRoot.IsValid() && LengthFromRoot.IsValidIndex(Index))
+		{
+			return LengthFromRoot[Index];
+		}
+		static constexpr float DefaultLengthFromRoot = 0.0;
+		return DefaultLengthFromRoot;
+	}
+
+	const float& FPointFacade::GetLengthFromSeed(const int32 Index) const
+	{
+		if (LengthFromSeed.IsValid() && LengthFromSeed.IsValidIndex(Index))
+		{
+			return LengthFromSeed[Index];
+		}
+		static constexpr float DefaultLengthFromSeed = 0.0;
+		return DefaultLengthFromSeed;
+	}
+
+	const float& FPointFacade::GetPointScaleGradient(const int32 Index) const
+	{
+		if (PointScaleGradient.IsValid() && PointScaleGradient.IsValidIndex(Index))
+		{
+			return PointScaleGradient[Index];
+		}
+
+		static constexpr float DefaultPScale = 0.0;
+		return DefaultPScale;
+	}
+
+	const float& FPointFacade::GetHullGradient(const int32 Index) const
+	{
+		if (HullGradient.IsValid() && HullGradient.IsValidIndex(Index))
+		{
+			return HullGradient[Index];
+		}
+
+		static constexpr float DefaultValue = 0.0;
+		return DefaultValue;
+	}
+
+	const float& FPointFacade::GetMainTrunkGradient(const int32 Index) const
+	{
+		if (MainTrunkGradient.IsValid() && MainTrunkGradient.IsValidIndex(Index))
+		{
+			return MainTrunkGradient[Index];
+		}
+
+		static constexpr float DefaultValue = 0.0;
+		return DefaultValue;
+	}
+
+	const float& FPointFacade::GetGroundGradient(const int32 Index) const
+	{
+		if (GroundGradient.IsValid() && GroundGradient.IsValidIndex(Index))
+		{
+			return GroundGradient[Index];
+		}
+
+		static constexpr float DefaultValue = 0.0;
+		return DefaultValue;
+	}
+
+	const float& FPointFacade::GetPointScale(const int32 Index) const
+	{
+		if (PointScale.IsValid() && PointScale.IsValidIndex(Index))
+		{
+			return PointScale[Index];
+		}
+
+		static constexpr float DefaultPScale = 0.0;
+		return DefaultPScale;
+	}
+
+	const int32 FPointFacade::GetBudNumber(const int32 Index) const
+	{
+		if (BudNumber.IsValid() && BudNumber.IsValidIndex(Index))
+		{
+			return BudNumber[Index];
+		}
+
+		return INDEX_NONE;
+	}
+	
+	const TArray<float>& FPointFacade::GetBudLightDetected(const int32 Index) const
+	{
+		if (BudLightDetected.IsValid() && BudLightDetected.IsValidIndex(Index))
+		{
+			return BudLightDetected[Index];
+		}
+
+		static const TArray<float> EmptyArray;
+		return EmptyArray;
+	}
+
+	const TArray<int>& FPointFacade::GetBudDevelopment(const int32 Index) const
+	{
+		if (BudDevelopment.IsValid() && BudDevelopment.IsValidIndex(Index))
+		{
+			return BudDevelopment[Index];
+		}
+
+		static const TArray<int> EmptyArray;
+		return EmptyArray;
+	}
+
+	const TArray<FVector3f>& FPointFacade::GetBudDirection(const int32 Index) const
+	{
+		if (BudDirections.IsValid() && BudDirections.IsValidIndex(Index))
+		{
+			return BudDirections[Index];
+		}
+
+		static const TArray<FVector3f> EmptyArray;
+		return EmptyArray;
+	}
+
+	const TArray<float>& FPointFacade::GetBudHormoneLevels(const int32 Index) const
+	{
+		if (BudHormoneLevels.IsValid() && BudHormoneLevels.IsValidIndex(Index))
+		{
+			return BudHormoneLevels[Index];
+		}
+
+		static const TArray<float> EmptyArray;
+		return EmptyArray;
+	}
+
+	void FPointFacade::SetBudHormoneLevels(const int32 Index, const TArray<float>& InBudHormoneLevels)
+	{
+		if (BudHormoneLevels.IsValid() && BudHormoneLevels.IsValidIndex(Index))
+		{
+			BudHormoneLevels.ModifyAt(Index, InBudHormoneLevels);
+		}
+	}
+
+	const float& FPointFacade::GetBranchGradient(const int32 Index) const
+	{
+		if (BranchGradients.IsValid() && BranchGradients.IsValidIndex(Index))
+		{
+			return BranchGradients[Index];
+		}
+
+		static constexpr float DefaultBranchGradient = 0.0;
+		return DefaultBranchGradient;
+	}
+
+	void FPointFacade::SetBranchGradient(const int32 Index, const float InBranchGradient)
+	{
+		if (BranchGradients.IsValid() && BranchGradients.IsValidIndex(Index))
+		{
+			BranchGradients.ModifyAt(Index, InBranchGradient);
+		}
+	}
+
+	const float& FPointFacade::GetPlantGradient(const int32 Index) const
+	{
+		if (PlantGradients.IsValid() && PlantGradients.IsValidIndex(Index))
+		{
+			return PlantGradients[Index];
+		}
+
+		static constexpr float DefaultPlantGradient = 0.0;
+		return DefaultPlantGradient;
+	}
+
+	const float& FPointFacade::GetNjordPixelIndex(const int32 Index) const
+	{
+		if (NjordPixelIndex.IsValid() && NjordPixelIndex.IsValidIndex(Index))
+		{
+			return NjordPixelIndex[Index];
+		}
+
+		static constexpr float IndexNone = -1.0;
+		return IndexNone;
+	}
+
+	void FPointFacade::SetPlantGradient(const int32 Index, const float InPlantGradient)
+	{
+		if (PlantGradients.IsValid() && PlantGradients.IsValidIndex(Index))
+		{
+			PlantGradients.ModifyAt(Index, InPlantGradient);
+		}
+	}
+
+	const float& FPointFacade::GetTextureCoordV(const int32 Index) const
+	{
+		if (TextureCoordV.IsValid() && TextureCoordV.IsValidIndex(Index))
+		{
+			return TextureCoordV[Index];
+		}
+
+		static constexpr float IndexNone = -1.0;
+		return IndexNone;
+	}
+
+	void FPointFacade::SeTextureCoordV(const int32 Index, const float InTextureCoordV)
+	{
+		if (!TextureCoordV.IsValid())
+		{
+			TextureCoordV.Add();
+		}
+		
+		if (TextureCoordV.IsValid() && TextureCoordV.IsValidIndex(Index))
+		{
+			TextureCoordV.ModifyAt(Index, InTextureCoordV);
+		}
+	}
+
+	const float& FPointFacade::GetTextureCoordUOffset(const int32 Index) const
+	{
+		if (TextureCoordUOffset.IsValid() && TextureCoordUOffset.IsValidIndex(Index))
+		{
+			return TextureCoordUOffset[Index];
+		}
+
+		static constexpr float IndexNone = -1.0;
+		return IndexNone;
+	}
+
+	void FPointFacade::SetTextureCoordUOffset(const int32 Index, const float InTextureCoordUOffset)
+	{
+		if (!TextureCoordUOffset.IsValid())
+		{
+			TextureCoordUOffset.Add();
+		}
+		
+		if (TextureCoordUOffset.IsValid() && TextureCoordUOffset.IsValidIndex(Index))
+		{
+			TextureCoordUOffset.ModifyAt(Index, InTextureCoordUOffset);
+		}
+	}
+
+	TManagedArray<FVector3f>& FPointFacade::ModifyPositions()
+	{
+		return Positions.Modify();
+	}
+
+	TManagedArray<float>& FPointFacade::ModifyPointScales()
+	{
+		return PointScale.Modify();
+	}
+
+	TManagedArray<float>& FPointFacade::ModifyLengthFromRoots()
+	{
+		return LengthFromRoot.Modify();
+	}
+
+	TManagedArray<float>& FPointFacade::ModifyLengthFromSeeds()
+	{
+		return LengthFromSeed.Modify();
+	}
+
+	TManagedArray<float>& FPointFacade::ModifyNjordPixelIDs()
+	{
+		return NjordPixelIndex.Modify();
+	}
+
+	const TManagedArray<FVector3f>& FPointFacade::GetPositions() const
+	{
+		return Positions.Get();
+	}
+
+	const TManagedArray<float>& FPointFacade::GetPointScales() const
+	{
+		return PointScale.Get();
+	}
+
+	const TManagedArray<float>& FPointFacade::GetLengthFromRootsArray() const
+	{
+		return LengthFromRoot.Get();
+	}
+
+	void FPointFacade::CopyEntry(const int32 FromIndex, const int32 ToIndex)
+	{
+		if (IsValid() && Positions.IsValidIndex(FromIndex) && Positions.IsValidIndex(ToIndex))
+		{
+			Positions.ModifyAt(ToIndex, Positions[FromIndex]);
+			LengthFromRoot.ModifyAt(ToIndex, LengthFromRoot[FromIndex]);
+			PointScaleGradient.ModifyAt(ToIndex, PointScaleGradient[FromIndex]);
+			HullGradient.ModifyAt(ToIndex, HullGradient[FromIndex]);
+			MainTrunkGradient.ModifyAt(ToIndex, MainTrunkGradient[FromIndex]);
+			GroundGradient.ModifyAt(ToIndex, GroundGradient[FromIndex]);
+			PointScale.ModifyAt(ToIndex, PointScale[FromIndex]);
+			BudDirections.ModifyAt(ToIndex, BudDirections[FromIndex]);
+			BranchGradients.ModifyAt(ToIndex, BranchGradients[FromIndex]);
+			BudHormoneLevels.ModifyAt(ToIndex, BudHormoneLevels[FromIndex]);
+			PlantGradients.ModifyAt(ToIndex, PlantGradients[FromIndex]);
+			LengthFromSeed.ModifyAt(ToIndex, LengthFromSeed[FromIndex]);
+			BudNumber.ModifyAt(ToIndex, BudNumber[FromIndex]);
+			NjordPixelIndex.ModifyAt(ToIndex, NjordPixelIndex[FromIndex]);
+			BudLightDetected.ModifyAt(ToIndex, BudLightDetected[FromIndex]);
+			BudDevelopment.ModifyAt(ToIndex, BudDevelopment[FromIndex]);
+			
+			if(TextureCoordV.IsValid())
+			{
+				TextureCoordV.ModifyAt(ToIndex, TextureCoordV[FromIndex]);
+			}
+			if(TextureCoordUOffset.IsValid())
+			{
+				TextureCoordUOffset.ModifyAt(ToIndex, TextureCoordUOffset[FromIndex]);
+			}
+		}
+	}
+
+	void FPointFacade::RemoveEntries(const int32 NumEntries, const int32 StartIndex)
+	{
+		if (IsValid() && Positions.IsValidIndex(StartIndex) && StartIndex + NumEntries <= Positions.Num())
+		{
+			Positions.RemoveElements(NumEntries, StartIndex);
+		}
+	}
+}
diff --git a/Engine/Plugins/Experimental/ProceduralVegetationEditor/Source/ProceduralVegetation/Private/Facades/PVVectorsFacade.cpp b/Engine/Plugins/Experimental/ProceduralVegetationEditor/Source/ProceduralVegetation/Private/Facades/PVVectorsFacade.cpp
new file mode 100644
--- /dev/null
+++ b/Engine/Plugins/Experimental/ProceduralVegetationEditor/Source/ProceduralVegetation/Private/Facades/PVVectorsFacade.cpp
@@ -0,0 +1,41 @@
+﻿// Copyright Epic Games, Inc. All Rights Reserved.
+
+#include "PVBudVectorsFacade.h"
+#include "Facades/PVAttributesNames.h"
+
+namespace PV::Facades
+{
+	FBudVectorsFacade::FBudVectorsFacade(FManagedArrayCollection& InCollection)
+			: BudDirection(InCollection, AttributeNames::BudDirection, GroupNames::PointGroup)
+	{}
+
+	FBudVectorsFacade::FBudVectorsFacade(const FManagedArrayCollection& InCollection)
+		: BudDirection(InCollection, AttributeNames::BudDirection, GroupNames::PointGroup)
+	{}
+
+	bool FBudVectorsFacade::IsValid() const
+	{
+		return BudDirection.IsValid();
+	}
+
+	const TArray<FVector3f>& FBudVectorsFacade::GetBudDirection(int32 Index) const
+	{
+		if(BudDirection.IsValid() && BudDirection.IsValidIndex(Index))
+		{
+			return BudDirection[Index];
+		}
+
+		static const TArray<FVector3f> EmptyBudDirections;
+		return EmptyBudDirections;
+	}
+
+	TManagedArray<TArray<FVector3f>>& FBudVectorsFacade::ModifyBudDirections()
+	{
+		return BudDirection.Modify();
+	}
+
+	const TManagedArray<TArray<FVector3f>>& FBudVectorsFacade::GetBudDirections() const
+	{
+		return BudDirection.Get();
+	}
+}
diff --git a/Engine/Plugins/Experimental/ProceduralVegetationEditor/Source/ProceduralVegetation/Private/Helpers/PVExportHelper.cpp b/Engine/Plugins/Experimental/ProceduralVegetationEditor/Source/ProceduralVegetation/Private/Helpers/PVExportHelper.cpp
new file mode 100644
--- /dev/null
+++ b/Engine/Plugins/Experimental/ProceduralVegetationEditor/Source/ProceduralVegetation/Private/Helpers/PVExportHelper.cpp
@@ -0,0 +1,983 @@
+﻿// Copyright Epic Games, Inc. All Rights Reserved.
+
+#include "Helpers/PVExportHelper.h"
+
+#include "DynamicMeshEditor.h"
+#if WITH_EDITOR
+#include "DynamicWindImportData.h"
+#endif
+#include "DynamicWindSkeletalData.h"
+#include "ProceduralVegetationLink.h"
+#include "ProceduralVegetationModule.h"
+#include "MeshDescriptionToDynamicMesh.h"
+#include "PlanarCut.h"
+#include "RenderUtils.h"
+#include "StaticMeshAttributes.h"
+#include "Animation/Skeleton.h"
+#include "AssetRegistry/IAssetRegistry.h"
+#include "Components/InstancedStaticMeshComponent.h"
+#include "ConversionUtils/SceneComponentToDynamicMesh.h"
+#include "DynamicMesh/DynamicBoneAttribute.h"
+#include "DynamicMesh/DynamicVertexSkinWeightsAttribute.h"
+
+#include "Engine/SkeletalMesh.h"
+#include "Engine/StaticMesh.h"
+#include "Engine/World.h"
+#include "Facades/PVBranchFacade.h"
+#include "Facades/PVPointFacade.h"
+#include "GeometryCollection/GeometryCollection.h"
+#include "GeometryScript/GeometryScriptTypes.h"
+#include "GeometryScript/MeshAssetFunctions.h"
+#include "GeometryScript/MeshBoneWeightFunctions.h"
+#include "Helpers/PVUtilities.h"
+#include "Misc/Paths.h"
+#include "UObject/Package.h"
+#include "ProceduralVegetation.h"
+#include "GeometryCollection/Facades/CollectionConstraintOverrideFacade.h"
+#include "Rendering/RenderCommandPipes.h"
+
+namespace PV::Export::Internal
+{
+	TObjectPtr<UDynamicMesh> CollectionToDynamicMesh(const FManagedArrayCollection& Collection)
+	{
+		TRACE_CPUPROFILER_EVENT_SCOPE(PV::Export::CollectionToDynamicMesh);
+		
+		TObjectPtr<UDynamicMesh> NewMesh = NewObject<UDynamicMesh>();
+
+		if (Collection.NumElements(FGeometryCollection::TransformGroup) > 0)
+		{
+			if (const TUniquePtr<FGeometryCollection> GeomCollection = TUniquePtr<FGeometryCollection>(Collection.NewCopy<FGeometryCollection>()))
+			{
+				const TManagedArray<FTransform3f>& BoneTransforms = Collection.GetAttribute<FTransform3f>(
+					"Transform", FGeometryCollection::TransformGroup);
+
+				TArray<int32> TransformIndices;
+				TransformIndices.AddUninitialized(BoneTransforms.Num());
+
+				int32 Idx = 0;
+				for (int32& TransformIdx : TransformIndices)
+				{
+					TransformIdx = Idx++;
+				}
+
+				FMeshDescription MeshDescription;
+				FStaticMeshAttributes Attributes(MeshDescription);
+				Attributes.Register();
+
+				FTransform TransformOut;
+
+				constexpr bool bCenterPivot = false;
+				ConvertToMeshDescription(MeshDescription, TransformOut, bCenterPivot, *GeomCollection, BoneTransforms, TransformIndices);
+
+				NewMesh->Reset();
+
+				UE::Geometry::FDynamicMesh3& DynMesh = NewMesh->GetMeshRef();
+				{
+					FMeshDescriptionToDynamicMesh ConverterToDynamicMesh;
+					ConverterToDynamicMesh.Convert(&MeshDescription, DynMesh);
+				}
+			}
+		}
+
+		return NewMesh;
+	}
+
+	FGeometryScriptCopyMeshToAssetOptions GetCopyMeshToAssetOptions(TArray<TObjectPtr<UMaterialInterface>> InMaterials)
+	{
+		FGeometryScriptCopyMeshToAssetOptions Options;
+		Options.bEnableRecomputeTangents = true;
+
+		Options.bReplaceMaterials = true;
+		Options.NewMaterials = InMaterials;
+
+		return Options;
+	}
+
+#if WITH_EDITORONLY_DATA
+	FGeometryScriptCopyMeshToAssetOptions GetCopyMeshToAssetOptions(TArray<TObjectPtr<UMaterialInterface>> InMaterials, const FMeshNaniteSettings& InNaniteSettings, const ENaniteShapePreservation InShapePreservation)
+	{
+		FGeometryScriptCopyMeshToAssetOptions Options = GetCopyMeshToAssetOptions(InMaterials);
+		Options.NewNaniteSettings = InNaniteSettings;
+		Options.NewNaniteSettings.bEnabled = true;
+		Options.NewNaniteSettings.ShapePreservation = InShapePreservation;
+		Options.bApplyNaniteSettings = true;
+
+		return Options;
+	}
+#endif
+
+	FString GetFoliageName(const Facades::FFoliageFacade& FoliageFacade, int32 Id)
+	{
+		Facades::FFoliageEntryData Data = FoliageFacade.GetFoliageEntry(Id);
+		return FoliageFacade.GetFoliageName(Data.NameId);
+	}
+
+	TArray<TObjectPtr<UMaterialInterface>> GetMaterialsFromCollection(const FManagedArrayCollection& Collection)
+	{
+		TArray<TObjectPtr<UMaterialInterface>> Materials;
+		if (Collection.HasAttribute("MaterialPath", FGeometryCollection::MaterialGroup))
+		{
+			const TManagedArray<FString>& MaterialPaths = Collection.GetAttribute<FString>("MaterialPath", FGeometryCollection::MaterialGroup);
+			for (const FString& Path : MaterialPaths)
+			{
+				Materials.Add(LoadObject<UMaterialInterface>(nullptr, *Path));
+			}
+		}
+		return Materials;
+	}
+
+	static bool SkeletalMeshToDynamicMesh(const USkeletalMesh* SkeletalMesh, const int32 LodIndex, FDynamicMesh3& ToDynamicMesh)
+	{
+#if WITH_EDITOR
+		if (SkeletalMesh->HasMeshDescription(LodIndex))
+		{
+			const FMeshDescription* SourceMesh = SkeletalMesh->GetMeshDescription(LodIndex);
+			if (!SourceMesh)
+			{
+				return false;
+			}
+
+			FMeshDescriptionToDynamicMesh Converter;
+			Converter.Convert(SourceMesh, ToDynamicMesh);
+			return true;
+		}
+#else
+		UE_LOG(LogProceduralVegetation, Warning, TEXT("SkeletalMeshToDynamicMesh is only supported in editor."));
+#endif
+		return false;
+	}
+
+	void CombineMeshInstancesToDynamicMesh(
+		const FManagedArrayCollection& Collection,
+		UDynamicMesh* DynamicMesh,
+		TArray<TObjectPtr<UMaterialInterface>>& Materials,
+		const FMeshInstanceCombined& OnMeshInstanceCombined
+	)
+	{
+		Facades::FFoliageFacade FoliageFacade(Collection);
+		int32 NumInstances = FoliageFacade.NumFoliageEntries();
+		
+		if(NumInstances > 0)
+		{
+			TMap<FString, FDynamicMesh3> FoliageDynamicMeshes;
+			TMap<FString, TObjectPtr<UStaticMesh>> FoliageStaticMeshes;
+			TMap<FString, TObjectPtr<USkeletalMesh>> FoliageSkeletalMeshes;
+			
+			for(int32 FoliageNameIndex = 0; FoliageNameIndex < FoliageFacade.NumFoliageNames(); ++FoliageNameIndex)
+			{
+				FString FoliageName = FoliageFacade.GetFoliageName(FoliageNameIndex);
+				const TObjectPtr<UStaticMesh> StaticMesh = LoadObject<UStaticMesh>(nullptr, *FoliageName);
+				const TObjectPtr<USkeletalMesh> SkeletalMesh = LoadObject<USkeletalMesh>(nullptr, *FoliageName);
+
+				FDynamicMesh3 FoliageDynamicMesh;
+				FText OutErrorMessage;
+
+				if(StaticMesh)
+				{
+					FoliageStaticMeshes.FindOrAdd(FoliageName) = StaticMesh;
+					
+					for (auto Material : StaticMesh->GetStaticMaterials())
+					{
+						Materials.AddUnique(Material.MaterialInterface);
+					}
+
+					UE::Conversion::FStaticMeshConversionOptions Options;
+					if(StaticMeshToDynamicMesh(StaticMesh, FoliageDynamicMesh, OutErrorMessage, Options))
+					{
+						FoliageDynamicMeshes.Add(FoliageName, MoveTemp(FoliageDynamicMesh));
+					}
+				}
+				else if (SkeletalMesh)
+				{
+					FoliageSkeletalMeshes.FindOrAdd(FoliageName) = SkeletalMesh;
+					
+					for (auto Material : SkeletalMesh->GetMaterials())
+					{
+						Materials.AddUnique(Material.MaterialInterface);
+					}
+
+					if(SkeletalMeshToDynamicMesh(SkeletalMesh, 0, FoliageDynamicMesh))
+					{
+						FoliageDynamicMesh.Attributes()->RemoveAttribute("BoneIndex");
+						FoliageDynamicMesh.Attributes()->RemoveAttribute("BoneParentIndex");
+						FoliageDynamicMeshes.Add(FoliageName, MoveTemp(FoliageDynamicMesh));
+					}
+				}
+			}
+
+			for(int32 Id = 0; Id < NumInstances; Id++)
+			{
+				FString FoliageName = GetFoliageName(FoliageFacade, Id);
+
+				if(FoliageDynamicMeshes.Find(FoliageName))
+				{
+					FDynamicMesh3& FoliageDynamicMesh = FoliageDynamicMeshes[FoliageName];
+
+					const TObjectPtr<UStaticMesh> StaticMesh = FoliageStaticMeshes.Contains(FoliageName) ? FoliageStaticMeshes[FoliageName] : nullptr;
+					const TObjectPtr<USkeletalMesh> SkeletalMesh = FoliageSkeletalMeshes.Contains(FoliageName) ? FoliageSkeletalMeshes[FoliageName] : nullptr;
+
+					if (SkeletalMesh)
+					{
+						int32 FoliageParentBoneID = FoliageFacade.GetParentBoneID(Id);
+
+						//Reparent root bone according to the instanced transform
+						UE::Geometry::FDynamicMeshAttributeSet* MeshAttributes = FoliageDynamicMesh.Attributes();
+						UE::Geometry::FDynamicMeshBoneParentIndexAttribute* BoneParentIndices = MeshAttributes->GetBoneParentIndices();
+						BoneParentIndices->SetValue(0,FoliageParentBoneID);
+					}
+				
+					FTransform Transform = FoliageFacade.GetFoliageTransform(Id);
+					UE::Geometry::FTransformSRT3d GeoTransform(Transform);
+
+					DynamicMesh->EditMesh([&](FDynamicMesh3& Mesh)
+					{
+						int32 VertexStart = Mesh.VertexCount();
+						UE::Geometry::FDynamicMeshEditor Editor = UE::Geometry::FDynamicMeshEditor(&Mesh);
+
+						UE::Geometry::FMeshIndexMappings IndexMap;
+						Editor.AppendMesh(
+							&FoliageDynamicMesh,
+							IndexMap,
+							[&](int, const FVector3d& Position)
+							{
+								return GeoTransform.TransformPosition(Position);
+							},
+							[&](int, const FVector3d& Normal)
+							{
+								return GeoTransform.TransformVectorNoScale(Normal);
+							}
+						);
+
+						UE::Geometry::FIndexMapi& TriangleMap = IndexMap.GetTriangleMap();
+							
+						//Remap Material IDs
+						if (Mesh.Attributes()->HasMaterialID())
+						{
+							UE::Geometry::FDynamicMeshMaterialAttribute* MaterialIDs = Mesh.Attributes()->GetMaterialID();
+
+							for (const auto& [OldTid, NewTid] : TriangleMap.GetForwardMap())
+							{
+								const int32 MaterialID = MaterialIDs->GetValue(OldTid);
+								if (StaticMesh)
+								{
+									UMaterialInterface* Material = StaticMesh->GetMaterial(MaterialID);
+									const int32 NewMaterialID = Materials.Find(Material);
+									MaterialIDs->SetValue(NewTid, NewMaterialID);
+								}
+								if (SkeletalMesh)
+								{
+									UMaterialInterface* Material = SkeletalMesh->GetMaterials()[MaterialID].MaterialInterface;
+									const int32 NewMaterialID = Materials.Find(Material);
+									MaterialIDs->SetValue(NewTid, NewMaterialID);
+								}
+							}
+						}
+
+						OnMeshInstanceCombined(StaticMesh->GetFullName(), Id, VertexStart, Mesh.VertexCount());
+					});
+				}
+			}
+		}
+	}
+
+#if WITH_EDITOR
+	void BuildNaniteAssemblyData(
+		const FManagedArrayCollection& Collection,
+		UStaticMesh* StaticMesh
+	)
+	{
+		//Create NaniteAssemblyDataBuilder and the main static mesh materials to it
+		FNaniteAssemblyDataBuilder AssemblyBuilder;
+		for (const FStaticMaterial& Material : StaticMesh->GetStaticMaterials())
+		{
+			AssemblyBuilder.AddMaterialSlot(Material.MaterialInterface);
+		}
+
+		Facades::FFoliageFacade FoliageFacade(Collection);
+		
+		TMap<FString, int32> MeshNamePartMap;
+		TMap<FString, TArray<TObjectPtr<UMaterialInterface>>> MeshMaterialsMap;
+		
+		for(int32 FoliageNameIndex = 0; FoliageNameIndex < FoliageFacade.NumFoliageNames(); ++FoliageNameIndex)
+		{
+			FString FoliageName = FoliageFacade.GetFoliageName(FoliageNameIndex);
+			TObjectPtr<UStaticMesh> PartStaticMesh = LoadObject<UStaticMesh>(nullptr, *FoliageName);
+			TObjectPtr<USkeletalMesh> PartSkeletalMesh = LoadObject<USkeletalMesh>(nullptr, *FoliageName);
+
+			if(PartStaticMesh)
+			{
+				int32 PartIndex = AssemblyBuilder.FindOrAddPart(PartStaticMesh.GetPath());
+				MeshNamePartMap.FindOrAdd(PartStaticMesh.GetPath()) = PartIndex;
+
+				for (auto Material : PartStaticMesh->GetStaticMaterials())
+				{
+					MeshMaterialsMap.FindOrAdd(FoliageName).AddUnique(Material.MaterialInterface);
+				}
+			}
+			else if(PartSkeletalMesh)
+			{
+				auto StaticMeshName = "SM_" + PartSkeletalMesh.GetName();
+				auto StaticMeshPath = PartSkeletalMesh.GetPath().Replace(*PartSkeletalMesh.GetName(), *StaticMeshName);
+				int32 PartIndex = AssemblyBuilder.FindOrAddPart(StaticMeshPath);
+				MeshNamePartMap.FindOrAdd(FoliageName) = PartIndex;
+		
+				TArray<FAssetData> AssetData;
+				IAssetRegistry::Get()->GetAssetsByPath(*PartSkeletalMesh.GetPackage()->GetName(), AssetData);
+		
+				if(AssetData.Num() <= 0)
+				{
+					FDynamicMesh3 PartDynamicMesh;
+					FText OutErrorMessage;
+
+					TObjectPtr<UDynamicMesh> DynamicMesh = NewObject<UDynamicMesh>();
+					
+					TArray<TObjectPtr<UMaterialInterface>> PartMaterials;
+					for (auto Material : PartSkeletalMesh->GetMaterials())
+					{
+						PartMaterials.Add(Material.MaterialInterface);
+						MeshMaterialsMap.FindOrAdd(FoliageName).AddUnique(Material.MaterialInterface);
+					}
+
+					if (SkeletalMeshToDynamicMesh(PartSkeletalMesh, 0, PartDynamicMesh))
+					{
+						auto MeshPackagePath = PartSkeletalMesh->GetPackage()->GetName().Replace(*PartSkeletalMesh.GetName(), *StaticMeshName);
+						UPackage* MeshPackage = CreatePackage(*MeshPackagePath);
+						
+						// Create a skeletal mesh to save in content browser
+						PartStaticMesh = NewObject<UStaticMesh>(MeshPackage, FName(StaticMeshName), RF_Standalone | RF_Public);
+
+						//Convert the dynamic mesh to static mesh
+						EGeometryScriptOutcomePins Output;
+#if WITH_EDITORONLY_DATA
+						FGeometryScriptCopyMeshToAssetOptions Options = GetCopyMeshToAssetOptions(PartMaterials, PartStaticMesh->GetNaniteSettings(), ENaniteShapePreservation::Voxelize);
+#else
+						FGeometryScriptCopyMeshToAssetOptions Options = GetCopyMeshToAssetOptions(PartMaterials);
+#endif
+						DynamicMesh->SetMesh(MoveTemp(PartDynamicMesh));
+
+						UGeometryScriptLibrary_StaticMeshFunctions::CopyMeshToStaticMesh(DynamicMesh, PartStaticMesh, Options, {}, Output);
+						IAssetRegistry::Get()->AssetCreated(PartStaticMesh);
+					}
+				}
+			}
+		}
+
+		AddNodeToBuilder(AssemblyBuilder, Collection, MeshNamePartMap, MeshMaterialsMap);
+		AssemblyBuilder.ApplyToStaticMesh(*StaticMesh);
+	}
+
+	void BuildNaniteAssemblyData(
+		const FManagedArrayCollection& Collection,
+		USkeletalMesh* SkeletalMesh
+	)
+	{
+		//Create NaniteAssemblyDataBuilder and the main static mesh materials to it
+		FNaniteAssemblyDataBuilder AssemblyBuilder;
+		Facades::FFoliageFacade FoliageFacade(Collection);
+		
+		for (const FSkeletalMaterial& Material : SkeletalMesh->GetMaterials())
+		{
+			AssemblyBuilder.AddMaterialSlot(Material.MaterialInterface);
+		}
+
+		TMap<FString, int32> MeshNamePartMap;
+		TMap<FString, TArray<TObjectPtr<UMaterialInterface>>> MeshMaterialsMap;
+			
+		for(int32 FoliageNameIndex = 0; FoliageNameIndex < FoliageFacade.NumFoliageNames(); ++FoliageNameIndex)
+		{
+			FString FoliageName = FoliageFacade.GetFoliageName(FoliageNameIndex);
+			TObjectPtr<UStaticMesh> PartStaticMesh = LoadObject<UStaticMesh>(nullptr, *FoliageName);
+			TObjectPtr<USkeletalMesh> PartSkeletalMesh = LoadObject<USkeletalMesh>(nullptr, *FoliageName);
+
+			if(PartStaticMesh)
+			{
+				auto SkeletalMeshName = "SKM_" + PartStaticMesh.GetName();
+				auto SkeletalMeshPath = PartStaticMesh.GetPath().Replace(*PartStaticMesh.GetName(), *SkeletalMeshName);
+				int32 PartIndex = AssemblyBuilder.FindOrAddPart(SkeletalMeshPath);
+				MeshNamePartMap.FindOrAdd(FoliageName) = PartIndex;
+		
+				TArray<FAssetData> AssetData;
+				IAssetRegistry::Get()->GetAssetsByPath(*PartStaticMesh.GetPackage()->GetName(), AssetData);
+		
+				if(AssetData.Num() <= 0)
+				{
+					FDynamicMesh3 PartDynamicMesh;
+					FText OutErrorMessage;
+
+					TArray<TObjectPtr<UMaterialInterface>> PartMaterials;
+					for (auto Material : PartStaticMesh->GetStaticMaterials())
+					{
+						PartMaterials.Add(Material.MaterialInterface);
+						MeshMaterialsMap.FindOrAdd(FoliageName).AddUnique(Material.MaterialInterface);
+					}
+
+					UE::Conversion::FStaticMeshConversionOptions Options;
+
+					if (StaticMeshToDynamicMesh(PartStaticMesh, PartDynamicMesh, OutErrorMessage, Options))
+					{
+						auto MeshPackagePath = PartStaticMesh->GetPackage()->GetName().Replace(*PartStaticMesh.GetName(), *SkeletalMeshName);
+						UPackage* MeshPackage = CreatePackage(*MeshPackagePath);
+						auto SkeletonName = "SK_" + PartStaticMesh.GetName();
+						auto SkeletonPath = PartStaticMesh.GetPackage()->GetName().Replace(*PartStaticMesh.GetName(), *SkeletonName);
+						UPackage* SkeletonPackage = CreatePackage(*(SkeletonPath));
+
+						// Create a skeletal mesh to save in content browser
+						PartSkeletalMesh = NewObject<USkeletalMesh>(MeshPackage, FName(SkeletalMeshName), RF_Standalone | RF_Public);
+						const TObjectPtr<USkeleton> PartSkeleton = NewObject<USkeleton>(SkeletonPackage, FName(SkeletonName),
+																						  RF_Standalone | RF_Public);
+						PartSkeletalMesh->SetSkeleton(PartSkeleton);
+
+						ConvertToDefaultSkeletalMesh(PartSkeletalMesh, &PartDynamicMesh, PartMaterials);
+
+						IAssetRegistry::Get()->AssetCreated(PartSkeletalMesh);
+						IAssetRegistry::Get()->AssetCreated(PartSkeleton);
+					}
+				}
+			}
+			else if (PartSkeletalMesh)
+			{
+				int32 PartIndex = AssemblyBuilder.FindOrAddPart(PartSkeletalMesh.GetPath());
+				MeshNamePartMap.FindOrAdd(PartSkeletalMesh.GetPath()) = PartIndex;
+
+				for (auto Material : PartSkeletalMesh->GetMaterials())
+				{
+					MeshMaterialsMap.FindOrAdd(FoliageName).AddUnique(Material.MaterialInterface);
+				}
+			}
+		}
+
+		AddNodeToBuilder(AssemblyBuilder, Collection, MeshNamePartMap, MeshMaterialsMap);
+		AssemblyBuilder.ApplyToSkeletalMesh(*SkeletalMesh);
+	}
+
+	void AddNodeToBuilder(FNaniteAssemblyDataBuilder& AssemblyBuilder, const FManagedArrayCollection& Collection,
+		const TMap<FString, int32>& InMeshNamePartMap, const TMap<FString, TArray<TObjectPtr<UMaterialInterface>>>& InMeshMaterialsMap)
+	{
+		Facades::FFoliageFacade FoliageFacade(Collection);
+		int32 NumInstances = FoliageFacade.NumFoliageEntries();
+			
+		for(int32 Id = 0; Id < NumInstances; Id++)
+		{
+			FString FoliageName = GetFoliageName(FoliageFacade, Id);
+
+			if(InMeshNamePartMap.Find(FoliageName))
+			{
+				int32 PartIndex = InMeshNamePartMap[FoliageName];
+				
+				FTransform Transform = FoliageFacade.GetFoliageTransform(Id);
+				UE::Geometry::FTransformSRT3d GeoTransform(Transform);
+
+				int32 FoliageParentBoneID = FoliageFacade.GetParentBoneID(Id);
+				
+				FNaniteAssemblyBoneInfluence Influence;
+				Influence.BoneIndex = FoliageParentBoneID;
+				//UE_LOG(LogProceduralVegetation, Warning, TEXT("Adding bone influence FoliageParentBoneID {%i}, Position{%s}, Scale{%s} , Rotation{%s}"),
+				//	FoliageParentBoneID, *Transform.GetLocation().ToString(), *Transform.GetScale3D().ToString(), *Transform.GetRotation().ToString());
+				Influence.BoneWeight = 1.0;
+				
+				AssemblyBuilder.AddNode(PartIndex, FTransform3f(GeoTransform), ENaniteAssemblyNodeTransformSpace::Local, { Influence });
+
+				//Add the materials for part to the assembly builder
+				int32 LocalMaterialIndex = 0;
+				auto FoliageMaterials = InMeshMaterialsMap.Find(FoliageName);
+				for (auto PartMaterial : *FoliageMaterials)
+				{
+					int32 PartMaterialIndex = AssemblyBuilder.GetMaterialSlots().IndexOfByPredicate(
+						[&] (const auto& Slot) { return Slot.Material == PartMaterial; }
+					);
+					if (PartMaterialIndex == INDEX_NONE)
+					{
+						PartMaterialIndex = AssemblyBuilder.AddMaterialSlot(PartMaterial);
+					}
+
+					AssemblyBuilder.RemapPartMaterial(PartIndex, LocalMaterialIndex, PartMaterialIndex);
+					LocalMaterialIndex++;
+				}
+			}
+		}
+	}
+
+	void ConvertToDefaultSkeletalMesh(USkeletalMesh* SkeletalMesh, FDynamicMesh3* Mesh,TArray<TObjectPtr<UMaterialInterface>> Materials)
+	{
+		TObjectPtr<UDynamicMesh> DynamicMesh = NewObject<UDynamicMesh>();
+		DynamicMesh->SetMesh(MoveTemp(*Mesh));
+		
+		//Figure out the correct parent bone index
+		int ParentBoneIndex = INDEX_NONE;
+		UE::Geometry::FDynamicMeshAttributeSet* MeshAttributes = DynamicMesh->GetMeshPtr()->Attributes();
+		MeshAttributes->EnableBones(1);
+
+		UE::Geometry::FDynamicMeshBoneNameAttribute* BoneNames = MeshAttributes->GetBoneNames();
+		UE::Geometry::FDynamicMeshBoneParentIndexAttribute* BoneParentIndices = MeshAttributes->GetBoneParentIndices();
+		UE::Geometry::FDynamicMeshBonePoseAttribute* BonePoses = MeshAttributes->GetBonePoses();
+
+		int BoneIndex = 0;
+		BoneNames->SetValue(BoneIndex, "Root");
+		BonePoses->SetValue(BoneIndex, FTransform::Identity);
+		BoneParentIndices->SetValue(BoneIndex, ParentBoneIndex);
+		
+		//Convert the dynamic mesh to static mesh
+		EGeometryScriptOutcomePins Output;
+#if WITH_EDITORONLY_DATA
+		FGeometryScriptCopyMeshToAssetOptions Options = GetCopyMeshToAssetOptions(Materials, SkeletalMesh->GetNaniteSettings(), ENaniteShapePreservation::Voxelize);
+#else
+		FGeometryScriptCopyMeshToAssetOptions Options = GetCopyMeshToAssetOptions(PartMaterials);
+#endif
+		Options.BoneHierarchyMismatchHandling = EGeometryScriptBoneHierarchyMismatchHandling::CreateNewReferenceSkeleton;
+		
+		UGeometryScriptLibrary_StaticMeshFunctions::CopyMeshToSkeletalMesh(
+			DynamicMesh,
+			SkeletalMesh,
+			Options,
+			{},
+			Output
+		);
+
+		TObjectPtr<USkeleton> Skeleton = SkeletalMesh->GetSkeleton();
+		Skeleton->RecreateBoneTree(SkeletalMesh);
+
+		const FGeometryScriptBoneWeightProfile SkinProfile = FGeometryScriptBoneWeightProfile();
+		FGeometryScriptSmoothBoneWeightsOptions SmoothBoneWeightsOptions;
+		SmoothBoneWeightsOptions.DistanceWeighingType = EGeometryScriptSmoothBoneWeightsType::GeodesicVoxel;
+		SmoothBoneWeightsOptions.MaxInfluences = 2;
+		SmoothBoneWeightsOptions.VoxelResolution = 512;
+		UGeometryScriptLibrary_MeshBoneWeightFunctions::ComputeSmoothBoneWeights(
+			DynamicMesh,
+			Skeleton,
+			SmoothBoneWeightsOptions,
+			SkinProfile
+		);
+
+		const FName ProfileName = SkinProfile.GetProfileName();
+		UGeometryScriptLibrary_StaticMeshFunctions::CopySkinWeightProfileToSkeletalMesh(
+			DynamicMesh,
+			SkeletalMesh,
+			ProfileName,
+			ProfileName,
+			{},
+			{},
+			Output
+		);
+	}
+#endif
+
+	void ExportCollectionToStaticMesh(
+		TObjectPtr<UStaticMesh> ExportMesh,
+		const FManagedArrayCollection& Collection,
+		bool bBuildNaniteAssemblies,
+		ENaniteShapePreservation ShapePreservation
+	)
+	{
+		FManagedArrayCollection ExportedCollection;
+		Collection.CopyTo(&ExportedCollection);
+		
+		TObjectPtr<UDynamicMesh> DynamicMesh = CollectionToDynamicMesh(ExportedCollection);
+
+		//Get the materials from the collection
+		TArray<TObjectPtr<UMaterialInterface>> Materials = GetMaterialsFromCollection(ExportedCollection);
+
+		//Get the foliage data from facade, spawn an actor attach the foliage to the actor
+		Facades::FFoliageFacade FoliageFacade(ExportedCollection);
+		
+		//Combine all the foliage instances into one dynamic mesh if not building nanite assemblies
+		if (!bBuildNaniteAssemblies)
+		{
+			CombineMeshInstancesToDynamicMesh(ExportedCollection, DynamicMesh, Materials,
+				[](FString MeshName, int32 InstanceID, int32 VertexStart, int32 VertexCount){});
+		}
+
+		//Convert the dynamic mesh to static mesh
+		EGeometryScriptOutcomePins Output;
+#if WITH_EDITORONLY_DATA
+		FGeometryScriptCopyMeshToAssetOptions Options = GetCopyMeshToAssetOptions(Materials, ExportMesh->GetNaniteSettings(), ShapePreservation);
+		
+#else
+		FGeometryScriptCopyMeshToAssetOptions Options = GetCopyMeshToAssetOptions(Materials);
+#endif
+		UGeometryScriptLibrary_StaticMeshFunctions::CopyMeshToStaticMesh(DynamicMesh, ExportMesh, Options, {}, Output);
+
+#if WITH_EDITOR
+		if (bBuildNaniteAssemblies)
+		{
+			if (NaniteAssembliesSupported())
+			{
+				//Build the nanite assemblies
+            	BuildNaniteAssemblyData(ExportedCollection, ExportMesh);
+            	ExportMesh->Build(true);
+			}
+			else
+			{
+				UE_LOG(LogProceduralVegetation, Warning, TEXT("Failed to build Nanite Assemblies, because neither Nanite Foliage nor Nanite Assemblies are enabled for the project."));
+			}
+		}
+#endif
+	}
+
+#if WITH_EDITOR
+	void BuildWindImportData(FDynamicWindSkeletalImportData& OutImportData, const FManagedArrayCollection& Collection, const TArray<PV::Facades::FBoneNode>& Bones)
+	{
+		Facades::FBranchFacade BranchFacade(Collection);
+
+		for (const auto& Bone : Bones)
+		{
+			int32 SimulationGroupIndex = BranchFacade.GetBranchSimulationGroupIndex(Bone.BranchIndex);
+			OutImportData.Joints.Add({Bone.BoneName, SimulationGroupIndex});
+		}
+	}
+#endif
+
+	void ExportCollectionToSkeletalMesh(
+		TObjectPtr<USkeletalMesh> ExportMesh,
+		const FManagedArrayCollection& Collection,
+		ENaniteShapePreservation ShapePreservation,
+		bool bBuildNaniteAssemblies
+	)
+	{
+		FManagedArrayCollection ExportedCollection;
+		Collection.CopyTo(&ExportedCollection);
+
+		Facades::FFoliageFacade FoliageFacade(ExportedCollection);
+		Facades::FPointFacade PointFacade(ExportedCollection);
+		Facades::FBranchFacade BranchFacade(ExportedCollection);
+		Facades::FBoneFacade BoneFacade = Facades::FBoneFacade(ExportedCollection);
+		
+		TObjectPtr<UDynamicMesh> DynamicMesh = CollectionToDynamicMesh(ExportedCollection);
+		UE::Geometry::FDynamicMeshAttributeSet* MeshAttributes = DynamicMesh->GetMeshPtr()->Attributes();
+		
+		//Get the materials from the collection
+		TArray<TObjectPtr<UMaterialInterface>> Materials = GetMaterialsFromCollection(ExportedCollection);
+
+		TArray<Facades::FBoneNode> BoneNodes;
+		BoneNodes = BoneFacade.GetBoneDataFromCollection();
+
+		//if bones already not created with bone reduction node create bones with full density
+		if (BoneNodes.IsEmpty())
+		{
+			BoneFacade.CreateBoneData(BoneNodes, 0);
+		}
+		
+		AssignBoneIDsToFoliage(BoneNodes, ExportedCollection);
+		
+		int32 BoneCount = BoneNodes.Num();
+		MeshAttributes->EnableBones(BoneCount);
+
+		TArray<int32> VertexPointIDs;
+		BoneFacade.GetPointIds(VertexPointIDs);
+
+		TArray<int32> VertexBoneIDs;
+		//Assign the Bone id to every vertex in the base mesh through VertexPointIDs
+		for (const int32& PointID : VertexPointIDs)
+		{
+			if (Facades::FBoneNode* BoneNode = Facades::FBoneFacade::FindClosestBone(ExportedCollection, BoneNodes, PointID))
+			{
+				VertexBoneIDs.Add(BoneNode->BoneIndex);
+			}
+			else
+			{
+				VertexBoneIDs.Add(INDEX_NONE);
+				UE_LOG(LogProceduralVegetation, Log, TEXT("Invalid Bone Assigned to PointID %i") , PointID);
+			}
+		}
+		int FoliageVertexStart = VertexBoneIDs.Num();
+
+		UE::Geometry::FDynamicMeshBoneNameAttribute* BoneNames = MeshAttributes->GetBoneNames();
+		UE::Geometry::FDynamicMeshBoneParentIndexAttribute* BoneParentIndices = MeshAttributes->GetBoneParentIndices();
+		UE::Geometry::FDynamicMeshBonePoseAttribute* BonePoses = MeshAttributes->GetBonePoses();
+
+		for (int i = 0; i < BoneCount; i++)
+		{
+			BoneNames->SetValue(i, BoneNodes[i].BoneName);
+			BonePoses->SetValue(i, BoneNodes[i].BoneTransform);
+			BoneParentIndices->SetValue(i, BoneNodes[i].ParentBoneIndex);
+		}
+
+		if (!bBuildNaniteAssemblies)
+		{
+			//Combine all the foliage instances into one dynamic mesh if not building nanite assemblies
+			CombineMeshInstancesToDynamicMesh(ExportedCollection, DynamicMesh, Materials,
+			[&](const FString& MeshName, int32 InstanceID, int32 VertexStart, int32 VertexCount)
+			{
+				//Assign parent bone id to foliage vertices
+				int32 FoliageParentBoneID = FoliageFacade.GetParentBoneID(InstanceID);
+					
+				for (int VertexID = VertexStart; VertexID < VertexCount; ++VertexID)
+				{
+					VertexBoneIDs.Add(FoliageParentBoneID);
+					//UE_LOG(LogProceduralVegetation, Warning, TEXT("Adding VertexBoneIDs, VertexId{%i} BoneID{%i}, VertexStart{%i}, VertexCount{%i}, MeshName{%s}, InstanceID{%i}"),
+					//VertexID, FoliageParentBoneID, VertexStart, VertexCount, *MeshName, InstanceID);
+				}
+			});
+		}
+
+		//Convert the dynamic mesh to static mesh
+		EGeometryScriptOutcomePins Output;
+		
+#if WITH_EDITORONLY_DATA
+		FGeometryScriptCopyMeshToAssetOptions Options = GetCopyMeshToAssetOptions(Materials, ExportMesh->GetNaniteSettings(), ShapePreservation);
+		
+#else
+		FGeometryScriptCopyMeshToAssetOptions Options = GetCopyMeshToAssetOptions(Materials);
+#endif
+		Options.BoneHierarchyMismatchHandling = EGeometryScriptBoneHierarchyMismatchHandling::CreateNewReferenceSkeleton;
+
+		UGeometryScriptLibrary_StaticMeshFunctions::CopyMeshToSkeletalMesh(
+			DynamicMesh,
+			ExportMesh,
+			Options,
+			{},
+			Output
+		);
+
+		TObjectPtr<USkeleton> Skeleton = ExportMesh->GetSkeleton();
+		Skeleton->RecreateBoneTree(ExportMesh);
+
+		const FGeometryScriptBoneWeightProfile SkinProfile = FGeometryScriptBoneWeightProfile();
+		FGeometryScriptSmoothBoneWeightsOptions SmoothBoneWeightsOptions;
+
+		SmoothBoneWeightsOptions.DistanceWeighingType = EGeometryScriptSmoothBoneWeightsType::DirectDistance;
+		SmoothBoneWeightsOptions.MaxInfluences = 3;
+		SmoothBoneWeightsOptions.Stiffness = 0.5;
+		
+		UGeometryScriptLibrary_MeshBoneWeightFunctions::ComputeSmoothBoneWeights(
+			DynamicMesh,
+			Skeleton,
+			SmoothBoneWeightsOptions,
+			SkinProfile
+		);
+
+		const FName ProfileName = SkinProfile.GetProfileName();
+		
+		//Here we make the correction to the weights
+		RemoveUnwantedSkinWeights(DynamicMesh, ProfileName, VertexBoneIDs, FoliageVertexStart,BoneNodes);
+		
+		UGeometryScriptLibrary_StaticMeshFunctions::CopySkinWeightProfileToSkeletalMesh(
+			DynamicMesh,
+			ExportMesh,
+			ProfileName,
+			ProfileName,
+			{},
+			{},
+			Output
+		);
+
+#if WITH_EDITOR
+		if (bBuildNaniteAssemblies)
+		{
+			if (NaniteAssembliesSupported())
+			{
+				//Build the nanite assemblies
+				BuildNaniteAssemblyData(ExportedCollection, ExportMesh);
+
+				FDynamicWindSkeletalImportData ImportData;
+				BuildWindImportData(ImportData,ExportedCollection, BoneNodes);
+				
+				if (!DynamicWind::ImportSkeletalData(*ExportMesh, ImportData))
+				{
+					UE_LOG(LogTemp, Warning, TEXT("Failed to build wind data"));
+				}
+				
+				ExportMesh->Build();
+			}
+			else
+			{
+				UE_LOG(LogProceduralVegetation, Warning, TEXT("Failed to build Nanite Assemblies, because neither Nanite Foliage nor Nanite Assemblies are enabled for the project."));
+			}
+		}
+#endif
+	}
+
+	void AttachProceduralVegetationLink(UObject* InExportedMesh, const TObjectPtr<UProceduralVegetation>& InProceduralVegetation)
+	{
+		if (IInterface_AssetUserData* IAssetUserData = Cast<IInterface_AssetUserData>(InExportedMesh))
+		{
+			UProceduralVegetationLink* Data = NewObject<UProceduralVegetationLink>(InExportedMesh);
+			Data->Source = InProceduralVegetation;
+		
+			IAssetUserData->AddAssetUserData(Data);	
+		}
+	}
+
+	void AssignBoneIDsToFoliage(const TArray<Facades::FBoneNode>& Bones, FManagedArrayCollection& Collection)
+	{
+		Facades::FPointFacade PointFacade(Collection);
+		Facades::FFoliageFacade FoliageFacade(Collection);
+		
+		int32 NumInstances = FoliageFacade.NumFoliageEntries();
+
+		auto FindBoneById = ( [&](const int32 Id)
+		{
+			const Facades::FBoneNode* BoneNode = Bones.FindByPredicate([Id](const Facades::FBoneNode& Node)
+			{
+				return Node.BoneIndex == Id;
+			});
+			return BoneNode;
+		});
+		
+		for(int32 Id = 0; Id < NumInstances; Id++)
+		{
+			Facades::FFoliageEntryData Data = FoliageFacade.GetFoliageEntry(Id);
+			float FoliageLFR = Data.LengthFromRoot;
+
+			TArray<int32> BoneIDs;
+			for (Facades::FBoneNode BoneNode : Bones)
+			{
+				if (Data.BranchId == BoneNode.BranchIndex)
+				{
+					BoneIDs.Add(BoneNode.BoneIndex);
+				}
+			}
+
+			bool bBoneIDAssigned = false;
+			for (const int32& BoneId : BoneIDs)
+			{
+				const Facades::FBoneNode* BoneNode = FindBoneById(BoneId);
+				check(BoneNode);
+				
+				int PointIndex = BoneNode->PointIndex;
+				const Facades::FBoneNode* BoneParentNode = FindBoneById(BoneNode->ParentBoneIndex);
+				
+				int ParentPointIndex = BoneParentNode ? BoneParentNode->PointIndex : INDEX_NONE;
+				float BoneLFR = PointFacade.GetLengthFromRoot(PointIndex);
+				float BoneParentLFR = PointFacade.GetLengthFromRoot(ParentPointIndex);
+
+				if (FoliageLFR <= BoneLFR && FoliageLFR >= BoneParentLFR)
+				{
+					FoliageFacade.SetParentBoneID(Id, BoneId);
+					bBoneIDAssigned = true;
+					break;
+				}
+			}
+
+			if (!bBoneIDAssigned)
+			{
+				UE_LOG(LogProceduralVegetation, Warning, TEXT("No bone assigned for foliage Id {%i} , Branch Bone count {%i}"), Id , BoneIDs.Num());
+			}
+		}
+	}
+
+	void RemoveUnwantedSkinWeights(TObjectPtr<UDynamicMesh> DynamicMesh, const FName ProfileName, const TArray<int32>& VertexBoneIDs, int32 FoliageVertexStart, const TArray<Facades::FBoneNode>& BoneNodes)
+	{
+		auto InMesh = DynamicMesh->GetMeshPtr();
+		check(InMesh)
+		const int32 NumVertices = VertexBoneIDs.Num();
+		
+		DynamicMesh->EditMesh([&](FDynamicMesh3& EditMesh)
+		{
+			UE::Geometry::FDynamicMeshVertexSkinWeightsAttribute *SkinWeights = EditMesh.Attributes()->GetSkinWeightsAttribute(ProfileName);
+
+			ParallelFor(NumVertices, [&](const int32 VertexIdx)
+			{
+				int32 BoneID = VertexBoneIDs[VertexIdx];
+				//UE_LOG(LogProceduralVegetation, Warning, TEXT("Remove Bones, VertexId{%i} BoneID{%i} BoneParentIndex{%i}"), VertexIdx, BoneID, ParentBoneID);
+
+				if(BoneID != INDEX_NONE)
+				{
+					int32 ParentBoneID = BoneNodes[BoneID].ParentBoneIndex;
+				
+					//UE_LOG(LogProceduralVegetation, Warning, TEXT("Remove Bones, VertexId{%i} BoneID{%i} BoneParentIndex{%i}"), VertexIdx, BoneID, ParentBoneID);
+
+					UE::AnimationCore::FBoneWeights Weights;
+					SkinWeights->GetValue(VertexIdx, Weights);
+
+					auto NextBone = BoneNodes.FindByPredicate([&](const Facades::FBoneNode& Node)
+					{
+						return Node.BoneIndex == BoneID + 1 && Node.BranchIndex == BoneNodes[BoneID].BranchIndex;
+					});
+
+					for (int32 WeightID = Weights.Num() - 1 ; WeightID >= 0; WeightID--)
+					{
+						auto Weight = Weights[WeightID];
+
+						bool bValidBone = (Weight.GetBoneIndex() == ParentBoneID);
+						bValidBone |= (Weight.GetBoneIndex() == BoneID);
+						bValidBone |= NextBone && (Weight.GetBoneIndex() == NextBone->BoneIndex);
+						
+						if (!bValidBone || VertexIdx > FoliageVertexStart || BoneID == 0)
+						{
+							UE::AnimationCore::FBoneWeightsSettings Settings;
+							Settings.SetNormalizeType(UE::AnimationCore::EBoneWeightNormalizeType::Always);
+							Weights.RemoveBoneWeight(Weight.GetBoneIndex(),Settings);
+						}
+					}
+
+					if (VertexIdx > FoliageVertexStart || BoneID == 0 || Weights.Num() == 0)
+					{
+						UE::AnimationCore::FBoneWeightsSettings Settings;
+						Settings.SetNormalizeType(UE::AnimationCore::EBoneWeightNormalizeType::Always);
+						Weights.SetBoneWeight(UE::AnimationCore::FBoneWeight(BoneID, 1.0f), Settings);
+					}
+					
+					SkinWeights->SetNewValue(VertexIdx, Weights);
+				}
+				else
+				{
+					UE_LOG(LogProceduralVegetation, Log, TEXT("Invalid bone index found while removing bones for vertex %i"), VertexIdx);
+				}
+			});
+				
+		}, EDynamicMeshChangeType::GeneralEdit, EDynamicMeshAttributeChangeFlags::Unknown, false);
+	}
+}
+
+namespace PV::Export
+{
+	void ExportCollectionAsMesh(
+		const TObjectPtr<UProceduralVegetation> InProceduralVegetation,
+		const FManagedArrayCollection& Collection,
+		const FManagedArrayCollection& FoliageCollection,
+		const FPVExportSettings& Settings
+	)
+	{
+		// Set up the Package path and name from settings
+		const FString PackageName = Settings.ContentBrowserFolder.Path;
+		const FString MeshName = Settings.MeshName.ToString();
+
+		const FString FullPath = FPaths::Combine(PackageName, MeshName);
+
+		if (Settings.ExportMeshType == EPVExportMeshType::StaticMesh)
+		{
+			UPackage* MeshPackage = CreatePackage(*FullPath);
+
+			// Create a static mesh to save in content browser
+			const TObjectPtr<UStaticMesh> ExportMesh = NewObject<UStaticMesh>(MeshPackage, FName(MeshName), RF_Standalone | RF_Public);
+
+			Internal::AttachProceduralVegetationLink(ExportMesh, InProceduralVegetation);
+			
+			PV::Export::Internal::ExportCollectionToStaticMesh(
+				ExportMesh,
+				Collection,
+				Settings.bCreateNaniteAssemblies,
+				Settings.NaniteShapePreservation
+			);
+
+			MeshPackage->SetDirtyFlag(true);
+		}
+		else if (Settings.ExportMeshType == EPVExportMeshType::SkeletalMesh)
+		{
+			UPackage* MeshPackage = CreatePackage(*FullPath);
+			UPackage* SkeletonPackage = CreatePackage(*(FullPath + "_Skeleton"));
+
+			// Create a skeletal mesh to save in content browser
+			const TObjectPtr<USkeletalMesh> ExportMesh = NewObject<USkeletalMesh>(MeshPackage, FName(MeshName), RF_Standalone | RF_Public);
+			const TObjectPtr<USkeleton> ExportSkeleton = NewObject<USkeleton>(SkeletonPackage, FName(MeshName + "_Skeleton"),
+			                                                                  RF_Standalone | RF_Public);
+			Internal::AttachProceduralVegetationLink(ExportMesh, InProceduralVegetation);
+			Internal::AttachProceduralVegetationLink(ExportSkeleton, InProceduralVegetation);
+			
+			ExportMesh->SetSkeleton(ExportSkeleton);
+
+			PV::Export::Internal::ExportCollectionToSkeletalMesh(
+				ExportMesh,
+				Collection,
+				Settings.NaniteShapePreservation,
+				Settings.bCreateNaniteAssemblies
+			);
+
+			MeshPackage->SetDirtyFlag(true);
+			SkeletonPackage->SetDirtyFlag(true);
+
+			IAssetRegistry::Get()->AssetCreated(ExportMesh);
+			IAssetRegistry::Get()->AssetCreated(ExportSkeleton);
+		}
+	}
+}
diff --git a/Engine/Plugins/Experimental/ProceduralVegetationEditor/Source/ProceduralVegetation/Private/Helpers/PVUtilities.cpp b/Engine/Plugins/Experimental/ProceduralVegetationEditor/Source/ProceduralVegetation/Private/Helpers/PVUtilities.cpp
new file mode 100644
--- /dev/null
+++ b/Engine/Plugins/Experimental/ProceduralVegetationEditor/Source/ProceduralVegetation/Private/Helpers/PVUtilities.cpp
@@ -0,0 +1,108 @@
+﻿// Copyright Epic Games, Inc. All Rights Reserved.
+
+#include "Helpers/PVUtilities.h"
+
+#include "ProceduralVegetationModule.h"
+#include "Components/InstancedSkinnedMeshComponent.h"
+#include "Engine/StaticMesh.h"
+#include "Engine/SkeletalMesh.h"
+#include "UObject/Package.h"
+#include "Components/InstancedStaticMeshComponent.h"
+#include "Components/MeshComponent.h"
+
+namespace PV::Utilities
+{
+	void AddFoliageInstancesToActor(const Facades::FFoliageFacade& FoliageFacade, AActor* Actor,
+	                                TMap<FString, TObjectPtr<UMeshComponent>>& OutInstancedComponentMap)
+	{
+		TRACE_CPUPROFILER_EVENT_SCOPE(PV::Utilities::AddFoliageInstancesToActor());
+
+		auto OnComponentCreated = [&Actor](UMeshComponent* InInstancedComponent)
+			{
+				InInstancedComponent->RegisterComponentWithWorld(Actor->GetWorld());
+				InInstancedComponent->AttachToComponent(Actor->GetRootComponent(), FAttachmentTransformRules::KeepRelativeTransform);
+				Actor->AddInstanceComponent(InInstancedComponent);
+			};
+		
+		AddFoliageInstances(FoliageFacade, Actor, FFoliageComponentCreatedCallback::CreateLambda(OnComponentCreated),OutInstancedComponentMap);
+	}
+
+	void AddFoliageInstances(const Facades::FFoliageFacade& InFoliageFacade, const UObject* InParent, FFoliageComponentCreatedCallback InCallback, TMap<FString, TObjectPtr<UMeshComponent>>& OutInstancedComponentMap)
+	{
+		TRACE_CPUPROFILER_EVENT_SCOPE(PV::Utilities::AddFoliageInstances());
+
+		int32 NumInstances = InFoliageFacade.NumFoliageEntries();
+
+		OutInstancedComponentMap.Empty();
+
+		if (NumInstances > 0)
+		{
+			for (int32 FoliageNameIndex = 0; FoliageNameIndex < InFoliageFacade.NumFoliageNames(); ++FoliageNameIndex)
+			{
+				FString FoliageName = InFoliageFacade.GetFoliageName(FoliageNameIndex);
+
+				UStaticMesh* StaticFoliageMesh = LoadObject<UStaticMesh>(nullptr, *FoliageName);
+				USkeletalMesh* SkeletalFoliageMesh = LoadObject<USkeletalMesh>(nullptr, *FoliageName);
+
+				if(StaticFoliageMesh)
+				{
+					UInstancedStaticMeshComponent* InstancedComponent = NewObject<UInstancedStaticMeshComponent>(InParent != nullptr ? InParent->GetPackage() : GetTransientPackage(), NAME_None);
+					InstancedComponent->SetStaticMesh(StaticFoliageMesh);
+					OutInstancedComponentMap.Add(FoliageName, InstancedComponent);
+					InCallback.ExecuteIfBound(InstancedComponent);
+				}
+				else if(SkeletalFoliageMesh)
+				{
+#if WITH_EDITORONLY_DATA
+					if (!SkeletalFoliageMesh->NaniteSettings.bEnabled)
+					{
+						SkeletalFoliageMesh->NaniteSettings.bEnabled = true;
+						SkeletalFoliageMesh->Build();
+						SkeletalFoliageMesh->MarkPackageDirty();
+						UE_LOG(LogProceduralVegetation, Warning, TEXT("Enabling Nanite for the Skeletal Mesh {%s}, its required for InstancedSkinnedMeshComponent"), *FoliageName);
+					}
+#else
+					UE_LOG(LogProceduralVegetation, Warning, TEXT("Instanced Skinned Mesh component only works with nanite meshes make sure nanite is enabled for mesh {%s}"), *FoliageName);
+#endif
+					
+					UInstancedSkinnedMeshComponent* InstancedComponent = NewObject<UInstancedSkinnedMeshComponent>(InParent != nullptr ? InParent->GetPackage() : GetTransientPackage(), NAME_None);
+					InstancedComponent->SetSkinnedAssetAndUpdate(SkeletalFoliageMesh);
+					OutInstancedComponentMap.Add(FoliageName, InstancedComponent);
+					InCallback.ExecuteIfBound(InstancedComponent);
+				}
+			}
+
+			TMap<FString, TArray<FTransform>> InstanceTransformsMap;
+
+			for (int32 Id = 0; Id < NumInstances; Id++)
+			{
+				Facades::FFoliageEntryData Data = InFoliageFacade.GetFoliageEntry(Id);
+				FString FoliageName = InFoliageFacade.GetFoliageName(Data.NameId);
+
+				const FTransform Transform = InFoliageFacade.GetFoliageTransform(Id);
+
+				InstanceTransformsMap.FindOrAdd(FoliageName).Add(Transform);
+			}
+		
+			for(auto FoliageEntry : InstanceTransformsMap)
+			{
+				if (FoliageEntry.Value.Num() > 0)
+				{
+					if (auto InstancedComponent = OutInstancedComponentMap.Find(FoliageEntry.Key))
+					{
+						if(TObjectPtr<UInstancedStaticMeshComponent> StaticMeshInstancedComponent = Cast<UInstancedStaticMeshComponent>(*InstancedComponent))
+						{
+							StaticMeshInstancedComponent->AddInstances(FoliageEntry.Value, false);	
+						}
+						else if(TObjectPtr<UInstancedSkinnedMeshComponent> SkinnedInstancedComponent = Cast<UInstancedSkinnedMeshComponent>(*InstancedComponent))
+						{
+							TArray<int32> BankIndices;
+							BankIndices.AddUninitialized(FoliageEntry.Value.Num());
+							SkinnedInstancedComponent->AddInstances(FoliageEntry.Value, BankIndices, false);
+						}
+					}
+				}
+			}
+		}
+	}
+}
