Rename MegaPlant to ProceduralVegetation

diff --git a/Engine/Plugins/Experimental/ProceduralVegetationEditor/Source/ProceduralVegetation/Private/Nodes/PVPCGMesher.cpp b/Engine/Plugins/Experimental/ProceduralVegetationEditor/Source/ProceduralVegetation/Private/Nodes/PVPCGMesher.cpp
new file mode 100644
--- /dev/null
+++ b/Engine/Plugins/Experimental/ProceduralVegetationEditor/Source/ProceduralVegetation/Private/Nodes/PVPCGMesher.cpp
@@ -0,0 +1,135 @@
+﻿// Copyright Epic Games, Inc. All Rights Reserved.
+
+#include "PVPCGMesher.h"
+#include "PVPCGData.h"
+#include "Data/PCGBasePointData.h"
+#include "Facades/PVBranchFacade.h"
+#include "Facades/PVProfileFacade.h"
+#include "GeometryCollection/ManagedArrayCollection.h"
+#include "Helpers/PCGSettingsHelpers.h"
+#include "Implementations/PVMesher.h"
+#include "Materials/MaterialInterface.h"
+
+#if WITH_EDITOR
+void UPVPCGMesherSettings::PostEditChangeProperty(FPropertyChangedEvent& PropertyChangedEvent)
+{
+	Super::PostEditChangeProperty(PropertyChangedEvent);
+
+	if (PropertyChangedEvent.Property
+		&& PropertyChangedEvent.Property->GetFName() == GET_MEMBER_NAME_CHECKED(FPVMesherSettings, DisplacementTexture))
+	{
+		if (MesherSettings.DisplacementTexture)
+		{
+			FPVMesher::ExtractDisplacementData(MesherSettings.DisplacementTexture, MesherSettings.DisplacementValues);
+		}
+		else
+		{
+			MesherSettings.DisplacementValues.Empty();
+		}
+
+		Modify();
+	}
+}
+
+void UPVPCGMesherSettings::PostLoad()
+{
+	Super::PostLoad();
+
+	if (MesherSettings.DisplacementTexture)
+	{
+		FPVMesher::ExtractDisplacementData(MesherSettings.DisplacementTexture, MesherSettings.DisplacementValues);
+	}
+	else
+	{
+		MesherSettings.DisplacementValues.Empty();
+	}
+}
+#endif
+
+FPCGElementPtr UPVPCGMesherSettings::CreateElement() const
+{
+	return MakeShared<FPVPCGMesherElement>();
+}
+
+bool FPVPCGMesherElement::ExecuteInternal(FPCGContext* InContext) const
+{
+	TRACE_CPUPROFILER_EVENT_SCOPE(FPVPCGMesherElement::Execute);
+
+	check(InContext);
+
+	const UPVPCGMesherSettings* Settings = InContext->GetInputSettings<UPVPCGMesherSettings>();
+	check(Settings);
+
+	UPVPCGMesherSettings* MutableSettings = const_cast<UPVPCGMesherSettings*>(Settings);
+	check(MutableSettings);
+
+	MutableSettings->MesherSettings.PlantProfileOptions.Empty();
+	MutableSettings->MesherSettings.PlantProfileOptions.Add(TEXT("None"));
+
+	const TArray<FPCGTaggedData>& Inputs = InContext->InputData.GetInputsByPin(PCGPinConstants::DefaultInputLabel);
+	if (!Inputs.IsEmpty())
+	{
+		if (const UPVPCGData* InputData = Cast<UPVPCGData>(Inputs[0].Data))
+		{
+			const FManagedArrayCollection& SkeletonCollection = InputData->GetCollection();
+			auto& MaterialSettings = MutableSettings->MesherSettings.MaterialSettings;
+					
+			if (!Settings->MesherSettings.bOverrideMaterial || MaterialSettings.MaterialSetups.Num() == 0)
+			{
+				const PV::Facades::FBranchFacade BranchFacade(SkeletonCollection);
+				const FString TrunkMaterial = BranchFacade.GetTrunkMaterialPath();
+				if (!TrunkMaterial.IsEmpty())
+				{
+					MaterialSettings.SetMaterial(TrunkMaterial, 0);
+				}
+				else
+				{
+					MutableSettings->MesherSettings.bOverrideMaterial = true;
+				}
+			}
+
+			if (const PV::Facades::FPlantProfileFacade PlantProfileFacade = PV::Facades::FPlantProfileFacade(SkeletonCollection);
+				PlantProfileFacade.NumProfileEntries() > 0)
+			{
+				for (int32 i = 0; i < PlantProfileFacade.NumProfileEntries(); i++)
+				{
+					MutableSettings->MesherSettings.PlantProfileOptions.Add(FString::FromInt(i));
+				}
+
+				MutableSettings->MesherSettings.bIsPlantProfileDropdownEnabled = true;
+			}
+			else
+			{
+				MutableSettings->MesherSettings.bIsPlantProfileDropdownEnabled = false;
+			}
+
+			UPVPCGData* OutManagedArrayCollectionData = FPCGContext::NewObject_AnyThread<UPVPCGData>(InContext);
+
+			FGeometryCollection OutGeometryCollection;
+			FPVMesher::GenerateGeometryCollection(SkeletonCollection, Settings->MesherSettings, OutGeometryCollection);
+
+			OutManagedArrayCollectionData->Initialize(MoveTemp(OutGeometryCollection));
+
+			FPCGTaggedData& CollectionOutput = InContext->OutputData.TaggedData.Emplace_GetRef();
+			CollectionOutput.Data = OutManagedArrayCollectionData;
+			CollectionOutput.Pin = PCGPinConstants::DefaultOutputLabel;
+
+			// TODO : Figure out if we want to keep the dynamic mesh inn the ProceduralVegetation data or as a pin
+			// UPCGDynamicMeshData* DynamicMeshData = FPCGContext::NewObject_AnyThread<UPCGDynamicMeshData>(InContext);
+			// TObjectPtr<UDynamicMesh> OutMesh = NewObject<UDynamicMesh>();
+			// FPVMesher::GenerateDynamicMesh(SkeletonCollection, OutMesh);
+			// DynamicMeshData->Initialize(OutMesh, true);
+			//
+			// FPCGTaggedData& MeshOutput = InContext->OutputData.TaggedData.Emplace_GetRef();
+			// MeshOutput.Data = DynamicMeshData;
+			// MeshOutput.Pin = TEXT("Mesh");
+		}
+		else
+		{
+			PCGLog::InputOutput::LogInvalidInputDataError(InContext);
+			return true;
+		}
+	}
+
+	return true;
+}
diff --git a/Engine/Plugins/Experimental/ProceduralVegetationEditor/Source/ProceduralVegetation/Private/Nodes/PVPCGMesher.h b/Engine/Plugins/Experimental/ProceduralVegetationEditor/Source/ProceduralVegetation/Private/Nodes/PVPCGMesher.h
new file mode 100644
--- /dev/null
+++ b/Engine/Plugins/Experimental/ProceduralVegetationEditor/Source/ProceduralVegetation/Private/Nodes/PVPCGMesher.h
@@ -0,0 +1,49 @@
+﻿// Copyright Epic Games, Inc. All Rights Reserved.
+#pragma once
+
+#include "CoreMinimal.h"
+#include "Implementations/PVMesher.h"
+#include "Nodes/PVPCGBaseSettings.h"
+#include "PVPCGData.h"
+#include "PVPCGMesher.generated.h"
+
+class UMaterialInterface;
+
+UCLASS(BlueprintType, ClassGroup = (Procedural))
+class UPVPCGMesherSettings : public UPVPCGBaseSettings
+{
+	GENERATED_BODY()
+
+public:
+#if WITH_EDITOR
+	virtual FName GetDefaultNodeName() const override { return FName(TEXT("ProceduralVegetationMesher")); }
+	virtual FText GetDefaultNodeTitle() const override { return NSLOCTEXT("PVPCGMesher", "NodeTitle", "Mesh Builder"); }
+	virtual FText GetNodeTooltipText() const override { return NSLOCTEXT("PVPCGMesher", "NodeTooltip", "Create a mesh for the skeleton"); }
+	virtual FLinearColor GetNodeTitleColor() const override { return FLinearColor(FColor::Cyan); }
+
+	virtual TArray<EPVRenderType> GetRenderType() const override { return TArray{EPVRenderType::Mesh}; }
+
+	virtual void PostEditChangeProperty(FPropertyChangedEvent& PropertyChangedEvent) override;
+
+	virtual void PostLoad() override;
+#endif
+
+protected:
+	virtual FPCGElementPtr CreateElement() const override;
+
+public:	
+	UPROPERTY(EditAnywhere, Category = "Mesh", meta = (ShowOnlyInnerProperties))
+	FPVMesherSettings MesherSettings;
+
+	UFUNCTION()
+	TArray<FString> GetPlantProfileOptions() const
+	{
+		return MesherSettings.PlantProfileOptions;
+	}
+};
+
+class FPVPCGMesherElement : public FPVPCGBaseElement
+{
+protected:
+	virtual bool ExecuteInternal(FPCGContext* InContext) const override;
+};
diff --git a/Engine/Plugins/Experimental/ProceduralVegetationEditor/Source/ProceduralVegetation/Private/Nodes/PVPCGOutput.cpp b/Engine/Plugins/Experimental/ProceduralVegetationEditor/Source/ProceduralVegetation/Private/Nodes/PVPCGOutput.cpp
new file mode 100644
--- /dev/null
+++ b/Engine/Plugins/Experimental/ProceduralVegetationEditor/Source/ProceduralVegetation/Private/Nodes/PVPCGOutput.cpp
@@ -0,0 +1,98 @@
+﻿// Copyright Epic Games, Inc. All Rights Reserved.
+
+#include "Nodes/PVPCGOutput.h"
+#include "PCGContext.h"
+#include "Helpers/PCGHelpers.h"
+#include "Misc/PackageName.h"
+
+#define LOCTEXT_NAMESPACE "PVPCGOutput"
+
+FString FormatLongPackageNameErrorCode(const FPackageName::EErrorCode ErrorCode)
+{
+	switch (ErrorCode)
+	{
+	case FPackageName::EErrorCode::PackageNameUnknown:
+		return "Unknown Export Path error";
+	case FPackageName::EErrorCode::PackageNameEmptyPath:
+		return "Empty Export Path";
+	case FPackageName::EErrorCode::PackageNamePathNotMounted:
+		return "Mount Point for Export Path is not valid";
+	case FPackageName::EErrorCode::PackageNamePathIsMemoryOnly:
+		return "Export Path exists in memory only. ";
+	case FPackageName::EErrorCode::PackageNameSpacesNotAllowed:
+		return "Spaces are not allowed in Export Path";	
+	case FPackageName::EErrorCode::PackageNameContainsInvalidCharacters:
+		return FString::Printf(TEXT("Export Path contains one of \n%s\n Invalid Characters"), TEXT(R"(\\:*?\"<>|' ,.&!~\n\r\t@#)"));
+	case FPackageName::EErrorCode::LongPackageNames_PathTooShort:
+		return "Export Path is too small to be a valid path";
+	case FPackageName::EErrorCode::LongPackageNames_PathWithNoStartingSlash:
+		return "Export Path has to start with a slash [/]";
+	case FPackageName::EErrorCode::LongPackageNames_PathWithTrailingSlash:
+		return "Export Path cannot end with a slash [/]";
+	case FPackageName::EErrorCode::LongPackageNames_PathWithDoubleSlash:
+		return "Export Path contains a double slash [/]";
+	}
+	return FString();
+}
+
+UPVPCGOutputSettings::UPVPCGOutputSettings()
+{
+	if (PCGHelpers::IsNewObjectAndNotDefault(this))
+	{
+		static FString DefaultMeshName = "ProceduralVegetationOutput";
+		static FString DefaultAssetPath = "/Game/ProceduralVegetation/";
+		ExportSettings.Initialize(DefaultAssetPath, DefaultMeshName);
+	}
+}
+
+TArray<FPCGPinProperties> UPVPCGOutputSettings::OutputPinProperties() const
+{
+	TArray<FPCGPinProperties> Properties;
+
+	FPCGPinProperties& Pin = Properties.Emplace_GetRef(PCGPinConstants::DefaultOutputLabel, EPCGDataType::Any);
+	Pin.SetRequiredPin();
+	Pin.SetAllowMultipleConnections(false);
+	Pin.bInvisiblePin = true;
+
+	return Properties;
+}
+
+FPCGElementPtr UPVPCGOutputSettings::CreateElement() const
+{
+	return MakeShared<FPVPCGOutputElement>();
+}
+
+bool FPVPCGOutputElement::ExecuteInternal(FPCGContext* InContext) const
+{
+	const UPVPCGOutputSettings* Settings = InContext->GetInputSettings<UPVPCGOutputSettings>();
+	check(Settings);
+
+	for (const FPCGTaggedData& Input : InContext->InputData.GetInputsByPin(PCGPinConstants::DefaultInputLabel))
+	{
+		if (const UPVPCGData* InputData = Cast<UPVPCGData>(Input.Data))
+		{
+			const FString OutputPath = Settings->ExportSettings.ContentBrowserFolder.Path;
+			if (FPackageName::EErrorCode ErrorCode; !FPackageName::IsValidLongPackageName(OutputPath, false, &ErrorCode))
+			{
+				PCGLog::LogErrorOnGraph(FText::FromString(FormatLongPackageNameErrorCode(ErrorCode)), InContext);
+				return true;
+			}
+
+			FManagedArrayCollection OutCollection = InputData->GetCollection();
+
+			UPVPCGData* ProceduralVegetationData = FPCGContext::NewObject_AnyThread<UPVPCGData>(InContext);
+			ProceduralVegetationData->Initialize(MoveTemp(OutCollection));
+
+			InContext->OutputData.TaggedData.Emplace(ProceduralVegetationData);
+		}
+		else
+		{
+			PCGLog::InputOutput::LogInvalidInputDataError(InContext);
+			return true;
+		}
+	}
+
+	return true;
+}
+
+#undef LOCTEXT_NAMESPACE
diff --git a/Engine/Plugins/Experimental/ProceduralVegetationEditor/Source/ProceduralVegetation/Private/Nodes/PVPCGRemoveBranches.cpp b/Engine/Plugins/Experimental/ProceduralVegetationEditor/Source/ProceduralVegetation/Private/Nodes/PVPCGRemoveBranches.cpp
new file mode 100644
--- /dev/null
+++ b/Engine/Plugins/Experimental/ProceduralVegetationEditor/Source/ProceduralVegetation/Private/Nodes/PVPCGRemoveBranches.cpp
@@ -0,0 +1,45 @@
+﻿// Copyright Epic Games, Inc. All Rights Reserved.
+
+#include "PVPCGRemoveBranches.h"
+#include "PVPCGData.h"
+#include "Data/PCGBasePointData.h"
+#include "GeometryCollection/ManagedArrayCollection.h"
+#include "Helpers/PCGSettingsHelpers.h"
+#include "Implementations/PVRemoveBranches.h"
+
+FPCGElementPtr UPVPCGRemoveBranchesSettings::CreateElement() const
+{
+	return MakeShared<FPVPCGRemoveBranchesElement>();
+}
+
+bool FPVPCGRemoveBranchesElement::ExecuteInternal(FPCGContext* InContext) const
+{
+	TRACE_CPUPROFILER_EVENT_SCOPE(FPVPCGRemoveBranchesElement::Execute);
+
+	check(InContext);
+
+	const UPVPCGRemoveBranchesSettings* Settings = InContext->GetInputSettings<UPVPCGRemoveBranchesSettings>();
+	check(Settings);
+	
+	for (const FPCGTaggedData& Input : InContext->InputData.GetInputsByPin(PCGPinConstants::DefaultInputLabel))
+	{
+		if(const UPVPCGData* InputData = Cast<UPVPCGData>(Input.Data))
+		{
+			FManagedArrayCollection OutCollection = InputData->GetCollection();
+		
+			UPVPCGData* OutManagedArrayCollectionData = FPCGContext::NewObject_AnyThread<UPVPCGData>(InContext);
+			
+			FPVRemoveBranches::ApplyRemoveBranches(Settings->BranchRemoveBasis, Settings->Threshold, OutCollection);
+			
+			OutManagedArrayCollectionData->Initialize(MoveTemp(OutCollection));
+			InContext->OutputData.TaggedData.Emplace(OutManagedArrayCollectionData);
+		}
+		else
+		{
+			PCGLog::InputOutput::LogInvalidInputDataError(InContext);
+			return true;
+		}
+	}
+	
+	return true;
+}
diff --git a/Engine/Plugins/Experimental/ProceduralVegetationEditor/Source/ProceduralVegetation/Private/Nodes/PVPCGRemoveBranches.h b/Engine/Plugins/Experimental/ProceduralVegetationEditor/Source/ProceduralVegetation/Private/Nodes/PVPCGRemoveBranches.h
new file mode 100644
--- /dev/null
+++ b/Engine/Plugins/Experimental/ProceduralVegetationEditor/Source/ProceduralVegetation/Private/Nodes/PVPCGRemoveBranches.h
@@ -0,0 +1,38 @@
+﻿// Copyright Epic Games, Inc. All Rights Reserved.
+#pragma once
+
+#include "CoreMinimal.h"
+#include "Nodes/PVPCGBaseSettings.h"
+#include "Implementations/PVRemoveBranches.h"
+#include "PVPCGRemoveBranches.generated.h"
+
+UCLASS(BlueprintType, ClassGroup = (Procedural))
+class UPVPCGRemoveBranchesSettings : public UPVPCGBaseSettings
+{
+	GENERATED_BODY()
+
+public:
+#if WITH_EDITOR
+	virtual FName GetDefaultNodeName() const override { return FName(TEXT("ProceduralVegetationRemoveBranches")); }
+	virtual FText GetDefaultNodeTitle() const override { return NSLOCTEXT("PVPCGScale", "NodeTitle", "Remove Branches"); }
+	virtual FText GetNodeTooltipText() const override { return NSLOCTEXT("PVPCGScale", "NodeTooltip", "Remove branches on the basis of selected criteria"); }
+	virtual FLinearColor GetNodeTitleColor() const override { return FLinearColor::Yellow; }
+
+	virtual TArray<EPVRenderType> GetRenderType() const override { return TArray{ EPVRenderType::Skeleton }; }
+#endif
+
+protected:
+	virtual FPCGElementPtr CreateElement() const override;
+public:
+	UPROPERTY(EditAnywhere, Category = Settings)
+	ERemoveBranchesBasis BranchRemoveBasis = ERemoveBranchesBasis::Length;
+
+	UPROPERTY(EditAnywhere, Category = Settings, meta=(ClampMin = 0f, ClampMax = 10.0f, UIMin = 0f, UIMax = 1.0f))
+	float Threshold = 0;
+};
+
+class FPVPCGRemoveBranchesElement : public FPVPCGBaseElement
+{
+protected:
+	virtual bool ExecuteInternal(FPCGContext* InContext) const override;
+};
diff --git a/Engine/Plugins/Experimental/ProceduralVegetationEditor/Source/ProceduralVegetation/Private/Nodes/PVPCGScale.cpp b/Engine/Plugins/Experimental/ProceduralVegetationEditor/Source/ProceduralVegetation/Private/Nodes/PVPCGScale.cpp
new file mode 100644
--- /dev/null
+++ b/Engine/Plugins/Experimental/ProceduralVegetationEditor/Source/ProceduralVegetation/Private/Nodes/PVPCGScale.cpp
@@ -0,0 +1,45 @@
+﻿// Copyright Epic Games, Inc. All Rights Reserved.
+
+#include "PVPCGScale.h"
+#include "PVPCGData.h"
+#include "Data/PCGBasePointData.h"
+#include "GeometryCollection/ManagedArrayCollection.h"
+#include "Helpers/PCGSettingsHelpers.h"
+#include "Implementations/PVScale.h"
+
+FPCGElementPtr UPVPCGScaleSettings::CreateElement() const
+{
+	return MakeShared<FPVPCGScaleElement>();
+}
+
+bool FPVPCGScaleElement::ExecuteInternal(FPCGContext* InContext) const
+{
+	TRACE_CPUPROFILER_EVENT_SCOPE(FPVPCGScaleElement::Execute);
+
+	check(InContext);
+
+	const UPVPCGScaleSettings* Settings = InContext->GetInputSettings<UPVPCGScaleSettings>();
+	check(Settings);
+	
+	for (const FPCGTaggedData& Input : InContext->InputData.GetInputsByPin(PCGPinConstants::DefaultInputLabel))
+	{
+		if(const UPVPCGData* InputData = Cast<UPVPCGData>(Input.Data))
+		{
+			FManagedArrayCollection OutCollection = InputData->GetCollection();
+		
+			UPVPCGData* OutManagedArrayCollectionData = FPCGContext::NewObject_AnyThread<UPVPCGData>(InContext);
+			
+			FPVScale::ApplyScale(Settings->Scale, OutCollection);
+			
+			OutManagedArrayCollectionData->Initialize(MoveTemp(OutCollection));
+			InContext->OutputData.TaggedData.Emplace(OutManagedArrayCollectionData);
+		}
+		else
+		{
+			PCGLog::InputOutput::LogInvalidInputDataError(InContext);
+			return true;
+		}
+	}
+	
+	return true;
+}
diff --git a/Engine/Plugins/Experimental/ProceduralVegetationEditor/Source/ProceduralVegetation/Private/Nodes/PVPCGScale.h b/Engine/Plugins/Experimental/ProceduralVegetationEditor/Source/ProceduralVegetation/Private/Nodes/PVPCGScale.h
new file mode 100644
--- /dev/null
+++ b/Engine/Plugins/Experimental/ProceduralVegetationEditor/Source/ProceduralVegetation/Private/Nodes/PVPCGScale.h
@@ -0,0 +1,34 @@
+﻿// Copyright Epic Games, Inc. All Rights Reserved.
+#pragma once
+
+#include "CoreMinimal.h"
+#include "Nodes/PVPCGBaseSettings.h"
+#include "PVPCGScale.generated.h"
+
+UCLASS(BlueprintType, ClassGroup = (Procedural))
+class UPVPCGScaleSettings : public UPVPCGBaseSettings
+{
+	GENERATED_BODY()
+
+public:
+#if WITH_EDITOR
+	virtual FName GetDefaultNodeName() const override { return FName(TEXT("ProceduralVegetationScale")); }
+	virtual FText GetDefaultNodeTitle() const override { return NSLOCTEXT("PVPCGScale", "NodeTitle", "Scale"); }
+	virtual FText GetNodeTooltipText() const override { return NSLOCTEXT("PVPCGScale", "NodeTooltip", "Scale the tree based on the settings"); }
+	virtual FLinearColor GetNodeTitleColor() const override { return FLinearColor::Yellow; }
+
+	virtual TArray<EPVRenderType> GetRenderType() const override { return TArray{ EPVRenderType::Skeleton }; }
+#endif
+
+protected:
+	virtual FPCGElementPtr CreateElement() const override;
+public:
+	UPROPERTY(EditAnywhere, Category = Settings, meta=(ClampMin = 0.01f, ClampMax = 100.0f, UIMin = 0.01f, UIMax = 10.0f))
+	float Scale = 1;
+};
+
+class FPVPCGScaleElement : public FPVPCGBaseElement
+{
+protected:
+	virtual bool ExecuteInternal(FPCGContext* InContext) const override;
+};
diff --git a/Engine/Plugins/Experimental/ProceduralVegetationEditor/Source/ProceduralVegetation/Private/PVPCGData.cpp b/Engine/Plugins/Experimental/ProceduralVegetationEditor/Source/ProceduralVegetation/Private/PVPCGData.cpp
new file mode 100644
--- /dev/null
+++ b/Engine/Plugins/Experimental/ProceduralVegetationEditor/Source/ProceduralVegetation/Private/PVPCGData.cpp
@@ -0,0 +1,86 @@
+// Copyright Epic Games, Inc. All Rights Reserved.
+
+#include "PVPCGData.h"
+#include "PCGContext.h"
+#include "Data/PCGPointArrayData.h"
+#include "Data/PCGPointData.h"
+#include "Helpers/PCGHelpers.h"
+#include "Serialization/ArchiveCrc32.h"
+#include "UObject/Package.h"
+
+UPVPCGData::UPVPCGData(const FObjectInitializer& ObjectInitializer)
+	: UPCGSpatialData(ObjectInitializer)
+{
+	
+}
+
+void UPVPCGData::Initialize(FManagedArrayCollection&& InCollection, bool bCanTakeOwnership)
+{
+	Collection = InCollection;
+}
+
+void UPVPCGData::AddToCrc(FArchiveCrc32& Ar, bool bFullDataCrc) const
+{
+	Super::AddToCrc(Ar, bFullDataCrc);
+
+	if (bFullDataCrc)
+	{
+		// Implementation note: metadata not supported at this point.
+
+		FString ClassName = StaticClass()->GetPathName();
+
+		Ar << ClassName;
+		Ar << Collection;
+#if WITH_EDITOR
+		Ar << RenderType;
+#endif
+	}
+	else
+	{
+		AddUIDToCrc(Ar);
+	}
+}
+
+FBox UPVPCGData::GetBounds() const
+{
+	return CachedBounds;
+}
+
+bool UPVPCGData::SamplePoint(const FTransform& InTransform, const FBox& InBounds, FPCGPoint& OutPoint, UPCGMetadata* OutMetadata) const
+{
+	return false;
+}
+
+const UPCGBasePointData* UPVPCGData::ToBasePointData(FPCGContext* Context, const FBox& InBounds, TSubclassOf<UPCGBasePointData> PointDataClass) const
+{
+	TRACE_CPUPROFILER_EVENT_SCOPE(UPCGPrimitiveData::CreatePointData);
+
+	UPCGBasePointData* Data = FPCGContext::NewObject_AnyThread<UPCGBasePointData>(Context, GetTransientPackage(), PointDataClass);
+	return Data;
+}
+
+UPCGSpatialData* UPVPCGData::CopyInternal(FPCGContext* Context) const
+{
+	UPVPCGData* NewProceduralVegetationData = FPCGContext::NewObject_AnyThread<UPVPCGData>(Context);
+
+	FManagedArrayCollection NewCollection;
+	Collection.CopyTo(&NewCollection);
+	
+	NewProceduralVegetationData->Initialize(MoveTemp(NewCollection), /*bCanTakeOwnership=*/false);
+
+#if WITH_EDITOR
+	NewProceduralVegetationData->RenderType = RenderType;
+#endif
+	
+	return NewProceduralVegetationData;
+}
+
+const UPCGPointArrayData* UPVPCGData::ToPointArrayData(FPCGContext* Context, const FBox& InBounds) const
+{
+	return CastChecked<UPCGPointArrayData>(ToBasePointData(Context, InBounds, UPCGPointArrayData::StaticClass()));
+}
+const UPCGPointData* UPVPCGData::ToPointData(FPCGContext* Context, const FBox& InBounds) const
+{
+	return CastChecked<UPCGPointData>(ToBasePointData(Context, InBounds, UPCGPointData::StaticClass()));
+}
+
diff --git a/Engine/Plugins/Experimental/ProceduralVegetationEditor/Source/ProceduralVegetation/Private/ProceduralVegetationLink.cpp b/Engine/Plugins/Experimental/ProceduralVegetationEditor/Source/ProceduralVegetation/Private/ProceduralVegetationLink.cpp
new file mode 100644
--- /dev/null
+++ b/Engine/Plugins/Experimental/ProceduralVegetationEditor/Source/ProceduralVegetation/Private/ProceduralVegetationLink.cpp
@@ -0,0 +1,82 @@
+// Copyright Epic Games, Inc. All Rights Reserved.
+
+#include "ProceduralVegetationLink.h"
+#include "Nodes/PVPCGFoliage.h"
+#include "Nodes/PVPCGGravity.h"
+#include "Nodes/PVPCGMesher.h"
+#include "Nodes/PVPCGMegaSeedLoad.h"
+#include "Nodes/PVPCGOutput.h"
+
+void SanitizeFolderName(FString& FolderName)
+{
+	const TCHAR* InvalidChar = INVALID_OBJECTPATH_CHARACTERS INVALID_LONGPACKAGE_CHARACTERS;
+
+	while (*InvalidChar)
+	{
+		FolderName.ReplaceCharInline(*InvalidChar, TCHAR('_'), ESearchCase::CaseSensitive);
+		++InvalidChar;
+	}
+
+	FolderName.ReplaceCharInline('/', TCHAR('_'), ESearchCase::CaseSensitive);
+}
+
+void UProceduralVegetation::CreateGraph(const UProceduralVegetationGraph* InGraph /*= nullptr*/)
+{
+	if (InGraph)
+	{
+		Graph = DuplicateObject(InGraph, this);
+	}
+	else
+	{
+		Graph = NewObject<UProceduralVegetationGraph>(this, UProceduralVegetationGraph::StaticClass());
+	}
+}
+
+#if WITH_EDITOR
+void UProceduralVegetation::CreateGraphFromMegaSeed(const TObjectPtr<UMegaSeedDataAsset> MegaSeedAsset)
+{
+	Graph = NewObject<UProceduralVegetationGraph>(this, UProceduralVegetationGraph::StaticClass());
+
+	UPCGSettings* DefaultNodeSettings = nullptr;
+
+	UPVPCGMegaSeedLoadSettings* MegaSeedLoadNodeSettings = NewObject<UPVPCGMegaSeedLoadSettings>(Graph);
+	MegaSeedLoadNodeSettings->MegaSeedDataAsset = MegaSeedAsset;
+	UPCGNode* MegaSeedLoadNode = Graph->AddNodeCopy(MegaSeedLoadNodeSettings, DefaultNodeSettings);
+
+	FInt32Vector2 Position(0, 300);
+
+	MegaSeedLoadNode->SetNodePosition(Position.X, Position.Y);
+	Position.X += 100;
+
+	for (const auto& [VariantName, VariantData] : MegaSeedAsset->Variants)
+	{
+		const UPVPCGGravitySettings* GravityNodeSettings = NewObject<UPVPCGGravitySettings>(Graph);
+		const UPVPCGMesherSettings* MesherNodeSettings = NewObject<UPVPCGMesherSettings>(Graph);
+		const UPVPCGFoliageSettings* FoliageNodeSettings = NewObject<UPVPCGFoliageSettings>(Graph);
+		UPVPCGOutputSettings* OutputNodeSettings = NewObject<UPVPCGOutputSettings>(Graph);
+		FString MeshName = *VariantName;
+		SanitizeFolderName(MeshName);
+		OutputNodeSettings->ExportSettings.MeshName = *MeshName;
+
+		UPCGNode* GravityNode = Graph->AddNodeCopy(GravityNodeSettings, DefaultNodeSettings);
+		UPCGNode* MesherNode = Graph->AddNodeCopy(MesherNodeSettings, DefaultNodeSettings);
+		UPCGNode* FoliageNode = Graph->AddNodeCopy(FoliageNodeSettings, DefaultNodeSettings);
+		UPCGNode* OutputNode = Graph->AddNodeCopy(OutputNodeSettings, DefaultNodeSettings);
+
+		OutputNode->NodeTitle = *VariantName;
+
+		MegaSeedLoadNode
+			->AddEdgeTo(*VariantName, GravityNode, PCGPinConstants::DefaultInputLabel)
+			->AddEdgeTo(PCGPinConstants::DefaultOutputLabel, MesherNode, PCGPinConstants::DefaultInputLabel)
+			->AddEdgeTo(PCGPinConstants::DefaultOutputLabel, FoliageNode, PCGPinConstants::DefaultInputLabel)
+			->AddEdgeTo(PCGPinConstants::DefaultOutputLabel, OutputNode, PCGPinConstants::DefaultInputLabel);
+
+		GravityNode->SetNodePosition(Position.X + 200, Position.Y);
+		MesherNode->SetNodePosition(Position.X + 400, Position.Y);
+		FoliageNode->SetNodePosition(Position.X + 600, Position.Y);
+		OutputNode->SetNodePosition(Position.X + 800, Position.Y);
+
+		Position.Y += 150;
+	}
+}
+#endif
diff --git a/Engine/Plugins/Experimental/ProceduralVegetationEditor/Source/ProceduralVegetation/ProceduralVegetation.Build.cs b/Engine/Plugins/Experimental/ProceduralVegetationEditor/Source/ProceduralVegetation/ProceduralVegetation.Build.cs
new file mode 100644
--- /dev/null
+++ b/Engine/Plugins/Experimental/ProceduralVegetationEditor/Source/ProceduralVegetation/ProceduralVegetation.Build.cs
@@ -0,0 +1,50 @@
+// Copyright Epic Games, Inc. All Rights Reserved.
+
+namespace UnrealBuildTool.Rules
+{
+	public class ProceduralVegetation : ModuleRules
+	{
+		public ProceduralVegetation(ReadOnlyTargetRules Target) : base(Target)
+		{
+			PCHUsage = PCHUsageMode.UseExplicitOrSharedPCHs;
+
+			PublicDependencyModuleNames.AddRange(
+				new string[]
+				{
+				});
+
+			PrivateDependencyModuleNames.AddRange(
+				new string[]
+				{
+					"Core",
+					"CoreUObject",
+					"Chaos",
+					"Engine",
+					"RenderCore",
+					"GeometryCore",
+					"GeometryFramework",
+					"GeometryScriptingCore",
+					"Json",
+					"JsonUtilities",
+					"MeshDescription",
+					"MeshConversion",
+					"PlanarCut",
+					"SkeletalMeshDescription",
+					"StaticMeshDescription",
+					"ModelingComponents",
+					"PCG",
+					"DynamicWind"// ... add private dependencies that you statically link with here ...	
+				}
+			);
+			
+			if (Target.Type == TargetRules.TargetType.Editor)
+			{
+				PrivateDependencyModuleNames.AddRange(
+					new string[]
+					{
+						"DynamicWindEditor"
+					});
+			}
+		}
+	}
+}
diff --git a/Engine/Plugins/Experimental/ProceduralVegetationEditor/Source/ProceduralVegetation/Public/Facades/PVAttributesNames.h b/Engine/Plugins/Experimental/ProceduralVegetationEditor/Source/ProceduralVegetation/Public/Facades/PVAttributesNames.h
new file mode 100644
--- /dev/null
+++ b/Engine/Plugins/Experimental/ProceduralVegetationEditor/Source/ProceduralVegetation/Public/Facades/PVAttributesNames.h
@@ -0,0 +1,87 @@
+// Copyright Epic Games, Inc. All Rights Reserved.
+#pragma once
+
+#include "CoreMinimal.h"
+
+namespace PV::GroupNames
+{
+	//Group Names
+	inline static const FName BranchGroup = FName(TEXT("Primitives"));
+	inline static const FName PointGroup = FName(TEXT("Points"));
+	inline static const FName FoliageGroup = FName(TEXT("Foliage"));
+	inline static const FName FoliageNamesGroup = FName(TEXT("FoliageNames"));
+	inline static const FName DetailsGroup = FName(TEXT("Details"));
+	inline static const FName PlantProfilesGroup = FName(TEXT("PlantProfiles"));
+	inline static const FName BonesGroup = FName(TEXT("Bones"));
+	inline static const FName VerticesGroup = FName(TEXT("Vertices"));
+}
+
+namespace PV::AttributeNames
+{
+	//Branch Attributes
+	inline static const FName BranchParents = FName(TEXT("parents"));
+	inline static const FName BranchChildren = FName(TEXT("children"));
+	inline static const FName BranchPoints = FName(TEXT("Points"));
+	inline static const FName BranchNumber = FName(TEXT("branchNumber"));
+	inline static const FName BranchSourceBudNumber = FName(TEXT("branchSourceBudNumber"));
+	inline static const FName BranchFoliageIDs = FName(TEXT("FoliageID"));
+	inline static const FName BranchUVMaterial = FName(TEXT("branchUVMat"));
+	inline static const FName TrunkMaterialPath = FName(TEXT("TrunkMaterialPath"));
+	inline static const FName BranchHierarchyNumber = FName(TEXT("branchHierarchyNumber"));
+	inline static const FName BranchSimulationGroupIndex = FName(TEXT("branchSimulationGroupIndex"));
+
+	//Point Attributes
+	inline static const FName PointPosition = FName(TEXT("Position"));
+	inline static const FName LengthFromRoot = FName(TEXT("lengthFromRoot"));
+	inline static const FName PointScaleGradient = FName(TEXT("LOD_totalPscaleGradient"));
+	inline static const FName PointScale = FName(TEXT("Scale"));
+	
+	inline static const FName LengthFromSeed = FName(TEXT("lengthFromSeed"));
+	inline static const FName BudLightDetected = FName(TEXT("budLightDetected"));
+	inline static const FName BudDevelopment = FName(TEXT("budDevelopment"));
+	inline static const FName TextureCoordV = FName(TEXT("uv_v"));
+	inline static const FName TextureCoordUOffset = FName(TEXT("uv_uOffset"));
+	
+	inline static const FName HullGradient = FName(TEXT("LOD_hullGradient"));
+	inline static const FName MainTrunkGradient = FName(TEXT("LOD_mainTrunkGradient"));
+	inline static const FName GroundGradient = FName(TEXT("LOD_groundGradient"));
+	inline static const FName BudNumber = FName(TEXT("budNumber"));
+	inline static const FName BudHormoneLevels = FName(TEXT("budHormoneLevels"));
+	inline static const FName BudDirection = FName(TEXT("budDirection"));
+	inline static const FName BranchGradient = FName(TEXT("branchGradient"));
+	inline static const FName PlantGradient = FName(TEXT("plantGradient"));
+	inline static const FName NjordPixelIndex = FName(TEXT("njord_pixelIdx"));
+
+	//Foliage Attributes
+	inline static const FName FoliageNameID = FName(TEXT("FoliageNameID"));
+	inline static const FName FoliageBranchID = FName(TEXT("FoliageBranchID"));
+	inline static const FName FoliagePivotPoint = FName(TEXT("FoliagePivotPoint"));
+	inline static const FName FoliageUPVector = FName(TEXT("FoliageUPVector"));
+	inline static const FName FoliageNormalVector = FName(TEXT("FoliageNormalVector"));
+	inline static const FName FoliageScale = FName(TEXT("FoliageScale"));
+	inline static const FName FoliageLengthFromRoot = FName(TEXT("LengthFromRoot"));
+	inline static const FName FoliageParentBoneID = FName(TEXT("FoliageParentBoneID"));
+
+	//FoliageNames Attributes
+	inline static const FName FoliageName = FName(TEXT("Name"));
+
+	//Details Attributes
+	inline static const FName LeafPhyllotaxy = FName(TEXT("phyllotaxyLeaf"));
+	inline static const FName FoliagePath = FName(TEXT("FoliagePath"));
+	inline static const FName Guid = FName(TEXT("Guid"));
+
+	//PlantProfiles Attributes
+	inline static const FName ProfilePoints = FName(TEXT("ProfilePoints"));
+
+	//Bone Attributes
+	inline static const FName BoneName = FName(TEXT("BoneName"));
+	inline static const FName BoneParentIndex = FName(TEXT("BoneParentIndex"));
+	inline static const FName BonePose = FName(TEXT("BonePose"));
+	inline static const FName BonePointIndex = FName(TEXT("BonePointIndex"));
+	inline static const FName BoneAbsolutePosition = FName(TEXT("BoneAbsolutePosition"));
+	inline static const FName BoneBranchIndex = FName(TEXT("BoneBranchIndex"));
+	inline static const FName NjordPixelId = FName(TEXT("NjordPixelId"));
+	inline static const FName VertexPointIds = FName(TEXT("VertexPointIds"));
+	inline static const FName BoneId = FName(TEXT("BoneId"));
+}
+
diff --git a/Engine/Plugins/Experimental/ProceduralVegetationEditor/Source/ProceduralVegetation/Public/Facades/PVBranchFacade.h b/Engine/Plugins/Experimental/ProceduralVegetationEditor/Source/ProceduralVegetation/Public/Facades/PVBranchFacade.h
new file mode 100644
--- /dev/null
+++ b/Engine/Plugins/Experimental/ProceduralVegetationEditor/Source/ProceduralVegetation/Public/Facades/PVBranchFacade.h
@@ -0,0 +1,88 @@
+﻿// Copyright Epic Games, Inc. All Rights Reserved.
+
+#pragma once
+
+#include "CoreMinimal.h"
+#include "IShrinkableFacade.h"
+#include "GeometryCollection/ManagedArrayAccessor.h"
+
+namespace PV::Facades
+{
+	/**
+	 * FBranchFacade is used to access and manipulate the data from Branch Group with in the ProceduralVegetation's FManagedArrayCollection
+	 * A branch is made up of multiple points. Each Branch will know about its parents and children, along with other attributes
+	 * Only add the frequently used branch attributes and their access to this facade, for the specific access write a new facade
+	 */
+	class PROCEDURALVEGETATION_API FBranchFacade final : public IShrinkable
+	{
+	public:
+		FBranchFacade(FManagedArrayCollection& InCollection);
+		FBranchFacade(const FManagedArrayCollection& InCollection);
+
+		bool IsConst() const { return Collection == nullptr; }
+
+		bool IsValid() const;
+
+		virtual int32 GetElementCount() const override;
+
+		const TArray<int32>& GetPoints(const int32 Index) const;
+
+		void SetPoints(const int32 Index, const TArray<int32>& InPoints);
+
+		const TArray<int32>& GetChildren(const int32 Index) const;
+
+		void SetChildren(const int32 Index, const TArray<int32>& InChildren);
+
+		const TArray<int32>& GetParents(const int32 Index) const;
+
+		void SetParents(const int32 Index, const TArray<int32>& InParents);
+
+		int32 GetBranchNumber(const int32 Index) const;
+
+		int32 GetBranchSourceBudNumber(const int32 Index) const;
+
+		int32 GetTrunkIndex() const;
+
+		TArray<int32> GetParentBranchIndices(int BranchIndex) const;
+
+		int32 GetGenerationNumber(const int32 Index) const;
+
+		const TManagedArray<int32>& GetBranchNumbers() const;
+
+		int32 GetBranchUVMaterial(const int32 Index) const;
+
+		void SetBranchUVMaterial(const int32 Index, const int32 InMaterial);
+		
+		int32 GetBranchHierarchyNumber(int32 Index) const;
+
+		int32 GetBranchSimulationGroupIndex(int32 Index) const;
+		
+		void SetBranchSimulationGroupIndex(int32 Index, int32 InLogicalDepth);
+
+		void SetTrunkMaterialPath(const FString& InPath);
+
+		FString GetTrunkMaterialPath() const;
+
+		virtual void CopyEntry(const int32 FromIndex, const int32 ToIndex) override;
+
+		virtual void RemoveEntries(const int32 NumEntries, const int32 StartIndex) override;
+
+		void GetSortedBranchIndicesByHierarchy(TArray<int32>& OutSortedIndices) const;
+		
+		int32 GetBranchIndexFromPointIndex(int32 PointIndex) const;
+
+	private:
+		FManagedArrayCollection* Collection = nullptr;
+
+		TManagedArrayAccessor<TArray<int32>> Parents;
+		TManagedArrayAccessor<TArray<int32>> Children;
+		TManagedArrayAccessor<TArray<int32>> BranchPoints;
+		TManagedArrayAccessor<int32> BranchNumbers;
+		TManagedArrayAccessor<int32> BranchSourceBudNumber;
+		TManagedArrayAccessor<TArray<int32>> BranchFoliageIDs;
+		TManagedArrayAccessor<int32> BranchUVMaterial;
+		TManagedArrayAccessor<int32> BranchHierarchyNumber;
+		TManagedArrayAccessor<int32> BranchSimulationGroupIndex;
+		TManagedArrayAccessor<FString> TrunkMaterialPathAttribute;
+	};
+}
diff --git a/Engine/Plugins/Experimental/ProceduralVegetationEditor/Source/ProceduralVegetation/Public/Facades/PVFoliageFacade.h b/Engine/Plugins/Experimental/ProceduralVegetationEditor/Source/ProceduralVegetation/Public/Facades/PVFoliageFacade.h
new file mode 100644
--- /dev/null
+++ b/Engine/Plugins/Experimental/ProceduralVegetationEditor/Source/ProceduralVegetation/Public/Facades/PVFoliageFacade.h
@@ -0,0 +1,105 @@
+﻿// Copyright Epic Games, Inc. All Rights Reserved.
+
+#pragma once
+
+#include "CoreMinimal.h"
+#include "IShrinkableFacade.h"
+#include "GeometryCollection/ManagedArrayAccessor.h"
+
+namespace PV::Facades
+{
+	struct PROCEDURALVEGETATION_API FFoliageEntryData
+	{
+		int32 NameId;
+		int32 BranchId;
+		FVector3f PivotPoint;
+		FVector3f UpVector;
+		FVector3f NormalVector;
+		float Scale;
+		float LengthFromRoot;
+		int32 ParentBoneID = INDEX_NONE;
+	};
+
+	/**
+	 * FFoliageFacade is used to access and manipulate Foliage related data for a ProceduralVegetation's FManagedArrayCollection
+	 * It presents access to foliage Mesh names, their respective Branch ids, and instancing transform data (Pivot point, Up Vector, and Scale)
+	 * It lays out data in two Groups: FoliageNames and Foliage where Foliage references FoliageNames and Branch/Primitives with Ids
+	 * It also adds Foliage entry Ids to the Branch/Primitives group. 
+	 */
+	class PROCEDURALVEGETATION_API FFoliageFacade final : public IShrinkable
+	{
+	public:
+		FFoliageFacade(FManagedArrayCollection& InCollection, const int32 InitialSize = 0);
+		FFoliageFacade(const FManagedArrayCollection& InCollection);
+
+		bool IsConst() const { return Collection == nullptr; }
+
+		bool IsValid() const;
+
+		int32 NumFoliageEntries() const { return NameIdsAttribute.Num(); };
+
+		FFoliageEntryData GetFoliageEntry(const int32 Index) const;
+
+		const float& GetLengthFromRoot(const int32 Index) const;
+
+		void SetLengthFromRoot(const int32 Index, const float Input);
+
+		int32 GetParentBoneID(int32 Index) const;
+
+		void SetParentBoneID(int32 Index, int32 Input);
+
+		const TArray<int32>& GetFoliageEntryIdsForBranch(int32 Index) const;
+
+		int32 AddFoliageEntry(const FFoliageEntryData& InputData);
+
+		void SetFoliageEntry(const int32 Index, const FFoliageEntryData& InputData);
+
+		void SetFoliageIdsArray(const int32 Index, const TArray<int32>& InputIds);
+
+		void SetFoliageBranchId(const int32 Index, const int32 InputId);
+
+		int32 NumFoliageNames() const { return NamesAttribute.Num(); };
+
+		FString GetFoliageName(const int32 Index) const;
+
+		const FVector3f& GetPivotPoint(const int32 Index) const;
+
+		void SetPivotPoint(const int32 Index, const FVector3f& Input);
+
+		void SetScale(const int32 Index, const float& Input);
+
+		void SetFoliageName(const int32 Index, const FString& Input);
+
+		void SetFoliageNames(const TArray<FString>& InputNames);
+
+		void SetFoliagePath(const FString InPath);
+
+		FString GetFoliagePath() const;
+
+		virtual int32 GetElementCount() const override;
+
+		virtual void CopyEntry(const int32 FromIndex, const int32 ToIndex) override;
+
+		virtual void RemoveEntries(const int32 NumEntries, const int32 StartIndex) override;
+		
+		FTransform GetFoliageTransform(int32 Id) const;
+
+	protected:
+		void DefineSchema(const int32 InitialSize = 0);
+
+		const FManagedArrayCollection& ConstCollection;
+		FManagedArrayCollection* Collection = nullptr;
+
+		TManagedArrayAccessor<FString> NamesAttribute;
+		TManagedArrayAccessor<int32> NameIdsAttribute;
+		TManagedArrayAccessor<int32> BranchIdsAttribute;
+		TManagedArrayAccessor<FVector3f> PivotPointsAttribute;
+		TManagedArrayAccessor<FVector3f> UpVectorsAttribute;
+		TManagedArrayAccessor<FVector3f> NormalVectorsAttribute;
+		TManagedArrayAccessor<float> ScalesAttribute;
+		TManagedArrayAccessor<float> LengthFromRootAttribute;
+		TManagedArrayAccessor<int32> ParentBoneIdsAttribute;
+		TManagedArrayAccessor<TArray<int32>> FoliageIdsAttribute;
+		TManagedArrayAccessor<FString> FoliagePathAttribute;
+	};
+}
diff --git a/Engine/Plugins/Experimental/ProceduralVegetationEditor/Source/ProceduralVegetation/Public/Facades/PVPointFacade.h b/Engine/Plugins/Experimental/ProceduralVegetationEditor/Source/ProceduralVegetation/Public/Facades/PVPointFacade.h
new file mode 100644
--- /dev/null
+++ b/Engine/Plugins/Experimental/ProceduralVegetationEditor/Source/ProceduralVegetation/Public/Facades/PVPointFacade.h
@@ -0,0 +1,111 @@
+﻿// Copyright Epic Games, Inc. All Rights Reserved.
+
+#pragma once
+
+#include "CoreMinimal.h"
+#include "IShrinkableFacade.h"
+#include "GeometryCollection/ManagedArrayAccessor.h"
+
+namespace PV::Facades
+{
+	/**
+	 * FPointFacade is used to access and manipulate the Point Group data from the ProceduralVegetation's FManagedArrayCollection
+	 * Only add the frequently used Point attributes and their access to this facade, for the specific access write a new facade
+	 */
+	class PROCEDURALVEGETATION_API FPointFacade final : public IShrinkable
+	{
+	public:
+		FPointFacade(FManagedArrayCollection& InCollection);
+		FPointFacade(const FManagedArrayCollection& InCollection);
+
+		bool IsValid() const;
+
+		virtual int32 GetElementCount() const override;
+
+		const FVector3f& GetPosition(const int32 Index) const;
+
+		const float& GetLengthFromRoot(const int32 Index) const;
+
+		const float& GetLengthFromSeed(const int32 Index) const;
+
+		const float& GetPointScaleGradient(const int32 Index) const;
+
+		const float& GetHullGradient(const int32 Index) const;
+
+		const float& GetMainTrunkGradient(const int32 Index) const;
+
+		const float& GetGroundGradient(const int32 Index) const;
+		
+		const float& GetPointScale(const int32 Index) const;
+
+		const int32 GetBudNumber(const int32 Index) const;
+
+		const TArray<float>& GetBudLightDetected(const int32 Index) const;
+
+		const TArray<int>& GetBudDevelopment(const int32 Index) const;
+		
+		const TArray<FVector3f>& GetBudDirection(const int32 Index) const;
+
+		const TArray<float>& GetBudHormoneLevels(const int32 Index) const;
+
+		void SetBudHormoneLevels(const int32 Index, const TArray<float>& InBudHormoneLevels);
+
+		const float& GetBranchGradient(const int32 Index) const;
+
+		void SetBranchGradient(const int32 Index, const float InBranchGradient);
+
+		const float& GetPlantGradient(const int32 Index) const;
+
+		const float& GetNjordPixelIndex(const int32 Index) const;
+
+		void SetPlantGradient(const int32 Index, const float InPlantGradient);
+
+		const float& GetTextureCoordV(int32 Index) const;
+		
+		void SeTextureCoordV(int32 Index, float InTextureCoordV);
+		
+		const float& GetTextureCoordUOffset(int32 Index) const;
+		
+		void SetTextureCoordUOffset(int32 Index, float InTextureCoordUOffset);
+
+		TManagedArray<FVector3f>& ModifyPositions();
+		
+		TManagedArray<float>& ModifyPointScales();
+
+		TManagedArray<float>& ModifyLengthFromRoots();
+		
+		TManagedArray<float>& ModifyLengthFromSeeds();
+		
+		TManagedArray<float>& ModifyNjordPixelIDs();
+
+		const TManagedArray<FVector3f>& GetPositions() const;
+		
+		const TManagedArray<float>& GetPointScales() const;
+
+		const TManagedArray<float>& GetLengthFromRootsArray() const;
+
+		virtual void CopyEntry(const int32 FromIndex, const int32 ToIndex) override;
+
+		virtual void RemoveEntries(const int32 NumEntries, const int32 StartIndex) override;
+
+	private:
+		TManagedArrayAccessor<FVector3f> Positions;
+		TManagedArrayAccessor<float> LengthFromRoot;
+		TManagedArrayAccessor<float> PointScaleGradient;
+		TManagedArrayAccessor<float> HullGradient;
+		TManagedArrayAccessor<float> MainTrunkGradient;
+		TManagedArrayAccessor<float> GroundGradient;
+		TManagedArrayAccessor<float> PointScale;
+		TManagedArrayAccessor<TArray<FVector3f>> BudDirections;
+		TManagedArrayAccessor<float> BranchGradients;
+		TManagedArrayAccessor<TArray<float>> BudHormoneLevels;
+		TManagedArrayAccessor<float> PlantGradients;
+		TManagedArrayAccessor<float> LengthFromSeed;
+		TManagedArrayAccessor<int32> BudNumber;
+		TManagedArrayAccessor<float> NjordPixelIndex;
+		TManagedArrayAccessor<TArray<float>> BudLightDetected;
+		TManagedArrayAccessor<TArray<int>> BudDevelopment;
+		TManagedArrayAccessor<float> TextureCoordV;
+		TManagedArrayAccessor<float> TextureCoordUOffset;
+	};
+}
diff --git a/Engine/Plugins/Experimental/ProceduralVegetationEditor/Source/ProceduralVegetation/Public/Helpers/PVExportHelper.h b/Engine/Plugins/Experimental/ProceduralVegetationEditor/Source/ProceduralVegetation/Public/Helpers/PVExportHelper.h
new file mode 100644
--- /dev/null
+++ b/Engine/Plugins/Experimental/ProceduralVegetationEditor/Source/ProceduralVegetation/Public/Helpers/PVExportHelper.h
@@ -0,0 +1,138 @@
+﻿// Copyright Epic Games, Inc. All Rights Reserved.
+#pragma once
+
+#include "UDynamicMesh.h"
+#include "NaniteAssemblyDataBuilder.h"
+#include "Engine/EngineTypes.h"
+#include "GeometryCollection/ManagedArrayCollection.h"
+#include "Materials/MaterialInterface.h"
+#include "Facades/PVBoneFacade.h"
+#include "PVExportHelper.generated.h"
+
+class UInstancedStaticMeshComponent;
+class UStaticMesh;
+class USkeletalMesh;
+class UProceduralVegetation;
+struct FGeometryScriptCopyMeshToAssetOptions;
+
+UENUM()
+enum class EPVExportMeshType: uint8
+{
+	StaticMesh,
+	SkeletalMesh,
+};
+
+USTRUCT()
+struct PROCEDURALVEGETATION_API FPVExportSettings
+{
+	GENERATED_BODY()
+
+	UPROPERTY()
+	bool bShouldExport = true;
+
+	UPROPERTY(EditAnywhere, Category = "Export", Meta = (ContentDir, DisplayName = "Content Browser Folder"))
+	FDirectoryPath ContentBrowserFolder;
+
+	UPROPERTY(EditAnywhere, Category = "Export", Meta = (DisplayName = "Mesh Name"))
+	FName MeshName;
+
+	UPROPERTY(EditAnywhere, Category = "Export", Meta = (DisplayName = "Export Mesh Type"))
+	EPVExportMeshType ExportMeshType = EPVExportMeshType::StaticMesh;
+
+	UPROPERTY(EditAnywhere, Category = "Export", Meta = (DisplayName = "Create Nanite Assemblies"))
+	bool bCreateNaniteAssemblies = true;
+	
+	UPROPERTY(EditAnywhere, Category = "Export", Meta = (DisplayName = "Nanite Shape Preservation Method"))
+	ENaniteShapePreservation NaniteShapePreservation = ENaniteShapePreservation::Voxelize;
+
+	void Initialize(const FString& InAssetPath, const FString& InName)
+	{
+		FString DefaultDirectory = FPaths::GetPath(InAssetPath);
+		ContentBrowserFolder.Path = DefaultDirectory;
+		
+		MeshName = FName(InName);
+	}
+
+	FString GetOutputObjectPath() const
+	{
+		const FString AssetName = MeshName.ToString();
+		const FString OutputPath = ContentBrowserFolder.Path / AssetName + '.' + AssetName;
+		return OutputPath;
+	}
+};
+
+namespace PV::Export::Internal
+{
+	typedef TFunction<void(const FString& MeshName, int32 InstanceID, int32 VertexStart, int32 VertexCount)> FMeshInstanceCombined;
+	
+	TObjectPtr<UDynamicMesh> CollectionToDynamicMesh(const FManagedArrayCollection& Collection);
+
+	FGeometryScriptCopyMeshToAssetOptions GetCopyMeshToAssetOptions(
+		TArray<TObjectPtr<UMaterialInterface>> InMaterials
+	);
+	
+#if WITH_EDITORONLY_DATA
+	FGeometryScriptCopyMeshToAssetOptions GetCopyMeshToAssetOptions(
+		TArray<TObjectPtr<UMaterialInterface>> InMaterials,
+		const FMeshNaniteSettings& InNaniteSettings,
+		const ENaniteShapePreservation InShapePreservation
+	);
+#endif
+	
+	void CombineMeshInstancesToDynamicMesh(
+		const FManagedArrayCollection& Collection,
+		UDynamicMesh* DynamicMesh,
+		TArray<TObjectPtr<UMaterialInterface>>& Materials,
+		const FMeshInstanceCombined& OnMeshInstanceCombined
+	);
+
+#if WITH_EDITOR
+	void BuildNaniteAssemblyData(
+		const FManagedArrayCollection& Collection,
+		UStaticMesh* StaticMesh
+	);
+
+	void AddNodeToBuilder(FNaniteAssemblyDataBuilder& AssemblyBuilder,
+		const FManagedArrayCollection& Collection,
+		const TMap<FString, int32>& InMeshNamePartMap,
+		const TMap<FString, TArray<TObjectPtr<UMaterialInterface>>>& InMeshMaterialsMap);
+	
+#endif
+
+	void ExportCollectionToStaticMesh(
+		TObjectPtr<UStaticMesh> ExportMesh,
+		const FManagedArrayCollection& Collection,
+		bool bBuildNaniteAssemblies,
+		ENaniteShapePreservation ShapePreservation
+	);
+
+	void ExportCollectionToSkeletalMesh(
+		TObjectPtr<USkeletalMesh> ExportMesh,
+		const FManagedArrayCollection& Collection,
+		ENaniteShapePreservation ShapePreservation,
+		bool bBuildNaniteAssemblies
+	);
+
+	void AttachProceduralVegetationLink(UObject* InExportedMesh, const TObjectPtr<UProceduralVegetation>& InProceduralVegetation);
+	
+	void AssignBoneIDsToFoliage(const TArray<Facades::FBoneNode>& Bones, FManagedArrayCollection& Collection);
+
+	void RemoveUnwantedSkinWeights(TObjectPtr<UDynamicMesh> DynamicMesh, const FName ProfileName, const TArray<int32>& VertexBoneIDs,
+		int32 FoliageVertexStart, const TArray<Facades::FBoneNode>& BoneNodes);
+
+	void ConvertToDefaultSkeletalMesh(USkeletalMesh* SkeletalMesh,
+		FDynamicMesh3* Mesh,
+		TArray<TObjectPtr<UMaterialInterface>> Materials
+		);
+	
+}
+
+namespace PV::Export
+{
+	void PROCEDURALVEGETATION_API ExportCollectionAsMesh(
+		const TObjectPtr<UProceduralVegetation> InProceduralVegetation,
+		const FManagedArrayCollection& Collection,
+		const FManagedArrayCollection& FoliageCollection,
+		const FPVExportSettings& Settings
+	);
+}
diff --git a/Engine/Plugins/Experimental/ProceduralVegetationEditor/Source/ProceduralVegetation/Public/Helpers/PVUtilities.h b/Engine/Plugins/Experimental/ProceduralVegetationEditor/Source/ProceduralVegetation/Public/Helpers/PVUtilities.h
new file mode 100644
--- /dev/null
+++ b/Engine/Plugins/Experimental/ProceduralVegetationEditor/Source/ProceduralVegetation/Public/Helpers/PVUtilities.h
@@ -0,0 +1,17 @@
+﻿// Copyright Epic Games, Inc. All Rights Reserved.
+#pragma once
+#include "Facades/PVRenderingFacade.h"
+
+class UInstancedStaticMeshComponent;
+class UMeshComponent;
+class FFoliageFacade;
+class AActor;
+
+namespace PV::Utilities
+{
+	DECLARE_DELEGATE_OneParam(FFoliageComponentCreatedCallback, UMeshComponent*);
+	DECLARE_DELEGATE_OneParam(FFoliageInstanceAddedCallback, int32);
+	
+	void PROCEDURALVEGETATION_API AddFoliageInstancesToActor(const Facades::FFoliageFacade& InFoliageFacade, AActor* InActor, TMap<FString, TObjectPtr<UMeshComponent>>& OutInstancedComponentMap);
+	void PROCEDURALVEGETATION_API AddFoliageInstances(const Facades::FFoliageFacade& InFoliageFacade, const UObject* InParent, FFoliageComponentCreatedCallback InCallback, TMap<FString, TObjectPtr<UMeshComponent>>& OutInstancedComponentMap);
+}
diff --git a/Engine/Plugins/Experimental/ProceduralVegetationEditor/Source/ProceduralVegetation/Public/MegaSeedDataAsset.h b/Engine/Plugins/Experimental/ProceduralVegetationEditor/Source/ProceduralVegetation/Public/MegaSeedDataAsset.h
new file mode 100644
--- /dev/null
+++ b/Engine/Plugins/Experimental/ProceduralVegetationEditor/Source/ProceduralVegetation/Public/MegaSeedDataAsset.h
@@ -0,0 +1,22 @@
+// Copyright Epic Games, Inc. All Rights Reserved.
+
+#pragma once
+
+#include "Engine/DataAsset.h"
+
+#include "GeometryCollection/ManagedArrayCollection.h"
+
+#include "MegaSeedDataAsset.generated.h"
+
+UCLASS(BlueprintType)
+class PROCEDURALVEGETATION_API UMegaSeedDataAsset : public UDataAsset
+{
+	GENERATED_BODY()
+
+public:
+	UPROPERTY(EditAnywhere, Category="MegaSeed")
+	TMap<FString, FManagedArrayCollection> Variants;
+
+public:
+	void LoadFromVariantFiles(const TMap<FString, FString>& InVariantFiles);
+};
diff --git a/Engine/Plugins/Experimental/ProceduralVegetationEditor/Source/ProceduralVegetation/Public/Nodes/PVPCGBaseSettings.h b/Engine/Plugins/Experimental/ProceduralVegetationEditor/Source/ProceduralVegetation/Public/Nodes/PVPCGBaseSettings.h
new file mode 100644
--- /dev/null
+++ b/Engine/Plugins/Experimental/ProceduralVegetationEditor/Source/ProceduralVegetation/Public/Nodes/PVPCGBaseSettings.h
@@ -0,0 +1,32 @@
+﻿// Copyright Epic Games, Inc. All Rights Reserved.
+#pragma once
+
+#include "CoreMinimal.h"
+#include "PCGSettings.h"
+#include "PVPCGData.h"
+#include "PVPCGBaseSettings.generated.h"
+
+UCLASS(BlueprintType, Abstract, HideCategories=(Debug), ClassGroup = (Procedural))
+class UPVPCGBaseSettings : public UPCGSettings
+{
+	GENERATED_BODY()
+
+public:
+#if WITH_EDITOR
+	virtual EPCGSettingsType GetType() const override { return EPCGSettingsType::Spatial; }
+	virtual TArray<EPVRenderType> GetRenderType() const { return TArray<EPVRenderType>(); }
+#endif
+
+protected:
+	virtual TArray<FPCGPinProperties> InputPinProperties() const override;
+	virtual TArray<FPCGPinProperties> OutputPinProperties() const override;
+	virtual FPCGElementPtr CreateElement() const override;
+};
+
+class FPVPCGBaseElement : public IPCGElement
+{
+protected:
+	virtual void PostExecuteInternal(FPCGContext* InContext) const override;
+	virtual bool ExecuteInternal(FPCGContext* Context) const override;
+};
+
diff --git a/Engine/Plugins/Experimental/ProceduralVegetationEditor/Source/ProceduralVegetation/Public/Nodes/PVPCGOutput.h b/Engine/Plugins/Experimental/ProceduralVegetationEditor/Source/ProceduralVegetation/Public/Nodes/PVPCGOutput.h
new file mode 100644
--- /dev/null
+++ b/Engine/Plugins/Experimental/ProceduralVegetationEditor/Source/ProceduralVegetation/Public/Nodes/PVPCGOutput.h
@@ -0,0 +1,37 @@
+﻿// Copyright Epic Games, Inc. All Rights Reserved.
+#pragma once
+
+#include "CoreMinimal.h"
+#include "Helpers/PVExportHelper.h"
+#include "Nodes/PVPCGBaseSettings.h"
+#include "PVPCGOutput.generated.h"
+
+UCLASS(ClassGroup = (Procedural))
+class UPVPCGOutputSettings : public UPVPCGBaseSettings
+{
+	GENERATED_BODY()
+
+public:
+	UPVPCGOutputSettings();
+
+#if WITH_EDITOR
+	virtual FName GetDefaultNodeName() const override { return FName(TEXT("ProceduralVegetationOutput")); }
+	virtual FText GetDefaultNodeTitle() const override { return NSLOCTEXT("PVPCGOutput", "NodeTitle", "ProceduralVegetation Output"); }
+	virtual FText GetNodeTooltipText() const override { return NSLOCTEXT("PVPCGOutput", "NodeTooltip", "Output node for ProceduralVegetation"); }
+	virtual EPCGSettingsType GetType() const override { return EPCGSettingsType::InputOutput; }
+#endif
+
+protected:
+	
+	virtual TArray<FPCGPinProperties> OutputPinProperties() const override;
+	virtual FPCGElementPtr CreateElement() const override;
+	
+public:
+	UPROPERTY(EditAnywhere, Category = ExportSettings, meta = (ShowOnlyInnerProperties))
+	FPVExportSettings ExportSettings;
+};
+
+class FPVPCGOutputElement : public FPVPCGBaseElement
+{
+	virtual bool ExecuteInternal(FPCGContext* Context) const override;
+};
\ No newline at end of file
diff --git a/Engine/Plugins/Experimental/ProceduralVegetationEditor/Source/ProceduralVegetation/Public/PVPCGData.h b/Engine/Plugins/Experimental/ProceduralVegetationEditor/Source/ProceduralVegetation/Public/PVPCGData.h
new file mode 100644
--- /dev/null
+++ b/Engine/Plugins/Experimental/ProceduralVegetationEditor/Source/ProceduralVegetation/Public/PVPCGData.h
@@ -0,0 +1,78 @@
+// Copyright Epic Games, Inc. All Rights Reserved.
+
+#pragma once
+
+#include "Data/PCGSpatialData.h"
+#include "GeometryCollection/ManagedArrayCollection.h"
+#include "PVPCGData.generated.h"
+
+struct FPCGContext;
+
+UENUM()
+enum class EPVRenderType : uint8
+{
+	None 		= 0,
+	Skeleton 	= 1 << 0,
+	Mesh 		= 1 << 1,
+	Foliage 	= 1 << 2,
+	Bones 		= 1 << 3,
+
+	// Add new renter type above
+	Count	
+};
+
+ENUM_CLASS_FLAGS(EPVRenderType)
+
+UCLASS(BlueprintType, ClassGroup = (Procedural))
+class PROCEDURALVEGETATION_API UPVPCGData : public UPCGSpatialData
+{
+	GENERATED_BODY()
+
+public:
+	UPVPCGData(const FObjectInitializer& ObjectInitializer);
+	
+	void Initialize(FManagedArrayCollection&& InCollection, bool bCanTakeOwnership = false);
+	
+	// ~Begin UPCGData interface
+	virtual EPCGDataType GetDataType() const override { return EPCGDataType::Other; }
+	virtual void AddToCrc(FArchiveCrc32& Ar, bool bFullDataCrc) const override;
+	// ~End UPCGData interface
+
+	// ~Begin UPCGSpatialData interface
+	virtual int GetDimension() const override { return 3; }
+	virtual FBox GetBounds() const override;
+	virtual bool SamplePoint(const FTransform& Transform, const FBox& Bounds, FPCGPoint& OutPoint, UPCGMetadata* OutMetadata) const override;
+	//~End UPCGSpatialData interface
+
+	const FManagedArrayCollection& GetCollection() const { return Collection; }
+
+#if WITH_EDITOR
+	const TArray<EPVRenderType>& GetRenderType() const { return RenderType; }
+	void SetRenderType(TArray<EPVRenderType> InRenderType) const { RenderType = InRenderType; }
+#endif
+	
+protected:
+
+	FManagedArrayCollection Collection;
+
+#if WITH_EDITOR
+	mutable TArray<EPVRenderType> RenderType;
+#endif
+	
+	// ~Begin UPCGData interface
+	virtual bool SupportsFullDataCrc() const override { return true; }
+	// ~End UPCGData interface
+
+	// ~Begin UPCGSpatialData interface
+	virtual UPCGSpatialData* CopyInternal(FPCGContext* Context) const override;
+public:
+	virtual const UPCGPointData* ToPointData(FPCGContext* Context, const FBox& InBounds) const override;
+	virtual const UPCGPointArrayData* ToPointArrayData(FPCGContext* Context, const FBox& InBounds) const override;
+	//~End UPCGSpatialData interface
+	
+private:
+	const UPCGBasePointData* ToBasePointData(FPCGContext* Context, const FBox& InBounds, TSubclassOf<UPCGBasePointData> PointDataClass) const;
+protected:
+	
+	mutable FBox CachedBounds = FBox(EForceInit::ForceInit);
+};
diff --git a/Engine/Plugins/Experimental/ProceduralVegetationEditor/Source/ProceduralVegetation/Public/ProceduralVegetation.h b/Engine/Plugins/Experimental/ProceduralVegetationEditor/Source/ProceduralVegetation/Public/ProceduralVegetation.h
new file mode 100644
--- /dev/null
+++ b/Engine/Plugins/Experimental/ProceduralVegetationEditor/Source/ProceduralVegetation/Public/ProceduralVegetation.h
@@ -0,0 +1,47 @@
+// Copyright Epic Games, Inc. All Rights Reserved.
+
+#pragma once
+#include "CoreMinimal.h"
+#include "PCGGraph.h"
+#include "ProceduralVegetation.generated.h"
+
+class UMegaSeedDataAsset;
+
+UCLASS(MinimalAPI)
+class UProceduralVegetationGraph : public UPCGGraph
+{
+	GENERATED_BODY()
+
+public:
+	UProceduralVegetationGraph(const FObjectInitializer& ObjectInitializer)
+		: UPCGGraph(ObjectInitializer)
+	{
+#if WITH_EDITOR
+		// Hides input and output nodes
+		SetHiddenFlagInputNode(/*bHidden=*/true);
+		SetHiddenFlagOutputNode(/*bHidden=*/true);
+#endif
+	}
+};
+
+/**
+ * Asset type for procedural plant generation
+ */
+UCLASS()
+class PROCEDURALVEGETATION_API UProceduralVegetation : public UObject
+{
+	GENERATED_BODY()
+
+	UPROPERTY()
+	TObjectPtr<UProceduralVegetationGraph> Graph;
+
+public:
+	TObjectPtr<UProceduralVegetationGraph>& GetGraph() { return Graph; }
+	const TObjectPtr<UProceduralVegetationGraph>& GetGraph() const { return Graph; }
+
+	void CreateGraph(const UProceduralVegetationGraph* InGraph = nullptr);
+
+#if WITH_EDITOR
+	void CreateGraphFromMegaSeed(const TObjectPtr<UMegaSeedDataAsset> MegaSeedAsset);
+#endif
+};
diff --git a/Engine/Plugins/Experimental/ProceduralVegetationEditor/Source/ProceduralVegetationEditor/Private/Importer/MegaSeedImportPipeline.cpp b/Engine/Plugins/Experimental/ProceduralVegetationEditor/Source/ProceduralVegetationEditor/Private/Importer/MegaSeedImportPipeline.cpp
new file mode 100644
--- /dev/null
+++ b/Engine/Plugins/Experimental/ProceduralVegetationEditor/Source/ProceduralVegetationEditor/Private/Importer/MegaSeedImportPipeline.cpp
@@ -0,0 +1,234 @@
+// Copyright Epic Games, Inc. All Rights Reserved.
+
+#include "MegaSeedImportPipeline.h"
+
+#include "InterchangeMaterialFactoryNode.h"
+#include "InterchangeStaticMeshFactoryNode.h"
+#include "InterchangeTextureFactoryNode.h"
+#include "RenderUtils.h"
+
+#include "Serialization/JsonSerializer.h"
+
+UMegaSeedImportPipeline::UMegaSeedImportPipeline()
+{}
+
+void UMegaSeedImportPipeline::ExecutePipeline(
+	UInterchangeBaseNodeContainer* NodeContainer,
+	const TArray<UInterchangeSourceData*>& SourceDatas,
+	const FString& ContentBasePath
+)
+{
+	Super::ExecutePipeline(NodeContainer, SourceDatas, ContentBasePath);
+
+	BaseNodeContainer = NodeContainer;
+
+	const UInterchangeSourceData* const* GltfSourceData = SourceDatas.FindByPredicate(
+		[](const UInterchangeSourceData* SourceData)
+			{
+				return FPaths::GetExtension(SourceData->GetFilename()) == "gltf";
+			}
+	);
+	if (!GltfSourceData)
+	{
+		return;
+	}
+
+	if (!LoadGltfSource((*GltfSourceData)->GetFilename()))
+	{
+		return;
+	}
+
+	TextureFactoryNodes = GetNodesOfType<UInterchangeTextureFactoryNode>();
+	MaterialInstanceFactoryNodes = GetNodesOfType<UInterchangeMaterialInstanceFactoryNode>();
+	StaticMeshFactoryNodes = GetNodesOfType<UInterchangeStaticMeshFactoryNode>();
+
+	for (const TObjectPtr<UInterchangeStaticMeshFactoryNode>& StaticMeshFactoryNode : StaticMeshFactoryNodes)
+	{
+		static const FString AttrKey = "Nanite_Shape_Preservation";
+
+		static const bool bAllowNaniteFoliage = NaniteVoxelsSupported();
+
+		StaticMeshFactoryNode->SetAttribute(AttrKey,
+			static_cast<uint8>(bAllowNaniteFoliage ? ENaniteShapePreservation::Voxelize : ENaniteShapePreservation::PreserveArea));
+		StaticMeshFactoryNode->AddApplyAndFillDelegates<uint8>(
+			AttrKey,
+			StaticMeshFactoryNode->GetObjectClass(),
+			"NaniteSettings.ShapePreservation"
+		);
+	}
+
+	const TArray<TSharedPtr<FJsonValue>>& Images = GltfJson->GetArrayField(TEXT("images"));
+	for (const TSharedPtr<FJsonValue>& Image : Images)
+	{
+		const TSharedPtr<FJsonObject>& Texture = Image->AsObject();
+		const FString TextureName = Texture->GetStringField(TEXT("name"));
+		TObjectPtr<UInterchangeTextureFactoryNode> TextureFactoryNode = FindTextureFactoryNodeByName(TextureName);
+		if (TextureFactoryNode == nullptr)
+		{
+			return;
+		}
+
+		TextureFactoryNode->SetCustomVirtualTextureStreaming(true);
+
+		if (const TSharedPtr<FJsonObject>* TextureExtras; Texture->TryGetObjectField(TEXT("extras"), TextureExtras))
+		{
+			if (FString CompressionSettings; (*TextureExtras)->TryGetStringField(TEXT("compression"), CompressionSettings))
+			{
+				if (CompressionSettings == "mask")
+					TextureFactoryNode->SetCustomCompressionSettings(TC_Masks);
+				else if (CompressionSettings == "displacement" || CompressionSettings == "alpha")
+					TextureFactoryNode->SetCustomCompressionSettings(TC_Alpha);
+			}
+			if (FString TextureSlot; (*TextureExtras)->TryGetStringField(TEXT("textureSlot"), TextureSlot))
+			{
+				int32 MaterialIndex = 0;
+				(*TextureExtras)->TryGetNumberField(TEXT("materialIndex"), MaterialIndex);
+				if (const TSharedPtr<FJsonObject>* MaterialObject = GetMaterialAtIndex(MaterialIndex))
+				{
+					const FString MaterialName = (*MaterialObject)->GetStringField(TEXT("name"));
+					if (UInterchangeMaterialInstanceFactoryNode* MaterialInstanceFactoryNode = FindMaterialInstanceFactoryNodeByName(MaterialName))
+					{
+						const FString ParameterName = UInterchangeShaderPortsAPI::MakeInputValueKey(TextureSlot);
+						const FString TextureFactoryUid = TextureFactoryNode->GetUniqueID();
+						MaterialInstanceFactoryNode->AddStringAttribute(ParameterName, TextureFactoryUid);
+						MaterialInstanceFactoryNode->AddFactoryDependencyUid(TextureFactoryUid);
+					}
+				}
+			}
+		}
+	}
+
+
+	const TArray<TSharedPtr<FJsonValue>>& GltfMaterials = GltfJson->GetArrayField(TEXT("materials"));
+	for (const TSharedPtr<FJsonValue>& GltfMaterial : GltfMaterials)
+	{
+		const TSharedPtr<FJsonObject>& MaterialObject = GltfMaterial->AsObject();
+		const FString& MaterialName = MaterialObject->GetStringField(TEXT("name"));
+		if (const TSharedPtr<FJsonObject>* MaterialExtras; MaterialObject->TryGetObjectField(TEXT("extras"), MaterialExtras))
+		{
+			UInterchangeMaterialInstanceFactoryNode* const& MaterialInstanceFactoryNode = FindMaterialInstanceFactoryNodeByName(MaterialName);
+			if (FString MaterialType; (*MaterialExtras)->TryGetStringField(TEXT("type"), MaterialType))
+			{
+				if (MaterialType == "foliage")
+				{
+					const FString FoliageParentMaterial =
+						"/ProceduralVegetationEditor/SampleTrees/MSPresets/MS_Foliage_Material_FAB_VT_V1/MasterMaterials/MA_Foliage_Trees.MA_Foliage_Trees";
+					MaterialInstanceFactoryNode->SetCustomParent(FoliageParentMaterial);
+				}
+				else if (MaterialType == "trunk" || MaterialType == "bark")
+				{
+					const FString BarkParentMaterial =
+						"/ProceduralVegetationEditor/SampleTrees/MSPresets/MS_Foliage_Material_FAB_VT_V1/MasterMaterials/MA_Foliage_Trees.MA_Foliage_Trees";
+					MaterialInstanceFactoryNode->SetCustomParent(BarkParentMaterial);
+				}
+			}
+
+			if (const TSharedPtr<FJsonObject>* CustomProperties; (*MaterialExtras)->TryGetObjectField(TEXT("properties"), CustomProperties))
+			{
+				if (EBlendMode BlendMode; (*CustomProperties)->TryGetNumberField(TEXT("blend_mode"), reinterpret_cast<int&>(BlendMode)))
+				{
+					static const FString BlendModeOverrideAttrKey = "MegaSeed_BlendModeOverride";
+					static const FString BlendModeValueAttrKey = "MegaSeed_BlendModeValue";
+
+					MaterialInstanceFactoryNode->SetAttribute(BlendModeOverrideAttrKey, true);
+					MaterialInstanceFactoryNode->AddApplyAndFillDelegates<bool>(
+						BlendModeOverrideAttrKey,
+						MaterialInstanceFactoryNode->GetObjectClass(),
+						"BasePropertyOverrides.bOverride_BlendMode"
+					);
+					MaterialInstanceFactoryNode->SetAttribute(BlendModeValueAttrKey, static_cast<int>(BlendMode));
+					MaterialInstanceFactoryNode->AddApplyAndFillDelegates<int>(
+						BlendModeValueAttrKey,
+						MaterialInstanceFactoryNode->GetObjectClass(),
+						"BasePropertyOverrides.BlendMode"
+					);
+				}
+			}
+
+			const TArray<TSharedPtr<FJsonValue>>* CustomParams;
+			if (!(*MaterialExtras)->TryGetArrayField(TEXT("params"), CustomParams))
+			{
+				continue;
+			}
+
+			for (const TSharedPtr<FJsonValue>& Param : *CustomParams)
+			{
+				if (Param->Type != EJson::Object)
+					continue;
+
+				const TSharedPtr<FJsonObject>& ParamObject = Param->AsObject();
+				const FString Name = ParamObject->GetStringField(TEXT("Name"));
+				const TSharedPtr<FJsonValue> Value = ParamObject->TryGetField(TEXT("Value"));
+
+				const FString ParameterName = UInterchangeShaderPortsAPI::MakeInputValueKey(Name);
+				if (Value->Type == EJson::Boolean)
+				{
+					MaterialInstanceFactoryNode->AddBooleanAttribute(ParameterName, Value->AsBool());
+				}
+				else if (Value->Type == EJson::Number)
+				{
+					MaterialInstanceFactoryNode->AddFloatAttribute(ParameterName, Value->AsNumber());
+				}
+				else if (Value->Type == EJson::Array)
+				{
+					const TArray<TSharedPtr<FJsonValue>>& ArrayValue = Value->AsArray();
+					if (ArrayValue.Num() >= 4)
+					{
+						MaterialInstanceFactoryNode->AddLinearColorAttribute(
+							ParameterName,
+							FLinearColor(
+								ArrayValue[0]->AsNumber(),
+								ArrayValue[1]->AsNumber(),
+								ArrayValue[2]->AsNumber(),
+								ArrayValue[3]->AsNumber()
+							)
+						);
+					}
+				}
+			}
+		}
+	}
+}
+
+bool UMegaSeedImportPipeline::LoadGltfSource(const FString& SourceFile)
+{
+	if (FString GltfFileData; FFileHelper::LoadFileToString(GltfFileData, *SourceFile))
+	{
+		GltfJson = MakeShareable(new FJsonObject);
+		return FJsonSerializer::Deserialize(TJsonReaderFactory<TCHAR>::Create(GltfFileData), GltfJson);
+	}
+	return false;
+}
+
+const TSharedPtr<FJsonObject>* UMegaSeedImportPipeline::GetMaterialAtIndex(uint32 Index) const
+{
+	TSharedPtr<FJsonObject>* Material = nullptr;
+	const TArray<TSharedPtr<FJsonValue>> Materials = GltfJson->GetArrayField(TEXT("materials"));
+	if (Materials.IsValidIndex(Index))
+	{
+		Materials[Index]->TryGetObject(Material);
+	}
+	return Material;
+}
+
+TObjectPtr<UInterchangeTextureFactoryNode> UMegaSeedImportPipeline::FindTextureFactoryNodeByName(const FString& DisplayName) const
+{
+	const TObjectPtr<UInterchangeTextureFactoryNode>* FoundNode = TextureFactoryNodes.FindByPredicate(
+		[&DisplayName](const TObjectPtr<UInterchangeTextureFactoryNode>& Node)
+			{
+				return Node->GetDisplayLabel() == DisplayName;
+			}
+	);
+	return FoundNode ? *FoundNode : nullptr;
+}
+
+TObjectPtr<UInterchangeMaterialInstanceFactoryNode> UMegaSeedImportPipeline::FindMaterialInstanceFactoryNodeByName(const FString& DisplayName) const
+{
+	const TObjectPtr<UInterchangeMaterialInstanceFactoryNode>* FoundNode = MaterialInstanceFactoryNodes.FindByPredicate(
+		[&DisplayName](const TObjectPtr<UInterchangeMaterialInstanceFactoryNode>& Node)
+			{
+				return Node->GetDisplayLabel() == DisplayName;
+			}
+	);
+	return FoundNode ? *FoundNode : nullptr;
+}
diff --git a/Engine/Plugins/Experimental/ProceduralVegetationEditor/Source/ProceduralVegetationEditor/Private/Importer/MegaSeedImportPipeline.h b/Engine/Plugins/Experimental/ProceduralVegetationEditor/Source/ProceduralVegetationEditor/Private/Importer/MegaSeedImportPipeline.h
new file mode 100644
--- /dev/null
+++ b/Engine/Plugins/Experimental/ProceduralVegetationEditor/Source/ProceduralVegetationEditor/Private/Importer/MegaSeedImportPipeline.h
@@ -0,0 +1,78 @@
+// Copyright Epic Games, Inc. All Rights Reserved.
+
+#pragma once
+
+#include "InterchangePipelineBase.h"
+
+#include "Dom/JsonObject.h"
+
+#include "MegaSeedImportPipeline.generated.h"
+
+class UInterchangeTextureFactoryNode;
+class UInterchangeMaterialInstanceFactoryNode;
+class UInterchangeStaticMeshFactoryNode;
+
+UCLASS(BlueprintType)
+class UMegaSeedImportPipeline : public UInterchangePipelineBase
+{
+	GENERATED_BODY()
+
+public:
+	UMegaSeedImportPipeline();
+
+protected:
+	virtual void ExecutePipeline(
+		UInterchangeBaseNodeContainer* NodeContainer,
+		const TArray<UInterchangeSourceData*>& SourceDatas,
+		const FString& ContentBasePath
+	) override;
+	
+	virtual bool CanExecuteOnAnyThread(EInterchangePipelineTask PipelineTask) override { return true; }
+
+private:
+	bool LoadGltfSource(const FString& SourceFile);
+
+	const TSharedPtr<FJsonObject>* GetMaterialAtIndex(uint32 Index) const;
+
+	TObjectPtr<UInterchangeTextureFactoryNode> FindTextureFactoryNodeByName(const FString& DisplayName) const;
+	TObjectPtr<UInterchangeMaterialInstanceFactoryNode> FindMaterialInstanceFactoryNodeByName(const FString& DisplayName) const;
+
+	template <class NodeType>
+	TArray<TObjectPtr<NodeType>> GetNodesOfType()
+	{
+		TArray<TObjectPtr<NodeType>> Nodes;
+		BaseNodeContainer->IterateNodesOfType<NodeType>(
+			[&Nodes](const FString& Uid, NodeType* Node)
+			{
+				Nodes.Add(Node);
+			}
+		);
+		return Nodes;
+	}
+
+	template <class NodeType>
+	NodeType* FindNodeOfTypeByName(const FString& DisplayName)
+	{
+		NodeType* FoundNode = nullptr;
+		BaseNodeContainer->IterateNodesOfType<NodeType>(
+			[&](const FString& Uid, NodeType* Node)
+			{
+				if (FoundNode)
+					return;
+				if (Node->GetDisplayLabel() == DisplayName)
+					FoundNode = Node;
+			}
+		);
+		return FoundNode;
+	}
+
+private:
+	UPROPERTY()
+	TObjectPtr<UInterchangeBaseNodeContainer> BaseNodeContainer = nullptr;
+
+	TSharedPtr<FJsonObject> GltfJson = nullptr;
+
+	TArray<TObjectPtr<UInterchangeTextureFactoryNode>> TextureFactoryNodes;
+	TArray<TObjectPtr<UInterchangeMaterialInstanceFactoryNode>> MaterialInstanceFactoryNodes;
+	TArray<TObjectPtr<UInterchangeStaticMeshFactoryNode>> StaticMeshFactoryNodes;
+};
diff --git a/Engine/Plugins/Experimental/ProceduralVegetationEditor/Source/ProceduralVegetationEditor/Private/Importer/MegaSeedImportWorkflow.cpp b/Engine/Plugins/Experimental/ProceduralVegetationEditor/Source/ProceduralVegetationEditor/Private/Importer/MegaSeedImportWorkflow.cpp
new file mode 100644
--- /dev/null
+++ b/Engine/Plugins/Experimental/ProceduralVegetationEditor/Source/ProceduralVegetationEditor/Private/Importer/MegaSeedImportWorkflow.cpp
@@ -0,0 +1,214 @@
+// Copyright Epic Games, Inc. All Rights Reserved.
+
+#include "MegaSeedImportWorkflow.h"
+#include "FabDownloader.h"
+#include "MegaSeedImporter.h"
+#include "MegaSeedDataAsset.h"
+#include "AssetRegistry/IAssetRegistry.h"
+#include "Async/Async.h"
+#include "Dom/JsonObject.h"
+#include "Dom/JsonValue.h"
+#include "Facades/PVBranchFacade.h"
+#include "Facades/PVFoliageFacade.h"
+#include "HAL/PlatformFileManager.h"
+#include "Misc/FileHelper.h"
+#include "ProceduralVegetation.h"
+#include "Serialization/JsonReader.h"
+#include "Serialization/JsonSerializer.h"
+#include "Utilities/AssetUtils.h"
+#include "Utilities/FabAssetsCache.h"
+
+FMegaSeedImportWorkflow::FMegaSeedImportWorkflow(const FString& InAssetId, const FString& InAssetName, const FString& InDownloadUrl)
+	: IFabWorkflow(InAssetId, InAssetName, InDownloadUrl)
+{}
+
+void FMegaSeedImportWorkflow::Execute()
+{
+	DownloadContent();
+}
+
+void FMegaSeedImportWorkflow::DownloadContent()
+{
+	const FString DownloadLocation = FFabAssetsCache::GetCacheLocation() / AssetId;
+
+	DownloadRequest = MakeShared<FFabDownloadRequest>(AssetId, DownloadUrl, DownloadLocation, EFabDownloadType::HTTP);
+	DownloadRequest->OnDownloadProgress().AddRaw(this, &FMegaSeedImportWorkflow::OnContentDownloadProgress);
+	DownloadRequest->OnDownloadComplete().AddRaw(this, &FMegaSeedImportWorkflow::OnContentDownloadComplete);
+	DownloadRequest->ExecuteRequest();
+
+	CreateDownloadNotification();
+}
+
+void FMegaSeedImportWorkflow::ImportContent(const TArray<FString>& SourceFiles)
+{
+	CreateImportNotification();
+
+	FString AssetImportFolder = AssetName.IsEmpty() ? AssetId : AssetName;
+	FAssetUtils::SanitizeFolderName(AssetImportFolder);
+	ImportLocation = "/Game/Fab/MegaSeed" / AssetImportFolder;
+
+	const FString MegaSeedFilePath = SourceFiles[0];
+	const FString MegaSeedDownloadFolder = FPaths::GetPath(MegaSeedFilePath);
+	TSharedPtr<FJsonObject> MegaSeedData = nullptr;
+
+	bool bSuccess = false;
+
+	FString MegaSeedFile;
+	bSuccess = FFileHelper::LoadFileToString(MegaSeedFile, *MegaSeedFilePath);
+	if (!bSuccess)
+	{
+		ExpireImportNotification(false);
+		CancelWorkflow();
+	}
+
+	MegaSeedData = MakeShareable(new FJsonObject);
+	bSuccess = FJsonSerializer::Deserialize(TJsonReaderFactory<TCHAR>::Create(MegaSeedFile), MegaSeedData);
+	if (!bSuccess)
+	{
+		ExpireImportNotification(false);
+		CancelWorkflow();
+	}
+
+	const TSharedPtr<FJsonObject>& MegaSeedResources = MegaSeedData->GetObjectField(TEXT("resources"));
+
+	const TArray<TSharedPtr<FJsonValue>>& MegaSeedVariations = MegaSeedResources->GetArrayField(TEXT("variations"));
+	const TArray<TSharedPtr<FJsonValue>>& PreBakedMeshes = MegaSeedResources->GetArrayField(TEXT("prebaked_meshes"));
+
+	TArray<FString> PreBakedMeshesPaths;
+	for (const TSharedPtr<FJsonValue>& PreBakedMeshData : PreBakedMeshes)
+	{
+		const TSharedPtr<FJsonObject>& PreBakedMesh = PreBakedMeshData->AsObject();
+		const FString MeshUri = MegaSeedDownloadFolder / PreBakedMesh->GetStringField(TEXT("uri"));
+		PreBakedMeshesPaths.Add(MeshUri);
+	}
+
+	const FString MegaSeedName = AssetImportFolder + "_MegaSeed";
+	const FString MegaSeedPackageName = ImportLocation / MegaSeedName;
+
+	UPackage* MegaSeedPackage = CreatePackage(*MegaSeedPackageName);
+	UMegaSeedDataAsset* MegaSeedAsset = NewObject<UMegaSeedDataAsset>(MegaSeedPackage, *MegaSeedName, RF_Public | RF_Standalone);
+
+	TMap<FString, FString> VariantFiles;
+	for (const auto& Variation : MegaSeedVariations)
+	{
+		const TSharedPtr<FJsonObject>& VariationObject = Variation->AsObject();
+		const FString& VariationUri = VariationObject->GetStringField(TEXT("uri"));
+		const FString& VariationName = VariationObject->GetStringField(TEXT("name"));
+		const FString VariationPath = MegaSeedDownloadFolder / VariationUri;
+		VariantFiles.Add(VariationName, VariationPath);
+	}
+	MegaSeedAsset->LoadFromVariantFiles(VariantFiles);
+
+	IAssetRegistry::Get()->AssetCreated(MegaSeedAsset);
+	MegaSeedPackage->SetDirtyFlag(true);
+
+	const FString ProceduralVegetationName = AssetImportFolder + "_Graph";
+	const FString ProceduralVegetationPackageName = ImportLocation / ProceduralVegetationName;
+	UPackage* ProceduralVegetationPackage = CreatePackage(*ProceduralVegetationPackageName);
+	UProceduralVegetation* ProceduralVegetation = NewObject<UProceduralVegetation>(ProceduralVegetationPackage, *ProceduralVegetationName, RF_Public | RF_Standalone);
+
+	IAssetRegistry::Get()->AssetCreated(ProceduralVegetation);
+
+	const TArray<TSharedPtr<FJsonValue>>& MegaSeedTrunkMaterials = MegaSeedResources->GetArrayField(TEXT("trunkMaterials"));
+	// TODO: Add option to have materials per generation
+	const FString& TrunkMaterialName = MegaSeedTrunkMaterials[0]->AsObject()->GetStringField(TEXT("name"));
+
+	const TArray<TSharedPtr<FJsonValue>>& MegaSeedFoliage = MegaSeedResources->GetArrayField(TEXT("foliage"));
+	for (const TSharedPtr<FJsonValue>& Foliage : MegaSeedFoliage)
+	{
+		const TSharedPtr<FJsonObject>& FoliageObject = Foliage->AsObject();
+		const FString& FoliageUri = FoliageObject->GetStringField(TEXT("uri"));
+		const FString FoliagePath = MegaSeedDownloadFolder / FoliageUri;
+
+		const FString FoliageImportLocation = ImportLocation / "Foliage";
+		for (auto& [VariantName, VariantCollection] : MegaSeedAsset->Variants)
+		{
+			PV::Facades::FFoliageFacade FoliageFacade(VariantCollection);
+			FoliageFacade.SetFoliagePath(FoliageImportLocation);
+
+			PV::Facades::FBranchFacade BranchFacade(VariantCollection);
+			BranchFacade.SetTrunkMaterialPath(FoliageImportLocation / TrunkMaterialName + '.' + TrunkMaterialName);
+		}
+
+		FMegaSeedImporter::ImportFoliage(
+			FoliagePath,
+			FoliageImportLocation,
+			[this, ProceduralVegetation, MegaSeedAsset, PreBakedMeshesPaths](const TArray<UObject*>& InImportedObjects)
+				{
+					AsyncTask(
+						ENamedThreads::GameThread,
+						[this, ProceduralVegetation, MegaSeedAsset]()
+							{
+#if WITH_EDITOR
+								ProceduralVegetation->CreateGraphFromMegaSeed(MegaSeedAsset);
+								ProceduralVegetation->MarkPackageDirty();
+#endif
+								ExpireImportNotification(true);
+								CompleteWorkflow();
+							}
+					);
+					for (auto& [VariantName, VariantCollection] : MegaSeedAsset->Variants)
+					{
+						PV::Facades::FFoliageFacade FoliageFacade(VariantCollection);
+						for (int32 FoliageIndex = 0; FoliageIndex < FoliageFacade.NumFoliageNames(); FoliageIndex++)
+						{
+							UObject* const* FoundFoliageMesh = InImportedObjects.FindByPredicate([&](const UObject* ImportedObject)
+								{
+									return ImportedObject->GetName().Contains(FoliageFacade.GetFoliageName(FoliageIndex));
+								});
+							if (FoundFoliageMesh)
+							{
+								FoliageFacade.SetFoliageName(FoliageIndex, (*FoundFoliageMesh)->GetPathName());
+							}
+						}
+					}
+
+					TArray<FString> ScanFiles;
+					for (const FString& MeshUri : PreBakedMeshesPaths)
+					{
+						const FString DestPath = FPaths::ConvertRelativePathToFull(FPaths::ProjectContentDir()) / ImportLocation.RightChop(6) /
+							FPaths::GetCleanFilename(MeshUri);
+						IFileManager::Get().Copy(*DestPath, *MeshUri);
+						ScanFiles.Add(DestPath);
+					}
+
+					IAssetRegistry::Get()->ScanFilesSynchronous(ScanFiles, true);
+				}
+		);
+	}
+}
+
+void FMegaSeedImportWorkflow::OnContentDownloadProgress(const FFabDownloadRequest* Request, const FFabDownloadStats& DownloadStats)
+{
+	SetDownloadNotificationProgress(DownloadStats.PercentComplete);
+}
+
+void FMegaSeedImportWorkflow::OnContentDownloadComplete(const FFabDownloadRequest* Request, const FFabDownloadStats& DownloadStats)
+{
+	if (!DownloadStats.bIsSuccess)
+	{
+		UE_LOG(LogTemp, Error, TEXT("Failed to download FAB Asset %s"), *AssetName);
+		ExpireDownloadNotification(false);
+		CancelWorkflow();
+		return;
+	}
+
+	const FString& ZipArchive = DownloadStats.DownloadedFiles[0];
+	const FString ExtractLocation = FPaths::GetBaseFilename(ZipArchive, false) + "_extracted";
+	if (!FAssetUtils::Unzip(ZipArchive, ExtractLocation))
+	{
+		UE_LOG(LogTemp, Error, TEXT("Failed to unzip MegaSeed Asset %s"), *AssetId);
+		ExpireDownloadNotification(false);
+		CancelWorkflow();
+		return;
+	}
+
+	IPlatformFile& FileManager = FPlatformFileManager::Get().GetPlatformFile();
+
+	TArray<FString> ImportFiles;
+	FileManager.FindFiles(ImportFiles, *ExtractLocation, TEXT(".json"));
+	ImportFiles = ImportFiles.FilterByPredicate([](const FString& File) { return File.EndsWith("_recipe.json"); });
+
+	ExpireDownloadNotification(true);
+	ImportContent(ImportFiles);
+}
diff --git a/Engine/Plugins/Experimental/ProceduralVegetationEditor/Source/ProceduralVegetationEditor/Private/Importer/MegaSeedImportWorkflow.h b/Engine/Plugins/Experimental/ProceduralVegetationEditor/Source/ProceduralVegetationEditor/Private/Importer/MegaSeedImportWorkflow.h
new file mode 100644
--- /dev/null
+++ b/Engine/Plugins/Experimental/ProceduralVegetationEditor/Source/ProceduralVegetationEditor/Private/Importer/MegaSeedImportWorkflow.h
@@ -0,0 +1,23 @@
+// Copyright Epic Games, Inc. All Rights Reserved.
+
+#pragma once
+
+#include "Workflows/FabWorkflow.h"
+
+class FMegaSeedImportWorkflow : public IFabWorkflow
+{
+public:
+	FMegaSeedImportWorkflow(const FString& InAssetId, const FString& InAssetName, const FString& InDownloadUrl);
+
+	virtual void Execute() override;
+
+protected:
+	virtual void DownloadContent() override;
+	virtual void ImportContent(const TArray<FString>& SourceFiles) override;
+
+	virtual void OnContentDownloadProgress(const FFabDownloadRequest* Request, const FFabDownloadStats& DownloadStats) override;
+	virtual void OnContentDownloadComplete(const FFabDownloadRequest* Request, const FFabDownloadStats& DownloadStats) override;
+
+private:
+	TSharedPtr<FFabDownloadRequest> DownloadRequest;
+};
diff --git a/Engine/Plugins/Experimental/ProceduralVegetationEditor/Source/ProceduralVegetationEditor/Private/Importer/MegaSeedImportWorkflowFactory.cpp b/Engine/Plugins/Experimental/ProceduralVegetationEditor/Source/ProceduralVegetationEditor/Private/Importer/MegaSeedImportWorkflowFactory.cpp
new file mode 100644
--- /dev/null
+++ b/Engine/Plugins/Experimental/ProceduralVegetationEditor/Source/ProceduralVegetationEditor/Private/Importer/MegaSeedImportWorkflowFactory.cpp
@@ -0,0 +1,16 @@
+// Copyright Epic Games, Inc. All Rights Reserved.
+
+#include "MegaSeedImportWorkflowFactory.h"
+
+#include "MegaSeedImportWorkflow.h"
+
+const TArray<FString>& FMegaSeedImportWorkflowFactory::GetImportAssetTypes()
+{
+	static TArray<FString> ImportAssetTypes{"megaseed"};
+	return ImportAssetTypes;
+}
+
+TSharedPtr<IFabWorkflow> FMegaSeedImportWorkflowFactory::Create(const FFabAssetMetadata& InImportAssetMetadata, const FString& InDownloadUrl)
+{
+	return MakeShared<FMegaSeedImportWorkflow>(InImportAssetMetadata.AssetId, InImportAssetMetadata.AssetName, InDownloadUrl);
+}
diff --git a/Engine/Plugins/Experimental/ProceduralVegetationEditor/Source/ProceduralVegetationEditor/Private/Importer/MegaSeedImportWorkflowFactory.h b/Engine/Plugins/Experimental/ProceduralVegetationEditor/Source/ProceduralVegetationEditor/Private/Importer/MegaSeedImportWorkflowFactory.h
new file mode 100644
--- /dev/null
+++ b/Engine/Plugins/Experimental/ProceduralVegetationEditor/Source/ProceduralVegetationEditor/Private/Importer/MegaSeedImportWorkflowFactory.h
@@ -0,0 +1,11 @@
+// Copyright Epic Games, Inc. All Rights Reserved.
+
+#pragma once
+#include "Workflows/FabWorkflowFactory.h"
+
+class FMegaSeedImportWorkflowFactory: public IFabWorkflowFactory
+{
+public:
+	virtual const TArray<FString>& GetImportAssetTypes() override;
+	virtual TSharedPtr<IFabWorkflow> Create(const FFabAssetMetadata& InImportAssetMetadata, const FString& InDownloadUrl = "") override;
+};
diff --git a/Engine/Plugins/Experimental/ProceduralVegetationEditor/Source/ProceduralVegetationEditor/Private/Importer/MegaSeedImporter.cpp b/Engine/Plugins/Experimental/ProceduralVegetationEditor/Source/ProceduralVegetationEditor/Private/Importer/MegaSeedImporter.cpp
new file mode 100644
--- /dev/null
+++ b/Engine/Plugins/Experimental/ProceduralVegetationEditor/Source/ProceduralVegetationEditor/Private/Importer/MegaSeedImporter.cpp
@@ -0,0 +1,141 @@
+// Copyright Epic Games, Inc. All Rights Reserved.
+
+#include "MegaSeedImporter.h"
+
+#include "AssetImportTask.h"
+#include "IAssetTools.h"
+#include "InterchangeGenericAssetsPipeline.h"
+#include "InterchangeGenericMaterialPipeline.h"
+#include "InterchangeGenericMeshPipeline.h"
+#include "InterchangeGenericTexturePipeline.h"
+#include "InterchangeProjectSettings.h"
+#include "MegaSeedImportPipeline.h"
+
+#include "Async/Async.h"
+
+UObject* FMegaSeedImporter::GetImportOptions(const FString& SourceFile, UObject* const OptionsOuter)
+{
+	check(OptionsOuter);
+
+	const UInterchangeManager& InterchangeManager = UInterchangeManager::GetInterchangeManager();
+	const UInterchangeSourceData* SourceData = UInterchangeManager::CreateSourceData(SourceFile);
+	if (InterchangeManager.IsInterchangeImportEnabled() && InterchangeManager.CanTranslateSourceData(SourceData))
+	{
+		const FName PipelineStackName = FInterchangeProjectSettingsUtils::GetDefaultPipelineStackName(false, *SourceData);
+		const FInterchangeImportSettings& InterchangeImportSettings = FInterchangeProjectSettingsUtils::GetDefaultImportSettings(false);
+		if (const FInterchangePipelineStack* const PipelineStack = InterchangeImportSettings.PipelineStacks.Find(PipelineStackName))
+		{
+			const TArray<FSoftObjectPath>* Pipelines = &PipelineStack->Pipelines;
+			UE::Interchange::FScopedTranslator ScopedTranslator(SourceData);
+			for (const FInterchangeTranslatorPipelines& TranslatorPipelines : PipelineStack->PerTranslatorPipelines)
+			{
+				const UClass* TranslatorClass = TranslatorPipelines.Translator.LoadSynchronous();
+				if (ScopedTranslator.GetTranslator() && ScopedTranslator.GetTranslator()->IsA(TranslatorClass))
+				{
+					Pipelines = &TranslatorPipelines.Pipelines;
+					break;
+				}
+			}
+
+			UInterchangePipelineStackOverride* StackOverride = NewObject<UInterchangePipelineStackOverride>(OptionsOuter);
+
+			for (const FSoftObjectPath& Pipeline : *Pipelines)
+			{
+				UInterchangePipelineBase* const DefaultPipeline = Cast<UInterchangePipelineBase>(Pipeline.TryLoad());
+				if (!DefaultPipeline)
+				{
+					continue;
+				}
+				if (UInterchangePipelineBase* GeneratedPipeline = UE::Interchange::GeneratePipelineInstance(Pipeline))
+				{
+					GeneratedPipeline->TransferAdjustSettings(DefaultPipeline);
+					GeneratedPipeline->AddToRoot();
+					if (UInterchangeGenericAssetsPipeline* const GenericAssetsPipeline = Cast<UInterchangeGenericAssetsPipeline>(GeneratedPipeline))
+					{
+						GenericAssetsPipeline->bSceneNameSubFolder = false;
+						GenericAssetsPipeline->bAssetTypeSubFolders = false;
+						GenericAssetsPipeline->MeshPipeline->bImportStaticMeshes = true;
+						GenericAssetsPipeline->MeshPipeline->bImportSkeletalMeshes = true;
+						GenericAssetsPipeline->MeshPipeline->bCombineStaticMeshes = false;
+						GenericAssetsPipeline->MeshPipeline->SkeletalMeshImportContentType = EInterchangeSkeletalMeshContentType::All;
+						GenericAssetsPipeline->MeshPipeline->bGenerateLightmapUVs = true;
+						GenericAssetsPipeline->MeshPipeline->bBuildNanite = true;
+						GenericAssetsPipeline->MaterialPipeline->bImportMaterials = true;
+						GenericAssetsPipeline->MaterialPipeline->TexturePipeline->bImportTextures = true;
+
+						GenericAssetsPipeline->MaterialPipeline->MaterialImport = EInterchangeMaterialImportOption::ImportAsMaterialInstances;
+						GenericAssetsPipeline->CommonMeshesProperties->bRecomputeNormals = false;
+						GenericAssetsPipeline->CommonMeshesProperties->bComputeWeightedNormals = false;
+						GenericAssetsPipeline->CommonMeshesProperties->VertexColorImportOption = EInterchangeVertexColorImportOption::IVCIO_Replace;
+					}
+					if (UInterchangeGenericTexturePipeline* const GenericTexturePipeline = Cast<
+						UInterchangeGenericTexturePipeline>(GeneratedPipeline))
+					{
+						GenericTexturePipeline->bAllowNonPowerOfTwo = true;
+						GenericTexturePipeline->bDetectNormalMapTexture = true;
+					}
+
+					StackOverride->OverridePipelines.Add(GeneratedPipeline);
+				}
+			}
+
+			UMegaSeedImportPipeline* MegaSeedImportPipeline = NewObject<UMegaSeedImportPipeline>();
+			MegaSeedImportPipeline->AddToRoot();
+			StackOverride->OverridePipelines.Add(MegaSeedImportPipeline);
+
+			return StackOverride;
+		}
+		return nullptr;
+	}
+
+	return nullptr;
+}
+
+void FMegaSeedImporter::CleanImportOptions(UObject* const Options)
+{
+	if (const UInterchangePipelineStackOverride* const InterchangeOptions = Cast<UInterchangePipelineStackOverride>(Options))
+	{
+		for (const FSoftObjectPath& OverridePipeline : InterchangeOptions->OverridePipelines)
+		{
+			if (UObject* const LoadedPipeline = OverridePipeline.TryLoad())
+				LoadedPipeline->RemoveFromRoot();
+		}
+	}
+}
+
+void FMegaSeedImporter::ImportFoliage(const FString& Source, const FString& Destination, const TFunction<void(const TArray<UObject*>&)>& Callback)
+{
+	UAssetImportTask* MeshImportTask = NewObject<UAssetImportTask>();
+	MeshImportTask->AddToRoot();
+
+	MeshImportTask->bAutomated = true;
+	MeshImportTask->bSave = false;
+	MeshImportTask->bAsync = true;
+	MeshImportTask->Filename = Source;
+	
+	MeshImportTask->DestinationPath = Destination;
+	MeshImportTask->bReplaceExisting = true;
+	MeshImportTask->Options = GetImportOptions(Source, MeshImportTask);
+
+	IAssetTools::Get().ImportAssetTasks({MeshImportTask});
+
+	Async(
+		EAsyncExecution::Thread,
+		[MeshImportTask, Callback]()
+		{
+			TArray<UObject*> ImportedObjects;
+			if (MeshImportTask->AsyncResults.IsValid())
+			{
+				FPlatformProcess::ConditionalSleep([=]() { return MeshImportTask->IsAsyncImportComplete(); }, 0.25f);
+				ImportedObjects.Append(MeshImportTask->AsyncResults->GetImportedObjects());
+			}
+			else
+			{
+				ImportedObjects.Append(MeshImportTask->GetObjects());
+			}
+			Callback(ImportedObjects);
+			CleanImportOptions(MeshImportTask->Options);
+			MeshImportTask->RemoveFromRoot();
+		}
+	);
+}
diff --git a/Engine/Plugins/Experimental/ProceduralVegetationEditor/Source/ProceduralVegetationEditor/Private/Importer/MegaSeedImporter.h b/Engine/Plugins/Experimental/ProceduralVegetationEditor/Source/ProceduralVegetationEditor/Private/Importer/MegaSeedImporter.h
new file mode 100644
--- /dev/null
+++ b/Engine/Plugins/Experimental/ProceduralVegetationEditor/Source/ProceduralVegetationEditor/Private/Importer/MegaSeedImporter.h
@@ -0,0 +1,15 @@
+// Copyright Epic Games, Inc. All Rights Reserved.
+
+#pragma once
+
+#include "CoreMinimal.h"
+
+class FMegaSeedImporter
+{
+private:
+	static UObject* GetImportOptions(const FString& SourceFile, UObject* const OptionsOuter);
+	static void CleanImportOptions(UObject* const Options);
+
+public:
+	static void ImportFoliage(const FString& Source, const FString& Destination, const TFunction<void(const TArray<UObject*>&)>& Callback);
+};
