Rename MegaPlant to ProceduralVegetation

diff --git a/Engine/Plugins/Experimental/ProceduralVegetationEditor/Source/ProceduralVegetationEditor/Private/PVEditor.cpp b/Engine/Plugins/Experimental/ProceduralVegetationEditor/Source/ProceduralVegetationEditor/Private/PVEditor.cpp
new file mode 100644
--- /dev/null
+++ b/Engine/Plugins/Experimental/ProceduralVegetationEditor/Source/ProceduralVegetationEditor/Private/PVEditor.cpp
@@ -0,0 +1,423 @@
+﻿// Copyright Epic Games, Inc. All Rights Reserved.
+
+#include "PVEditor.h"
+#include "EngineAnalytics.h"
+#include "PVEditorCommands.h"
+#include "Editor/PCGEditorDefaultExecutionSource.h"
+#include "Subsystems/PCGEditorSubsystem.h"
+#include "PVEditorSchema.h"
+#include "ProceduralVegetation.h"
+#include "PVPCGData.h"
+#include "Nodes/PVPCGOutput.h"
+#include "ToolMenus.h"
+#include "ToolMenuSection.h"
+#include "Dataflow/DataflowCollectionSpreadSheetWidget.h"
+#include "PCGEditor/Private/Nodes/PCGEditorGraphNodeBase.h"
+#include "Widgets/SPVCreateDialog.h"
+#include "Widgets/SPVExportSelectionDialog.h"
+
+#if PROCEDURAL_VEGETATION_DEBUG_SHOW_COLLECTION_TAB
+#include "Widgets/Docking/SDockTab.h"
+#endif
+
+#define LOCTEXT_NAMESPACE "PVEditor"
+
+#if PROCEDURAL_VEGETATION_DEBUG_SHOW_COLLECTION_TAB
+const FName FPVEditor::CollectionSpreadSheetTabId (TEXT("PVEditor_CollectionSpreadSheet"));
+#endif
+
+FPVEditor::FPVEditor()
+{
+}
+
+void FPVEditor::Initialize(const EToolkitMode::Type InMode, const TSharedPtr<class IToolkitHost>& InToolkitHost, UProceduralVegetation* InProceduralVegetation)
+{
+	check(InProceduralVegetation);
+
+	ProceduralVegetationBeingEdited = InProceduralVegetation;
+	
+	UProceduralVegetationGraph* ProceduralVegetationGraph = InProceduralVegetation->GetGraph();
+
+	check(ProceduralVegetationGraph);
+
+	// Initialize widgets before calling base
+
+#if PROCEDURAL_VEGETATION_DEBUG_SHOW_COLLECTION_TAB
+	CollectionSpreadSheetWidget = CreateCollectionSpreadSheetWidget();
+#endif
+	
+	FPCGEditor::Initialize(InMode, InToolkitHost, ProceduralVegetationGraph);
+
+	AddEditingObject(InProceduralVegetation);
+
+	check(GEditor);
+	
+	UPCGEditorSubsystem* Subsystem = GEditor->GetEditorSubsystem<UPCGEditorSubsystem>();
+
+	check(Subsystem);
+	
+	const FPCGEditorDefaultExecutionSourceParams Params
+	{
+		.Graph = ProceduralVegetationGraph
+	};
+	
+	ExecutionSource = Subsystem->CreateEditorExecutionSource(Params);
+
+	// Set the stack
+	FPCGStack Stack;
+	Stack.PushFrame(ExecutionSource.Get());
+	Stack.PushFrame(ProceduralVegetationGraph);
+	SetStackBeingInspected(Stack);
+
+	// Ask for generation
+	ExecutionSource->Generate();
+
+	SessionStartTime = FDateTime::Now();
+
+	if (FEngineAnalytics::IsAvailable())
+	{
+		FEngineAnalytics::GetProvider().RecordEvent(TEXT("Editor.Usage.ProceduralVegetation.SessionStarted"));
+	}
+}
+
+TSubclassOf<UPCGEditorGraphSchema> FPVEditor::GetSchemaClass() const
+{
+	return UPVEditorSchema::StaticClass();
+}
+
+void FPVEditor::AddReferencedObjects(FReferenceCollector& Collector)
+{
+	FPCGEditor::AddReferencedObjects(Collector);
+
+	Collector.AddReferencedObject(ExecutionSource);
+}
+
+FText FPVEditor::GetToolkitName() const
+{
+	return GetLabelForObject(ProceduralVegetationBeingEdited);
+}
+
+FName FPVEditor::GetToolkitFName() const
+{
+	return TEXT("ProceduralVegetationEditor");
+}
+
+FText FPVEditor::GetBaseToolkitName() const
+{
+	return INVTEXT("Procedural Vegetation Editor");
+}
+
+FText FPVEditor::GetToolkitToolTipText() const
+{
+	return GetToolTipTextForObject(ProceduralVegetationBeingEdited);
+}
+
+FLinearColor FPVEditor::GetWorldCentricTabColorScale() const
+{
+	return FLinearColor::Blue;
+}
+
+FString FPVEditor::GetWorldCentricTabPrefix() const
+{
+	return TEXT("Procedural Vegetation Editor ");
+}
+
+void FPVEditor::OnClose()
+{
+	if (FEngineAnalytics::IsAvailable())
+	{
+		TArray<FAnalyticsEventAttribute> Attributes;
+		Attributes.Add(FAnalyticsEventAttribute(TEXT("TimeActive.Seconds"),  (FDateTime::Now() - SessionStartTime).GetTotalSeconds()));
+				
+		FEngineAnalytics::GetProvider().RecordEvent(TEXT("Editor.Usage.ProceduralVegetation.SessionEnded"), Attributes);
+	}
+	
+	FPCGEditor::OnClose();
+
+	if (ExecutionSource)
+	{
+		ExecutionSource->Sunset();
+		ExecutionSource = nullptr;
+	}
+}
+
+void FPVEditor::RegisterTabSpawners(const TSharedRef<class FTabManager>& InTabManager)
+{
+	FPCGEditor::RegisterTabSpawners(InTabManager);
+
+#if PROCEDURAL_VEGETATION_DEBUG_SHOW_COLLECTION_TAB
+	InTabManager->RegisterTabSpawner(CollectionSpreadSheetTabId, FOnSpawnTab::CreateSP(this, &FPVEditor::SpawnTab_CollectionSpreadSheet))
+		.SetDisplayName(LOCTEXT("CollectionSpreadsheetTab", "Collection Spreadsheet"))
+		.SetGroup(WorkspaceMenuCategory.ToSharedRef());
+#endif
+}
+
+void FPVEditor::UnregisterTabSpawners(const TSharedRef<class FTabManager>& InTabManager)
+{
+	FPCGEditor::UnregisterTabSpawners(InTabManager);
+
+#if PROCEDURAL_VEGETATION_DEBUG_SHOW_COLLECTION_TAB
+	InTabManager->UnregisterTabSpawner(CollectionSpreadSheetTabId);
+#endif
+}
+
+IPCGBaseSubsystem* FPVEditor::GetSubsystem() const
+{
+	return GEditor ? GEditor->GetEditorSubsystem<UPCGEditorSubsystem>() : nullptr;
+}
+
+TAttribute<FGraphAppearanceInfo> FPVEditor::GetAppearanceInfo() const
+{
+	FGraphAppearanceInfo AppearanceInfo;
+	AppearanceInfo.CornerText = INVTEXT("Megaseed");
+
+	return AppearanceInfo;
+}
+
+TSharedRef<FTabManager::FLayout> FPVEditor::GetDefaultLayout() const
+{
+	return FTabManager::NewLayout("Standalone_PVPCGGraphEditor_DefaultLayout_v1.0.4")
+	->AddArea // Main PCG Graph Editor Area
+	(
+		FTabManager::NewPrimaryArea()->SetOrientation(Orient_Vertical)
+		->Split // Top Section - Graph, Data Viewport, HLSL Source Editor, and Details View
+		(
+			FTabManager::NewSplitter()->SetOrientation(Orient_Horizontal)->SetSizeCoefficient(0.9f)
+			->SetSizeCoefficient(0.40f)
+			->Split
+			(
+				FTabManager::NewSplitter()->SetOrientation(Orient_Vertical)
+				->SetSizeCoefficient(0.1f)
+				->Split // Viewport and AVL
+				(
+					FTabManager::NewStack()
+					->SetSizeCoefficient(0.95f)
+					->AddTab(GetPanelID(EPCGEditorPanel::Viewport1), ETabState::OpenedTab)
+					->SetForegroundTab(GetPanelID(EPCGEditorPanel::Viewport1))
+					->SetHideTabWell(true)
+				)
+				->Split // Bottom Section - ALV
+				(FTabManager::NewStack()
+					->SetSizeCoefficient(0.05f)
+					->SetHideTabWell(false)
+					->AddTab(GetPanelID(EPCGEditorPanel::Attributes1), ETabState::ClosedTab)
+
+#if PROCEDURAL_VEGETATION_DEBUG_SHOW_COLLECTION_TAB
+					->AddTab(CollectionSpreadSheetTabId, ETabState::ClosedTab)
+#endif
+				)	
+			)
+			->Split // Graph and details panel
+			(
+				FTabManager::NewSplitter()
+				->SetOrientation(Orient_Vertical)
+				->SetSizeCoefficient(0.25f)
+				->Split
+				(
+				FTabManager::NewStack()
+					->SetSizeCoefficient(0.6f)
+					->AddTab(GetPanelID(EPCGEditorPanel::GraphEditor), ETabState::OpenedTab)
+					->SetHideTabWell(true)
+				)
+				->Split
+				(
+				FTabManager::NewStack()
+				->SetSizeCoefficient(0.4f)
+				->AddTab(GetPanelID(EPCGEditorPanel::PropertyDetails1), ETabState::OpenedTab)
+				)
+			)
+		)
+	);
+}
+
+void FPVEditor::RegisterToolbar() const
+{
+	UToolMenus* ToolMenus = UToolMenus::Get();
+	FName ParentName;
+	const FName ToolbarName = GetToolMenuToolbarName(ParentName);
+
+	UToolMenu* const AssetToolbar = UToolMenus::Get()->ExtendMenu(ToolbarName);
+	FToolMenuSection& Section = AssetToolbar->FindOrAddSection("ProceduralVegetationTools");
+
+	const FPVEditorCommands& Commands = FPVEditorCommands::Get();
+	Section.AddEntry(FToolMenuEntry::InitToolBarButton(
+			Commands.Export,
+			TAttribute<FText>(),
+			TAttribute<FText>(),
+			FSlateIcon(FAppStyle::GetAppStyleSetName(), "Icons.Toolbar.Export")));
+}
+
+void FPVEditor::BindCommands()
+{
+	FPCGEditor::BindCommands();
+
+	const FPVEditorCommands& EditorCommands = FPVEditorCommands::Get();
+
+	ToolkitCommands->MapAction(
+		EditorCommands.Export,
+		FExecuteAction::CreateSP(this, &FPVEditor::OnExport_Clicked));
+}
+
+void FPVEditor::OnSelectedNodesChanged(const TSet<UObject*>& InNewSelection)
+{
+	FPCGEditor::OnSelectedNodesChanged(InNewSelection);
+
+#if PROCEDURAL_VEGETATION_DEBUG_SHOW_COLLECTION_TAB
+	const UPCGNode* PCGNode = nullptr;
+	
+	if (!InNewSelection.IsEmpty())
+	{
+		const UPCGEditorGraphNodeBase* Node = Cast<UPCGEditorGraphNodeBase>(*InNewSelection.CreateConstIterator());
+		const FObjectProperty* ObjectProperty = CastField<FObjectProperty>(UPCGEditorGraphNodeBase::StaticClass()->FindPropertyByName(TEXT("PCGNode")));
+		if (Node && ObjectProperty)
+		{
+			PCGNode = Cast<UPCGNode>(ObjectProperty->GetObjectPropertyValue_InContainer(Node));
+		}
+	}
+
+	UpdateCollectionSpreadSheet(PCGNode);
+#endif
+}
+
+void FPVEditor::OnExport_Clicked()
+{
+	ExportInternal();
+}
+
+FReply FPVEditor::ExportInternal()
+{
+	const UProceduralVegetationGraph* ProceduralVegetationGraph = Cast<UProceduralVegetationGraph>(GetPCGGraph());
+	
+	TSharedPtr<SPVExportSelectionDialog> ExportSelectionWidget =
+		SNew(SPVExportSelectionDialog)
+		.Title(LOCTEXT("Export Selection", "Choose outputs to export"))
+		.Graph(ProceduralVegetationGraph);
+
+	if (ExportSelectionWidget->ShowModal() == EAppReturnType::Ok)
+	{
+		const TArray<UPCGNode*>& Nodes = ProceduralVegetationGraph->GetNodes();
+
+		TArray<const UPCGNode*> NodesToExport;
+		
+		for (const UPCGNode* Node : Nodes)
+		{
+			if(const UPVPCGOutputSettings* OutputSettings = Cast<UPVPCGOutputSettings>(Node->GetSettings()))
+			{
+				if(OutputSettings->ExportSettings.bShouldExport)
+				{
+					NodesToExport.Add(Node);
+				}
+			}
+		}
+
+		FScopedSlowTask Progress(NodesToExport.Num(), NSLOCTEXT("ExportAsset", "ExportAsset_SlowTask", "Exporting Asset(s) ..."));
+		Progress.MakeDialog();
+		
+		for(auto Node : NodesToExport)
+		{
+			if(const UPVPCGOutputSettings* OutputSettings = Cast<UPVPCGOutputSettings>(Node->GetSettings()))
+			{
+				FManagedArrayCollection OutCollection;
+				FName PinName = NAME_None;
+				if (GetCollectionFromNode(Node, PinName, OutCollection))
+				{
+					PV::Export::ExportCollectionAsMesh(ProceduralVegetationBeingEdited, OutCollection, OutCollection, OutputSettings->ExportSettings);
+					Progress.EnterProgressFrame();
+				}
+			}
+		}
+	}
+	
+	return FReply::Handled();
+}
+
+bool FPVEditor::GetCollectionFromNode(const UPCGNode* InNode, FName& OutPinName, FManagedArrayCollection& OutCollection) const
+{
+	const UPCGPin* Pin = nullptr;
+		
+	const TArray<TObjectPtr<UPCGPin>>& Pins = InNode->GetOutputPins();
+
+	if (Pins.IsValidIndex(0))
+	{
+		Pin = Pins[0];
+	}
+		
+	check(Pin)
+		
+	const FPCGStack* PCGStack = GetStackBeingInspected();
+	check(PCGStack)
+		
+	// Create a temporary stack with Node+Pin to query the exact DataCollection
+	FPCGStack Stack = *PCGStack;
+	TArray<FPCGStackFrame>& StackFrames = Stack.GetStackFramesMutable();
+	StackFrames.Reserve(StackFrames.Num() + 2);
+	StackFrames.Emplace(InNode);
+	StackFrames.Emplace(Pin);
+		
+	const FPCGDataCollection* Data = ExecutionSource->GetExecutionState().GetInspection().GetInspectionData(Stack);
+		
+	if(Data != nullptr)
+	{
+		const FPCGTaggedData* TaggedData = Data->TaggedData.GetData();
+		
+		if(TaggedData != nullptr)
+		{
+			const UPVPCGData* ProceduralVegetationData = Cast<UPVPCGData>(TaggedData->Data);
+		
+			if(ProceduralVegetationData != nullptr)
+			{
+				OutPinName = Pin->Properties.Label;
+				OutCollection = ProceduralVegetationData->GetCollection();
+				return true;
+			}
+		}
+	}
+
+	return false;
+}
+
+#if PROCEDURAL_VEGETATION_DEBUG_SHOW_COLLECTION_TAB
+TSharedRef<SCollectionSpreadSheetWidget> FPVEditor::CreateCollectionSpreadSheetWidget()
+{
+	return SNew(SCollectionSpreadSheetWidget);
+}
+
+TSharedRef<SDockTab> FPVEditor::SpawnTab_CollectionSpreadSheet(const FSpawnTabArgs& SpawnTabArgs)
+{
+	check(SpawnTabArgs.GetTabId() == FPVEditor::CollectionSpreadSheetTabId);
+
+	TSharedRef<SDockTab> DockableTab = SNew(SDockTab)
+		[
+			CollectionSpreadSheetWidget.ToSharedRef()
+		];
+
+	CollectionSpreadSheetTab = DockableTab;
+
+	return DockableTab;
+}
+
+void FPVEditor::UpdateCollectionSpreadSheet(const UPCGNode* InSelectedNode) const
+{
+	if (CollectionSpreadSheetTab.IsValid() && CollectionSpreadSheetTab.Pin()->GetVisibility().IsVisible() && CollectionSpreadSheetWidget)
+	{
+		if (InSelectedNode)
+		{
+			FManagedArrayCollection OutCollection;
+			FName OutputName;
+			GetCollectionFromNode(InSelectedNode, OutputName, OutCollection);
+			
+			CollectionSpreadSheetWidget->GetCollectionTable()->GetCollectionInfoMap().Empty();
+			CollectionSpreadSheetWidget->GetCollectionTable()->GetCollectionInfoMap().Add(OutputName.ToString(), { OutCollection });
+			CollectionSpreadSheetWidget->SetData(InSelectedNode->GetNodeTitle(EPCGNodeTitleType::FullTitle).ToString());
+		}
+		else
+		{
+			CollectionSpreadSheetWidget->SetData(FString());
+		}
+
+		CollectionSpreadSheetWidget->RefreshWidget();	
+	}
+}
+#endif
+
+#undef LOCTEXT_NAMESPACE
diff --git a/Engine/Plugins/Experimental/ProceduralVegetationEditor/Source/ProceduralVegetationEditor/Private/PVEditor.h b/Engine/Plugins/Experimental/ProceduralVegetationEditor/Source/ProceduralVegetationEditor/Private/PVEditor.h
new file mode 100644
--- /dev/null
+++ b/Engine/Plugins/Experimental/ProceduralVegetationEditor/Source/ProceduralVegetationEditor/Private/PVEditor.h
@@ -0,0 +1,71 @@
+﻿// Copyright Epic Games, Inc. All Rights Reserved.
+
+#pragma once
+
+#include "PCGEditor.h"
+
+class SCollectionSpreadSheetWidget;
+struct FManagedArrayCollection;
+class UPCGEditorDefaultExecutionSource;
+class UProceduralVegetation;
+
+#define PROCEDURAL_VEGETATION_DEBUG_SHOW_COLLECTION_TAB !UE_BUILD_SHIPPING_WITH_EDITOR 
+
+class FPVEditor : public FPCGEditor
+{
+public:
+	FPVEditor();
+
+	void Initialize(const EToolkitMode::Type InMode, const TSharedPtr<class IToolkitHost>& InToolkitHost, UProceduralVegetation* InProceduralVegetation);
+	virtual TSubclassOf<UPCGEditorGraphSchema> GetSchemaClass() const override;
+
+	virtual void AddReferencedObjects(FReferenceCollector& Collector) override;
+	
+	// ~Begin FAssetEditorToolkit interface
+	virtual FText GetToolkitName() const override;
+	virtual FName GetToolkitFName() const override;
+	virtual FText GetBaseToolkitName() const override;
+	virtual FText GetToolkitToolTipText() const override;
+	virtual FLinearColor GetWorldCentricTabColorScale() const override;
+	virtual FString GetWorldCentricTabPrefix() const override;
+	virtual void OnClose() override;
+	virtual void RegisterTabSpawners(const TSharedRef<class FTabManager>& InTabManager) override;
+	virtual void UnregisterTabSpawners(const TSharedRef<class FTabManager>& InTabManager) override;
+	// ~End FAssetEditorToolkit interface
+
+	virtual IPCGBaseSubsystem* GetSubsystem() const override;
+
+protected:
+	// ~Begin FPCGEditor interface
+	virtual TAttribute<FGraphAppearanceInfo> GetAppearanceInfo() const override;
+	virtual TSharedRef<FTabManager::FLayout> GetDefaultLayout() const override;
+	virtual void RegisterToolbar() const override;
+	virtual void BindCommands() override;
+	virtual void OnSelectedNodesChanged(const TSet<UObject*>& InNewSelection) override;
+	// ~End FPCGEditor interface
+
+	bool bGenerationInProgress = false;
+	TObjectPtr<UPCGEditorDefaultExecutionSource> ExecutionSource;
+
+private:
+	TObjectPtr<UProceduralVegetation> ProceduralVegetationBeingEdited = nullptr;
+
+	FDateTime SessionStartTime;
+
+#if PROCEDURAL_VEGETATION_DEBUG_SHOW_COLLECTION_TAB
+	static const FName CollectionSpreadSheetTabId;
+
+	TWeakPtr<SDockTab> CollectionSpreadSheetTab;
+	TSharedPtr<SCollectionSpreadSheetWidget> CollectionSpreadSheetWidget;
+
+	TSharedRef<SCollectionSpreadSheetWidget> CreateCollectionSpreadSheetWidget();
+
+	TSharedRef<SDockTab> SpawnTab_CollectionSpreadSheet(const FSpawnTabArgs& SpawnTabArgs);
+	
+	void UpdateCollectionSpreadSheet(const UPCGNode* InSelectedNode) const;
+#endif
+	
+	bool GetCollectionFromNode(const UPCGNode* InNode, FName& OutPinName, FManagedArrayCollection& OutCollection) const;
+	void OnExport_Clicked();
+	FReply ExportInternal();
+};
diff --git a/Engine/Plugins/Experimental/ProceduralVegetationEditor/Source/ProceduralVegetationEditor/Private/PVEditorSchema.cpp b/Engine/Plugins/Experimental/ProceduralVegetationEditor/Source/ProceduralVegetationEditor/Private/PVEditorSchema.cpp
new file mode 100644
--- /dev/null
+++ b/Engine/Plugins/Experimental/ProceduralVegetationEditor/Source/ProceduralVegetationEditor/Private/PVEditorSchema.cpp
@@ -0,0 +1,23 @@
+﻿// Copyright Epic Games, Inc. All Rights Reserved.
+
+#include "PVEditorSchema.h"
+#include "Nodes/PVPCGBaseSettings.h"
+
+EPCGElementType UPVEditorSchema::GetElementTypeFiltering() const
+{
+	return EPCGElementType::Native;
+}
+
+const FPCGGraphEditorCustomization& UPVEditorSchema::GetGraphEditorCustomization(const UEdGraph* InEdGraph) const
+{
+	static const FPCGGraphEditorCustomization GraphEditorCustomization
+	{
+		.bFilterSettings = true,
+		.FilteredSettingsTypes =
+	{
+			UPVPCGBaseSettings::StaticClass()
+		}
+	};
+
+	return GraphEditorCustomization;
+}
diff --git a/Engine/Plugins/Experimental/ProceduralVegetationEditor/Source/ProceduralVegetationEditor/Private/Visualizations/PVPCGDataVisualization.cpp b/Engine/Plugins/Experimental/ProceduralVegetationEditor/Source/ProceduralVegetationEditor/Private/Visualizations/PVPCGDataVisualization.cpp
new file mode 100644
--- /dev/null
+++ b/Engine/Plugins/Experimental/ProceduralVegetationEditor/Source/ProceduralVegetationEditor/Private/Visualizations/PVPCGDataVisualization.cpp
@@ -0,0 +1,226 @@
+// Copyright Epic Games, Inc. All Rights Reserved.
+
+#include "PVPCGDataVisualization.h"
+
+#include "AdvancedPreviewScene.h"
+#include "EditorViewportClient.h"
+#include "GeometryCollectionToDynamicMesh.h"
+#include "PVBoneComponent.h"
+#include "PCGComponent.h"
+#include "PCGContext.h"
+#include "PCGDataVisualization.h"
+#include "PCGEditorSettings.h"
+#include "PlanarCut.h"
+#include "Components/DynamicMeshComponent.h"
+#include "Components/InstancedStaticMeshComponent.h"
+#include "Components/LineBatchComponent.h"
+#include "Facades/PVFoliageFacade.h"
+#include "GeometryCollection/GeometryCollectionObject.h"
+#include "GeometryCollection/Facades/CollectionRenderingFacade.h"
+#include "Helpers/PVUtilities.h"
+#include "PVPCGData.h"
+#include "Components/TextRenderComponent.h"
+#include "Visualizations/PVSkeletonVisualizerComponent.h"
+
+#define LOCTEXT_NAMESPACE "PCGManagedArrayCollectionDataVisualization"
+
+class UPVBoneComponent;
+
+DEFINE_LOG_CATEGORY(LogProceduralVegetationDataVisualization);
+
+FPVPCGDataVisualization::FPVPCGDataVisualization()
+{
+	UPCGEditorSettings* PCGEditorSettings = GetMutableDefault<UPCGEditorSettings>();
+	PCGEditorSettings->bAutoFocusViewport = false;
+
+	RegisterVisualizations();
+}
+
+void FPVPCGDataVisualization::RegisterVisualizations()
+{
+	RenderMap.Add(EPVRenderType::Skeleton, FPVRenderCallback::CreateRaw(this, &FPVPCGDataVisualization::SkeletonRenderer));
+	RenderMap.Add(EPVRenderType::Mesh, FPVRenderCallback::CreateRaw(this, &FPVPCGDataVisualization::MeshRenderer));
+	RenderMap.Add(EPVRenderType::Foliage, FPVRenderCallback::CreateRaw(this, &FPVPCGDataVisualization::FoliageRenderer));
+	RenderMap.Add(EPVRenderType::Bones, FPVRenderCallback::CreateRaw(this, &FPVPCGDataVisualization::BonesRenderer));
+}
+
+void FPVPCGDataVisualization::ExecuteDebugDisplay(FPCGContext* Context, const UPCGSettingsInterface* SettingsInterface, const UPCGData* Data,
+                                                         AActor* TargetActor) const
+{
+	const UPVPCGData* ManagedArrayCollectionData = CastChecked<UPVPCGData>(Data);
+	UPCGComponent* SourcePCGComponent = Cast<UPCGComponent>(Context->ExecutionSource.Get());
+	if (!ManagedArrayCollectionData || !SourcePCGComponent)
+	{
+		return;
+	}
+}
+
+FPCGSetupSceneFunc FPVPCGDataVisualization::GetViewportSetupFunc(const UPCGSettingsInterface* SettingsInterface, const UPCGData* Data) const
+{
+	return [this, WeakData=TWeakObjectPtr<const UPVPCGData>(Cast<UPVPCGData>(Data))](
+		FPCGSceneSetupParams& InOutParams)
+	{
+		check(InOutParams.Scene);
+
+		if (!WeakData.IsValid())
+		{
+			UE_LOG(LogProceduralVegetationDataVisualization, Error, TEXT("Failed to setup data viewport, the data was lost or invalid."));
+			return;
+		}
+
+		const FManagedArrayCollection& Collection = WeakData->GetCollection();
+
+		FBoxSphereBounds Bounds(ForceInit);
+
+		for (const EPVRenderType& RenderType : WeakData->GetRenderType())
+		{
+			if (RenderMap.Contains(RenderType))
+			{
+				InOutParams.Scene->GetLineBatcher()->Flush();
+				RenderMap[RenderType].Execute(Collection, InOutParams, Bounds);
+			}
+		}
+
+		if (!InOutParams.Scene->IsGridEnabled())
+		{
+			InOutParams.Scene->HandleToggleGrid();
+		}
+		if (!InOutParams.Scene->IsEnvironmentEnabled())
+		{
+			InOutParams.Scene->HandleToggleEnvironment();
+		}
+		if (!InOutParams.Scene->IsPostProcessingEnabled())
+		{
+			InOutParams.Scene->HandleTogglePostProcessing();
+		}
+
+		UStaticMesh* MannequinMesh = LoadObject<UStaticMesh>(
+			nullptr,
+			TEXT("/ProceduralVegetationEditor/Mannequin/Viewport_Mannequin_T_Pose.Viewport_Mannequin_T_Pose")
+		);
+
+		constexpr bool bShowMannequin = true;
+		if (bShowMannequin)
+		{
+			UStaticMeshComponent* MannequinComponent = NewObject<UStaticMeshComponent>(GetTransientPackage(), NAME_None, RF_Transient);
+			MannequinComponent->SetStaticMesh(MannequinMesh);
+
+			InOutParams.ManagedResources.Add(MannequinComponent);
+			InOutParams.Scene->AddComponent(MannequinComponent, FTransform(FVector(-Bounds.BoxExtent.X * 2.0, 0, 0)));
+		}
+
+		InOutParams.Scene->SetFloorVisibility(true);
+		InOutParams.Scene->SetFloorOffset(1.0f);
+
+		InOutParams.FocusBounds = FBoxSphereBounds(Bounds.Origin, Bounds.BoxExtent, Bounds.SphereRadius);
+	};
+}
+
+void FPVPCGDataVisualization::MeshRenderer(const FManagedArrayCollection& InCollection, FPCGSceneSetupParams& InOutParams,
+                                                  FBoxSphereBounds& OutBounds)
+{
+	TRACE_CPUPROFILER_EVENT_SCOPE(FPVPCGDataVisualization::MeshRenderer);
+
+	TArray<UMaterialInterface*> Materials;
+	
+	if (InCollection.HasAttribute("MaterialPath", FGeometryCollection::MaterialGroup))
+	{
+		const TManagedArray<FString>& MaterialPaths = InCollection.GetAttribute<FString>("MaterialPath", FGeometryCollection::MaterialGroup);
+		for (const FString& MaterialPath : MaterialPaths.GetConstArray())
+		{
+			Materials.Add(LoadObject<UMaterialInterface>(nullptr, MaterialPath));
+		}
+	}
+
+	UE::Geometry::FGeometryCollectionToDynamicMeshes Converter;
+	Converter.Init(InCollection, {});
+	for (const auto& MeshInfo : Converter.Meshes)
+	{
+		FDynamicMesh3* DynamicMesh = MeshInfo.Mesh.Get();
+		if (DynamicMesh->VertexCount())
+		{
+			TObjectPtr<UDynamicMeshComponent> DynamicMeshComponent = NewObject<UDynamicMeshComponent>(GetTransientPackage(), NAME_None, RF_Transient);
+			DynamicMeshComponent->SetMesh(MoveTemp(*DynamicMesh));
+			DynamicMeshComponent->SetTangentsType(EDynamicMeshComponentTangentsMode::AutoCalculated);
+			DynamicMeshComponent->ConfigureMaterialSet(Materials);
+			DynamicMeshComponent->UpdateBounds();
+
+			InOutParams.ManagedResources.Add(DynamicMeshComponent);
+			InOutParams.Scene->AddComponent(DynamicMeshComponent, FTransform::Identity);
+
+			OutBounds = OutBounds + DynamicMeshComponent->CalcLocalBounds();
+		}
+	}
+}
+
+void FPVPCGDataVisualization::SkeletonRenderer(const FManagedArrayCollection& InCollection, FPCGSceneSetupParams& InOutParams,
+                                                      FBoxSphereBounds& OutBounds)
+{
+	TRACE_CPUPROFILER_EVENT_SCOPE(FPVPCGDataVisualization::SkeletonRenderer);
+
+	const TObjectPtr<UPVSkeletonVisualizerComponent> VisualizerComponent = NewObject<UPVSkeletonVisualizerComponent>(
+		GetTransientPackage(), NAME_None, RF_Transient);
+	VisualizerComponent->SetCollection(&InCollection);
+	VisualizerComponent->UpdateBounds();
+
+	InOutParams.ManagedResources.Add(VisualizerComponent);
+	InOutParams.Scene->AddComponent(VisualizerComponent, FTransform::Identity);
+
+	OutBounds = OutBounds + VisualizerComponent->CalcLocalBounds();
+}
+
+void FPVPCGDataVisualization::FoliageRenderer(const FManagedArrayCollection& InCollection, FPCGSceneSetupParams& InOutParams,
+                                                     FBoxSphereBounds& OutBounds)
+{
+	TRACE_CPUPROFILER_EVENT_SCOPE(FPVPCGDataVisualization::FoliageRenderer);
+
+	PV::Facades::FFoliageFacade FoliageFacade(InCollection);
+
+	auto OnComponentCreated = [&InOutParams](UMeshComponent* InInstancedComponent)
+	{
+		InOutParams.ManagedResources.Add(InInstancedComponent);
+		InInstancedComponent->RegisterComponentWithWorld(InOutParams.Scene->GetWorld());
+		InOutParams.Scene->AddComponent(InInstancedComponent, FTransform::Identity);
+	};
+
+	TMap<FString, TObjectPtr<UMeshComponent>> OutInstancedComponentMap;
+
+	PV::Utilities::AddFoliageInstances(FoliageFacade, nullptr,
+	                                           PV::Utilities::FFoliageComponentCreatedCallback::CreateLambda(OnComponentCreated),
+	                                           OutInstancedComponentMap);
+}
+
+void FPVPCGDataVisualization::BonesRenderer(const FManagedArrayCollection& InCollection, FPCGSceneSetupParams& InOutParams,
+	FBoxSphereBounds& OutBounds)
+{
+	TRACE_CPUPROFILER_EVENT_SCOPE(FPVPCGDataVisualization::BonesRenderer);
+
+	const TObjectPtr<UPVBoneComponent> BoneComponent = NewObject<UPVBoneComponent>(
+		GetTransientPackage(), NAME_None, RF_Transient);
+	BoneComponent->SetCollection(&InCollection);
+	BoneComponent->UpdateBounds();
+
+	InOutParams.ManagedResources.Add(BoneComponent);
+	InOutParams.Scene->AddComponent(BoneComponent, FTransform::Identity);
+
+	FString InTextToDraw = FString::Format(TEXT("{0} Bones"),{BoneComponent->GetBoneCount()});
+	const FVector3f& Pos = (FVector3f::ForwardVector * 100);
+
+	DrawDebugString(BoneComponent->GetWorld(), FVector::Zero(), InTextToDraw);
+	
+	TObjectPtr<UTextRenderComponent> Component = NewObject<UTextRenderComponent>(GetTransientPackage(), NAME_None, RF_Transient);
+	Component->SetText(FText::FromString(InTextToDraw));
+	Component->SetTextRenderColor(FColor::Blue);
+	Component->SetWorldSize(40);
+	Component->SetGenerateOverlapEvents(false);
+						
+	InOutParams.ManagedResources.Add(Component);
+	FTransform TextTransform = FTransform(FVector(Pos.X, Pos.Y, Pos.Z));
+	FRotator TextRotator = FRotator(0,90,0);
+	TextTransform.SetRotation(TextRotator.Quaternion());
+	InOutParams.Scene->AddComponent(Component, TextTransform);
+
+	OutBounds = OutBounds + BoneComponent->CalcLocalBounds();
+}
+
+#undef LOCTEXT_NAMESPACE
diff --git a/Engine/Plugins/Experimental/ProceduralVegetationEditor/Source/ProceduralVegetationEditor/Private/Visualizations/PVPCGDataVisualization.h b/Engine/Plugins/Experimental/ProceduralVegetationEditor/Source/ProceduralVegetationEditor/Private/Visualizations/PVPCGDataVisualization.h
new file mode 100644
--- /dev/null
+++ b/Engine/Plugins/Experimental/ProceduralVegetationEditor/Source/ProceduralVegetationEditor/Private/Visualizations/PVPCGDataVisualization.h
@@ -0,0 +1,36 @@
+// Copyright Epic Games, Inc. All Rights Reserved.
+
+#pragma once
+
+#include "PCGDataVisualization.h"
+#include "DataVisualizations/PCGSpatialDataVisualization.h"
+#include "PVPCGData.h"
+
+class AActor;
+class UPCGData;
+struct FPCGContext;
+
+DECLARE_DELEGATE_ThreeParams(FPVRenderCallback, const FManagedArrayCollection&, FPCGSceneSetupParams&, FBoxSphereBounds&);
+
+class FPVPCGDataVisualization : public IPCGSpatialDataVisualization
+{
+public:
+	FPVPCGDataVisualization();
+	
+	// ~Begin IPCGDataVisualization interface
+	virtual void ExecuteDebugDisplay(FPCGContext* Context, const UPCGSettingsInterface* SettingsInterface, const UPCGData* Data, AActor* TargetActor) const override;
+	virtual FPCGSetupSceneFunc GetViewportSetupFunc(const UPCGSettingsInterface* SettingsInterface, const UPCGData* Data) const override;
+	// ~End IPCGDataVisualization interface
+
+protected:
+	TMap<EPVRenderType, FPVRenderCallback> RenderMap;
+	
+	void RegisterVisualizations();
+
+	void MeshRenderer(const FManagedArrayCollection& InCollection, FPCGSceneSetupParams& InOutParams, FBoxSphereBounds& OutBounds);
+	void SkeletonRenderer(const FManagedArrayCollection& InCollection, FPCGSceneSetupParams& InOutParams, FBoxSphereBounds& OutBounds);
+	void FoliageRenderer(const FManagedArrayCollection& InCollection, FPCGSceneSetupParams& InOutParams, FBoxSphereBounds& OutBounds);
+	void BonesRenderer(const FManagedArrayCollection& InCollection, FPCGSceneSetupParams& InOutParams, FBoxSphereBounds& OutBounds);
+};
+
+DECLARE_LOG_CATEGORY_EXTERN(LogProceduralVegetationDataVisualization, Log, All);
\ No newline at end of file
diff --git a/Engine/Plugins/Experimental/ProceduralVegetationEditor/Source/ProceduralVegetationEditor/Private/Visualizations/PVSkeletonVisualizerComponent.cpp b/Engine/Plugins/Experimental/ProceduralVegetationEditor/Source/ProceduralVegetationEditor/Private/Visualizations/PVSkeletonVisualizerComponent.cpp
new file mode 100644
--- /dev/null
+++ b/Engine/Plugins/Experimental/ProceduralVegetationEditor/Source/ProceduralVegetationEditor/Private/Visualizations/PVSkeletonVisualizerComponent.cpp
@@ -0,0 +1,64 @@
+// Copyright Epic Games, Inc. All Rights Reserved.
+
+
+#include "PVSkeletonVisualizerComponent.h"
+#include "ProceduralVegetationEditorModule.h"
+#include "PrimitiveDrawInterface.h"
+#include "Facades/PVBranchFacade.h"
+#include "Facades/PVPointFacade.h"
+
+UPVSkeletonVisualizerComponent::UPVSkeletonVisualizerComponent()
+	: SkeletonCollection(nullptr)
+{
+	bAutoActivate = true;
+	bTickInEditor = true;
+	PrimaryComponentTick.bCanEverTick = true;
+	bUseEditorCompositing = true;
+	SetGenerateOverlapEvents(false);
+	bIgnoreStreamingManagerUpdate = true;
+	
+	BBox = FBox(ForceInit);
+	Bounds = FBoxSphereBounds(FVector::ZeroVector, FVector::OneVector, 1);
+}
+
+const FManagedArrayCollection* UPVSkeletonVisualizerComponent::GetCollection() const
+{
+	return SkeletonCollection;
+}
+
+void UPVSkeletonVisualizerComponent::SetCollection(const FManagedArrayCollection* const InSkeletonCollection)
+{
+	if (InSkeletonCollection == nullptr)
+	{
+		return;
+	}
+
+	const PV::Facades::FBranchFacade BranchFacade(*InSkeletonCollection);
+	const PV::Facades::FPointFacade PointFacade(*InSkeletonCollection);
+	if (!BranchFacade.IsValid() || !PointFacade.IsValid())
+	{
+		UE_LOG(LogProceduralVegetationEditor, Error, TEXT("Failed to setup data viewport, Collection data is not valid for Procedural Vegetation."));
+		return;
+	}
+
+	SkeletonCollection = InSkeletonCollection;
+
+	const TManagedArray<FVector3f>& PointsPositions = PointFacade.GetPositions();
+
+	InitBounds();
+	for (int BranchIndex = 0; BranchIndex < BranchFacade.GetElementCount(); ++BranchIndex)
+	{
+		const TArray<int32>& BranchPoints = BranchFacade.GetPoints(BranchIndex);
+		FVector PreviousPosition = FVector(PointsPositions[BranchPoints[0]]);
+		BBox += PreviousPosition;
+
+		for (int32 BranchPointIndex = 1; BranchPointIndex < BranchPoints.Num(); ++BranchPointIndex)
+		{
+			FVector CurrentPosition = FVector(PointsPositions[BranchPoints[BranchPointIndex]]);
+
+			AddLine(PreviousPosition, CurrentPosition);
+
+			PreviousPosition = CurrentPosition;
+		}
+	}
+}
\ No newline at end of file
diff --git a/Engine/Plugins/Experimental/ProceduralVegetationEditor/Source/ProceduralVegetationEditor/Private/Widgets/SPVExportIem.cpp b/Engine/Plugins/Experimental/ProceduralVegetationEditor/Source/ProceduralVegetationEditor/Private/Widgets/SPVExportIem.cpp
new file mode 100644
--- /dev/null
+++ b/Engine/Plugins/Experimental/ProceduralVegetationEditor/Source/ProceduralVegetationEditor/Private/Widgets/SPVExportIem.cpp
@@ -0,0 +1,50 @@
+// Copyright Epic Games, Inc. All Rights Reserved.
+#include "SPVExportItem.h"
+#include "Widgets/Layout/SBox.h"
+#include <Widgets/Text/STextBlock.h>
+
+#define LOCTEXT_NAMESPACE "SPVExportItem"
+
+void SPVExportItem::Construct(const FArguments& InArgs)
+{
+	OnOutputSelectionChanged = InArgs._OnOutputSelectionChanged;
+	bIsSelected = InArgs._bIsSelected;
+	auto& Name = InArgs._Name;
+
+	auto OnSelectionChanged = [this, Name](ECheckBoxState NewState)
+	{
+		bIsSelected = NewState == ECheckBoxState::Checked;
+		OnOutputSelectionChanged.ExecuteIfBound(Name.ToString(), NewState);
+	};
+
+	ChildSlot
+	[
+		SNew(SHorizontalBox)
+
+		+ SHorizontalBox::Slot()
+		[
+			SNew(SCheckBox)
+			.IsChecked(this, &SPVExportItem::OnCheckBoxState)
+			.OnCheckStateChanged_Lambda(OnSelectionChanged)
+			.Content()
+			[
+				SNew(SHorizontalBox)
+
+				+ SHorizontalBox::Slot()
+				.VAlign(VAlign_Center)
+				.Padding(5)
+				[
+					SNew(STextBlock)
+					.Text(InArgs._Name)
+				]
+			]
+		]
+	];
+}
+
+ECheckBoxState SPVExportItem::OnCheckBoxState() const
+{
+	return bIsSelected ? ECheckBoxState::Checked : ECheckBoxState::Unchecked;
+}
+
+#undef LOCTEXT_NAMESPACE
\ No newline at end of file
diff --git a/Engine/Plugins/Experimental/ProceduralVegetationEditor/Source/ProceduralVegetationEditor/Private/Widgets/SPVExportItem.h b/Engine/Plugins/Experimental/ProceduralVegetationEditor/Source/ProceduralVegetationEditor/Private/Widgets/SPVExportItem.h
new file mode 100644
--- /dev/null
+++ b/Engine/Plugins/Experimental/ProceduralVegetationEditor/Source/ProceduralVegetationEditor/Private/Widgets/SPVExportItem.h
@@ -0,0 +1,27 @@
+// Copyright Epic Games, Inc. All Rights Reserved.
+#pragma once
+#include "CoreMinimal.h"
+#include "Widgets/SBoxPanel.h"
+#include "Input/Reply.h"
+#include <Widgets/Input/SCheckBox.h>
+#include "Widgets/SCompoundWidget.h"
+
+class SPVExportItem : public SCompoundWidget
+{
+	DECLARE_DELEGATE_TwoParams(FOnOutputSelectionChanged, const FString ItemName, ECheckBoxState NewState);
+
+	SLATE_BEGIN_ARGS(SPVExportItem)
+	{}
+	SLATE_ARGUMENT(FText, Name)
+	SLATE_ARGUMENT(bool, bIsSelected)
+	SLATE_EVENT(FOnOutputSelectionChanged, OnOutputSelectionChanged)
+	SLATE_END_ARGS()
+
+public:
+	void Construct(const FArguments& InArgs);
+	ECheckBoxState OnCheckBoxState() const;
+
+	FOnOutputSelectionChanged OnOutputSelectionChanged;
+private:
+	bool bIsSelected;
+};
diff --git a/Engine/Plugins/Experimental/ProceduralVegetationEditor/Source/ProceduralVegetationEditor/Private/Widgets/SPVExportSelectionDialog.cpp b/Engine/Plugins/Experimental/ProceduralVegetationEditor/Source/ProceduralVegetationEditor/Private/Widgets/SPVExportSelectionDialog.cpp
new file mode 100644
--- /dev/null
+++ b/Engine/Plugins/Experimental/ProceduralVegetationEditor/Source/ProceduralVegetationEditor/Private/Widgets/SPVExportSelectionDialog.cpp
@@ -0,0 +1,131 @@
+// Copyright Epic Games, Inc. All Rights Reserved.
+#include "SPVExportSelectionDialog.h"
+#include "Widgets/Layout/SSeparator.h"
+#include "SPrimaryButton.h"
+#include "Widgets/Layout/SUniformGridPanel.h"
+#include "Editor.h"
+#include "SPVExportItem.h"
+#include "ProceduralVegetation.h"
+#include "Nodes/PVPCGOutput.h"
+#include "Widgets/Layout/SScrollBox.h"
+
+#define LOCTEXT_NAMESPACE "SPVExportSelectionDialog"
+
+void SPVExportSelectionDialog::Construct(const FArguments& InArgs)
+{
+	Graph = InArgs._Graph;
+	
+	SWindow::Construct(SWindow::FArguments()
+		.Title(InArgs._Title)
+		.SupportsMinimize(false)
+		.SupportsMaximize(false)
+		.ClientSize(FVector2D(350, 450))
+		[
+			SNew(SVerticalBox)
+			+ SVerticalBox::Slot() // Add user input block
+			.Padding(2, 2, 2, 4)
+			[
+				SNew(SBorder)
+				.BorderImage(FAppStyle::GetBrush("ToolPanel.GroupBorder"))
+				[
+					SNew(SVerticalBox)
+					+ SVerticalBox::Slot()
+					.FillHeight(1)
+					[
+						SAssignNew(ScrollBox,SScrollBox)
+					]
+				]
+			]
+			+ SVerticalBox::Slot()
+			.AutoHeight()
+			.HAlign(HAlign_Right)
+			.VAlign(VAlign_Bottom)
+			.Padding(8.f, 16.f)
+			[
+				SNew(SUniformGridPanel)
+				.MinDesiredSlotWidth(FAppStyle::GetFloat("StandardDialog.MinDesiredSlotWidth"))
+				.MinDesiredSlotHeight(FAppStyle::GetFloat("StandardDialog.MinDesiredSlotHeight"))
+				.SlotPadding(FAppStyle::GetMargin("StandardDialog.SlotPadding"))
+				+ SUniformGridPanel::Slot(0, 0)
+				[
+					SNew(SPrimaryButton)
+					.Text(LOCTEXT("Export", "Export"))
+					.OnClicked(this, &SPVExportSelectionDialog::OnButtonClick, EAppReturnType::Ok)
+				]
+				+ SUniformGridPanel::Slot(1, 0)
+				[
+					SNew(SButton)
+					.Text(LOCTEXT("Cancel", "Cancel"))
+					.ContentPadding(FAppStyle::GetMargin("StandardDialog.ContentPadding"))
+					.OnClicked(this, &SPVExportSelectionDialog::OnButtonClick, EAppReturnType::Cancel)
+				]
+			]
+		]);
+	AddExportItems();
+}
+void SPVExportSelectionDialog::AddExportItems()
+{
+	ScrollBox->ClearChildren();
+	
+	Graph->ForEachNode([this](UPCGNode* Node)
+	{
+		const UPVPCGOutputSettings* OutputSettings = Cast<UPVPCGOutputSettings>(Node->GetSettings());
+
+		if (OutputSettings)
+		{
+			FText OutputName = Node->GetNodeTitle(EPCGNodeTitleType::FullTitle);
+	
+			ScrollBox->AddSlot()
+			.Padding(5)
+			[
+				SNew(SPVExportItem)
+				.Name(OutputName)
+				.OnOutputSelectionChanged(this, &SPVExportSelectionDialog::OnOutputSelectionChanged)
+				.bIsSelected(OutputSettings->ExportSettings.bShouldExport)
+			];
+			ScrollBox->AddSlot()
+			[
+				SNew(SSeparator)
+				.Thickness(1)
+			];
+		}
+
+		return true;
+	});
+}
+FReply SPVExportSelectionDialog::OnButtonClick(EAppReturnType::Type ButtonID)
+{
+	UserResponse = ButtonID;
+	if (ButtonID == EAppReturnType::Cancel || ButtonID == EAppReturnType::Ok)
+	{
+		// Only close the window if canceling or if the ok
+		RequestDestroyWindow();
+	}
+	else
+	{
+		// reset the user response in case the window is closed using 'x'.
+		UserResponse = EAppReturnType::Cancel;
+	}
+	return FReply::Handled();
+}
+EAppReturnType::Type SPVExportSelectionDialog::ShowModal()
+{
+	GEditor->EditorAddModalWindow(SharedThis(this));
+	return UserResponse;
+}
+void SPVExportSelectionDialog::OnOutputSelectionChanged(const FString ItemName, ECheckBoxState NewState)
+{
+	Graph->ForEachNode([this, &ItemName, &NewState](UPCGNode* Node)
+	{
+		if(ItemName == Node->GetNodeTitle(EPCGNodeTitleType::FullTitle).ToString())
+		{
+			if (UPVPCGOutputSettings* OutputSettings = Cast<UPVPCGOutputSettings>(Node->GetSettings()))
+			{
+				OutputSettings->ExportSettings.bShouldExport = NewState ==ECheckBoxState::Checked;
+			}
+		}
+		
+		return true;
+	});
+}
+#undef LOCTEXT_NAMESPACE
\ No newline at end of file
diff --git a/Engine/Plugins/Experimental/ProceduralVegetationEditor/Source/ProceduralVegetationEditor/Private/Widgets/SPVExportSelectionDialog.h b/Engine/Plugins/Experimental/ProceduralVegetationEditor/Source/ProceduralVegetationEditor/Private/Widgets/SPVExportSelectionDialog.h
new file mode 100644
--- /dev/null
+++ b/Engine/Plugins/Experimental/ProceduralVegetationEditor/Source/ProceduralVegetationEditor/Private/Widgets/SPVExportSelectionDialog.h
@@ -0,0 +1,35 @@
+// Copyright Epic Games, Inc. All Rights Reserved.
+#pragma once
+#include "CoreMinimal.h"
+#include "Input/Reply.h"
+#include "Widgets/SWindow.h"
+
+class SScrollBox;
+class UProceduralVegetationGraph;
+
+class SPVExportSelectionDialog : public SWindow
+{
+	SLATE_BEGIN_ARGS(SPVExportSelectionDialog)
+	{}
+	SLATE_ARGUMENT(FText, Title)
+	SLATE_ARGUMENT(const UProceduralVegetationGraph*, Graph)
+	SLATE_ARGUMENT(TArray<FText>, ExportItems)
+	SLATE_END_ARGS()
+
+public:
+	void Construct(const FArguments& InArgs);
+
+	/** Displays the dialog in a blocking fashion */
+	EAppReturnType::Type ShowModal();
+
+protected:
+	FReply OnButtonClick(EAppReturnType::Type ButtonID);
+	EAppReturnType::Type UserResponse = EAppReturnType::Type::Cancel;
+
+private:
+	TSharedPtr<SScrollBox> ScrollBox;
+	const UProceduralVegetationGraph* Graph = nullptr;
+
+	void AddExportItems();
+	void OnOutputSelectionChanged(const FString ItemName, ECheckBoxState NewState);
+};
