Change display name of  ProceduralVegetation to Procedural Vegetation

diff --git a/Engine/Plugins/Experimental/ProceduralVegetationEditor/Source/ProceduralVegetation/Private/Nodes/PVPCGImporter.h b/Engine/Plugins/Experimental/ProceduralVegetationEditor/Source/ProceduralVegetation/Private/Nodes/PVPCGImporter.h
--- a/Engine/Plugins/Experimental/ProceduralVegetationEditor/Source/ProceduralVegetation/Private/Nodes/PVPCGImporter.h
+++ b/Engine/Plugins/Experimental/ProceduralVegetationEditor/Source/ProceduralVegetation/Private/Nodes/PVPCGImporter.h
@@ -15,8 +15,8 @@ class UPVPCGImporterSettings : public UPVPCGBaseSettings
 public:
 #if WITH_EDITOR
 	virtual FName GetDefaultNodeName() const override { return FName(TEXT("ProceduralVegetationImporter")); }
-	virtual FText GetDefaultNodeTitle() const override { return NSLOCTEXT("PVPCGImporter", "NodeTitle", "ProceduralVegetation Importer"); }
-	virtual FText GetNodeTooltipText() const override { return NSLOCTEXT("PVPCGImporter", "NodeTooltip", "Create ProceduralVegetation skeleton"); }
+	virtual FText GetDefaultNodeTitle() const override { return NSLOCTEXT("PVPCGImporter", "NodeTitle", "Procedural Vegetation Importer"); }
+	virtual FText GetNodeTooltipText() const override { return NSLOCTEXT("PVPCGImporter", "NodeTooltip", "Create Procedural Vegetation Skeleton"); }
 	virtual FLinearColor GetNodeTitleColor() const override { return FLinearColor(FColor::Black); }
 
 	virtual TArray<EPVRenderType> GetRenderType() const override { return TArray{ EPVRenderType::Skeleton }; }
