Fix old "Megaseed" name being used as corner text in the Procedural Generation Editor graph

diff --git a/Engine/Plugins/Experimental/ProceduralVegetationEditor/Source/ProceduralVegetationEditor/Private/PVEditor.cpp b/Engine/Plugins/Experimental/ProceduralVegetationEditor/Source/ProceduralVegetationEditor/Private/PVEditor.cpp
--- a/Engine/Plugins/Experimental/ProceduralVegetationEditor/Source/ProceduralVegetationEditor/Private/PVEditor.cpp
+++ b/Engine/Plugins/Experimental/ProceduralVegetationEditor/Source/ProceduralVegetationEditor/Private/PVEditor.cpp
@@ -169,7 +169,7 @@ IPCGBaseSubsystem* FPVEditor::GetSubsystem() const
 TAttribute<FGraphAppearanceInfo> FPVEditor::GetAppearanceInfo() const
 {
 	FGraphAppearanceInfo AppearanceInfo;
-	AppearanceInfo.CornerText = INVTEXT("Megaseed");
+	AppearanceInfo.CornerText = INVTEXT("PVE Graph");
 
 	return AppearanceInfo;
 }
