[ProceduralVegetationEditor] Added Trimsheet texture suport in material settings

diff --git a/Engine/Plugins/Experimental/ProceduralVegetationEditor/Source/ProceduralVegetation/Private/Facades/PVPointFacade.cpp b/Engine/Plugins/Experimental/ProceduralVegetationEditor/Source/ProceduralVegetation/Private/Facades/PVPointFacade.cpp
--- a/Engine/Plugins/Experimental/ProceduralVegetationEditor/Source/ProceduralVegetation/Private/Facades/PVPointFacade.cpp
+++ b/Engine/Plugins/Experimental/ProceduralVegetationEditor/Source/ProceduralVegetation/Private/Facades/PVPointFacade.cpp
@@ -24,6 +24,7 @@ namespace PV::Facades
 		, BudDevelopment(InCollection, AttributeNames::BudDevelopment, GroupNames::PointGroup)
 		, TextureCoordV(InCollection, AttributeNames::TextureCoordV, GroupNames::PointGroup)
 		, TextureCoordUOffset(InCollection, AttributeNames::TextureCoordUOffset, GroupNames::PointGroup)
+		, URange(InCollection, AttributeNames::URange, GroupNames::PointGroup)
 	{
 	}
 
@@ -46,6 +47,7 @@ namespace PV::Facades
 		, BudDevelopment(InCollection, AttributeNames::BudDevelopment, GroupNames::PointGroup)
 		, TextureCoordV(InCollection, AttributeNames::TextureCoordV, GroupNames::PointGroup)
 		, TextureCoordUOffset(InCollection, AttributeNames::TextureCoordUOffset, GroupNames::PointGroup)
+		, URange(InCollection, AttributeNames::URange, GroupNames::PointGroup)
 	{
 	}
 
@@ -319,6 +321,30 @@ namespace PV::Facades
 		}
 	}
 
+	const FVector2f& FPointFacade::GetURange(const int32 Index) const
+	{
+		if (URange.IsValid() && URange.IsValidIndex(Index))
+		{
+			return URange[Index];
+		}
+
+		static const FVector2f Range(0,1);
+		return Range;
+	}
+
+	void FPointFacade::SetURange(const int32 Index, const FVector2f InURange)
+	{
+		if (!URange.IsValid())
+		{
+			URange.Add();
+		}
+		
+		if (URange.IsValid() && URange.IsValidIndex(Index))
+		{
+			URange.ModifyAt(Index, InURange);
+		}
+	}
+
 	TManagedArray<FVector3f>& FPointFacade::ModifyPositions()
 	{
 		return Positions.Modify();
@@ -388,6 +414,10 @@ namespace PV::Facades
 			{
 				TextureCoordUOffset.ModifyAt(ToIndex, TextureCoordUOffset[FromIndex]);
 			}
+			if(URange.IsValid())
+			{
+				URange.ModifyAt(ToIndex, URange[FromIndex]);
+			}
 		}
 	}
 
diff --git a/Engine/Plugins/Experimental/ProceduralVegetationEditor/Source/ProceduralVegetation/Private/Implementations/PVMaterialSettings.cpp b/Engine/Plugins/Experimental/ProceduralVegetationEditor/Source/ProceduralVegetation/Private/Implementations/PVMaterialSettings.cpp
--- a/Engine/Plugins/Experimental/ProceduralVegetationEditor/Source/ProceduralVegetation/Private/Implementations/PVMaterialSettings.cpp
+++ b/Engine/Plugins/Experimental/ProceduralVegetationEditor/Source/ProceduralVegetation/Private/Implementations/PVMaterialSettings.cpp
@@ -302,12 +302,17 @@ void FPVMaterialSettings::SetOutputUVs(FManagedArrayCollection& Collection, TArr
 		float UVYTextureOffset = 0;
 		float UVXTextureOffset = 0;
 		float UVYTextureScale = 1;
+		FFloatRange URange(0, 1);
 		//Get all the branches with the material
-		if (MaterialSetups.IsValidIndex(BranchUVMaterial) && MaterialSetups[BranchUVMaterial].YTextureMode == EYTextureMode::Default)
+		if (MaterialSetups.IsValidIndex(BranchUVMaterial))
 		{
-			UVYTextureOffset = BranchUVData[BranchIndex].YOffset;
-			UVYTextureScale = BranchUVData[BranchIndex].YScale;
-			UVRatio *= 1;
+			URange = MaterialSetups[BranchUVMaterial].URange;
+			if (MaterialSetups[BranchUVMaterial].YTextureMode == EYTextureMode::Default)
+			{
+				UVYTextureOffset = BranchUVData[BranchIndex].YOffset;
+				UVYTextureScale = BranchUVData[BranchIndex].YScale;
+				UVRatio *= 1;
+			}
 		}
 
 		UVXTextureOffset = BranchUVData[BranchIndex].XOffset;
@@ -325,6 +330,7 @@ void FPVMaterialSettings::SetOutputUVs(FManagedArrayCollection& Collection, TArr
 
 			PointFacade.SeTextureCoordV(BranchPoints[j], UV.Y);
 			PointFacade.SetTextureCoordUOffset(BranchPoints[j], UV.X);
+			PointFacade.SetURange( BranchPoints[j], FVector2f(URange.GetLowerBound().GetValue(), URange.GetUpperBound().GetValue()));
 		}
 	}
 }
diff --git a/Engine/Plugins/Experimental/ProceduralVegetationEditor/Source/ProceduralVegetation/Private/Implementations/PVMaterialSettings.h b/Engine/Plugins/Experimental/ProceduralVegetationEditor/Source/ProceduralVegetation/Private/Implementations/PVMaterialSettings.h
--- a/Engine/Plugins/Experimental/ProceduralVegetationEditor/Source/ProceduralVegetation/Private/Implementations/PVMaterialSettings.h
+++ b/Engine/Plugins/Experimental/ProceduralVegetationEditor/Source/ProceduralVegetation/Private/Implementations/PVMaterialSettings.h
@@ -47,23 +47,26 @@ struct FTrunkGenerationMaterialSetup
 	UPROPERTY(EditAnywhere, Category="Materials")
 	TObjectPtr<UMaterialInterface> Material;
 
-	UPROPERTY(EditAnywhere, Category="Materials")
+	UPROPERTY(EditAnywhere, Category="Materials" , meta=(DisplayName = "Y Texture Mode"))
 	EYTextureMode YTextureMode = EYTextureMode::Default;
 
-	UPROPERTY(EditAnywhere, Category="Materials", meta=(EditCondition = "YTextureMode == EYTextureMode::Default"))
+	UPROPERTY(EditAnywhere, Category="Materials", meta=(DisplayName = "Y Scale", EditCondition = "YTextureMode == EYTextureMode::Default"))
 	float YScale = 1.0f;
 
-	UPROPERTY(EditAnywhere, Category="Materials", meta=(EditCondition = "YTextureMode == EYTextureMode::Default"))
+	UPROPERTY(EditAnywhere, Category="Materials", meta=(DisplayName = "Y Offset", EditCondition = "YTextureMode == EYTextureMode::Default"))
 	float YOffset = 0.0f;
 
-	UPROPERTY(EditAnywhere, Category="Materials", meta=(EditCondition = "YTextureMode == EYTextureMode::Default"))
+	UPROPERTY(EditAnywhere, Category="Materials", meta=(DisplayName = "Y Offset Random", EditCondition = "YTextureMode == EYTextureMode::Default"))
 	float YOffsetRandom = 0.0f;
 
 	UPROPERTY(EditAnywhere, Category="Materials")
 	float XOffset = 0.0f;
 
-	UPROPERTY(EditAnywhere, Category="Materials")
+	UPROPERTY(EditAnywhere, Category="Materials" , meta=(DisplayName = "X Offset Random"))
 	float XOffsetRandom = 0.0f;
+
+	UPROPERTY(EditAnywhere, Category="Materials", meta=(DisplayName = "X Range"))
+	FFloatRange URange = FFloatRange(0,1);
 	
 };
 
diff --git a/Engine/Plugins/Experimental/ProceduralVegetationEditor/Source/ProceduralVegetation/Private/Implementations/PVMesher.cpp b/Engine/Plugins/Experimental/ProceduralVegetationEditor/Source/ProceduralVegetation/Private/Implementations/PVMesher.cpp
--- a/Engine/Plugins/Experimental/ProceduralVegetationEditor/Source/ProceduralVegetation/Private/Implementations/PVMesher.cpp
+++ b/Engine/Plugins/Experimental/ProceduralVegetationEditor/Source/ProceduralVegetation/Private/Implementations/PVMesher.cpp
@@ -606,11 +606,17 @@ void FPVMesher::GenerateBranchMeshData(const bool bPrimitiveIsTrunk, const int32
 	for (int32 Index = 0; Index <= MaxIndex; ++Index)
 	{
 		const int32& PointIndex = PrimitivePoints[Index];
-		const float NjordPixelId = PointIndex;//PointFacade.GetNjordPixelIndex(PointIndex);
+		const float NjordPixelId = PointIndex;
 
 		float TextureCoordV = PointFacade.GetTextureCoordV(PointIndex);
 		float TextureCoordUOffset = PointFacade.GetTextureCoordUOffset(PointIndex);
+		FVector2f OldRange = FVector2f(0, 1.0f);
+		FVector2f NewRange = PointFacade.GetURange(PointIndex);
 
+		float Interval = NewRange.GetMax() - NewRange.GetMin();
+		float URatio = FMath::IsNearlyZero(Interval) ? 1.0f : 1.0f / Interval;
+		TextureCoordV /= URatio;
+		
 		const int32 PointDivisions = TargetMeshDivisions[PointIndex];
 		const FVector3f& PointPosition = PointPositions[PointIndex];
 		const float PointScale = PointScales[PointIndex];
@@ -638,7 +644,8 @@ void FPVMesher::GenerateBranchMeshData(const bool bPrimitiveIsTrunk, const int32
 			const float RotationDelta = UE_TWO_PI / PointDivisions;
 			for (int32 Division = 0; Division < PointDivisions; ++Division)
 			{
-				const float UV_U = (1 - (static_cast<float>(Division) / PointDivisions)) + TextureCoordUOffset;
+				float UV_U = (1 - (static_cast<float>(Division) / PointDivisions)) + TextureCoordUOffset;
+				UV_U = FMath::GetMappedRangeValueUnclamped(OldRange, NewRange, FMath::Wrap(UV_U, OldRange.X, OldRange.Y));
 
 				float ProfileMultiplier = 1.0f;
 				if (bShouldApplyProfile)
@@ -681,7 +688,7 @@ void FPVMesher::GenerateBranchMeshData(const bool bPrimitiveIsTrunk, const int32
 		}
 
 		const FLocalDynamicMeshData::FVertex FirstVertex = OutLocalMeshData.Vertices[CurrentIndices[0]];
-		CurrentIndices.Add(OutLocalMeshData.Vertices.Emplace(FirstVertex.Position, FirstVertex.Normal, FVector2f(TextureCoordUOffset, TextureCoordV),
+		CurrentIndices.Add(OutLocalMeshData.Vertices.Emplace(FirstVertex.Position, FirstVertex.Normal, FVector2f(FMath::GetMappedRangeValueUnclamped(OldRange,NewRange, FMath::Wrap(TextureCoordUOffset, OldRange.X, OldRange.Y)), TextureCoordV),
 			NjordPixelId));
 
 		PreviousPosition = PointPosition;
diff --git a/Engine/Plugins/Experimental/ProceduralVegetationEditor/Source/ProceduralVegetation/Public/Facades/PVAttributesNames.h b/Engine/Plugins/Experimental/ProceduralVegetationEditor/Source/ProceduralVegetation/Public/Facades/PVAttributesNames.h
--- a/Engine/Plugins/Experimental/ProceduralVegetationEditor/Source/ProceduralVegetation/Public/Facades/PVAttributesNames.h
+++ b/Engine/Plugins/Experimental/ProceduralVegetationEditor/Source/ProceduralVegetation/Public/Facades/PVAttributesNames.h
@@ -41,6 +41,7 @@ namespace PV::AttributeNames
 	inline static const FName BudDevelopment = FName(TEXT("budDevelopment"));
 	inline static const FName TextureCoordV = FName(TEXT("uv_v"));
 	inline static const FName TextureCoordUOffset = FName(TEXT("uv_uOffset"));
+	inline static const FName URange = FName(TEXT("uv_uRange"));
 	
 	inline static const FName HullGradient = FName(TEXT("LOD_hullGradient"));
 	inline static const FName MainTrunkGradient = FName(TEXT("LOD_mainTrunkGradient"));
diff --git a/Engine/Plugins/Experimental/ProceduralVegetationEditor/Source/ProceduralVegetation/Public/Facades/PVPointFacade.h b/Engine/Plugins/Experimental/ProceduralVegetationEditor/Source/ProceduralVegetation/Public/Facades/PVPointFacade.h
--- a/Engine/Plugins/Experimental/ProceduralVegetationEditor/Source/ProceduralVegetation/Public/Facades/PVPointFacade.h
+++ b/Engine/Plugins/Experimental/ProceduralVegetationEditor/Source/ProceduralVegetation/Public/Facades/PVPointFacade.h
@@ -67,6 +67,10 @@ namespace PV::Facades
 		const float& GetTextureCoordUOffset(int32 Index) const;
 		
 		void SetTextureCoordUOffset(int32 Index, float InTextureCoordUOffset);
+		
+		const FVector2f& GetURange(int32 Index) const;
+		
+		void SetURange(int32 Index, FVector2f InURange);
 
 		TManagedArray<FVector3f>& ModifyPositions();
 		
@@ -107,5 +111,6 @@ namespace PV::Facades
 		TManagedArrayAccessor<TArray<int>> BudDevelopment;
 		TManagedArrayAccessor<float> TextureCoordV;
 		TManagedArrayAccessor<float> TextureCoordUOffset;
+		TManagedArrayAccessor<FVector2f> URange;
 	};
 }
