Fix loctext key/namespace collision in PVE

diff --git a/Engine/Plugins/Experimental/ProceduralVegetationEditor/Source/ProceduralVegetation/Private/Nodes/PVPCGRemoveBranches.h b/Engine/Plugins/Experimental/ProceduralVegetationEditor/Source/ProceduralVegetation/Private/Nodes/PVPCGRemoveBranches.h
--- a/Engine/Plugins/Experimental/ProceduralVegetationEditor/Source/ProceduralVegetation/Private/Nodes/PVPCGRemoveBranches.h
+++ b/Engine/Plugins/Experimental/ProceduralVegetationEditor/Source/ProceduralVegetation/Private/Nodes/PVPCGRemoveBranches.h
@@ -14,8 +14,8 @@ class UPVPCGRemoveBranchesSettings : public UPVPCGBaseSettings
 public:
 #if WITH_EDITOR
 	virtual FName GetDefaultNodeName() const override { return FName(TEXT("ProceduralVegetationRemoveBranches")); }
-	virtual FText GetDefaultNodeTitle() const override { return NSLOCTEXT("PVPCGScale", "NodeTitle", "Remove Branches"); }
-	virtual FText GetNodeTooltipText() const override { return NSLOCTEXT("PVPCGScale", "NodeTooltip", "Remove branches on the basis of selected criteria"); }
+	virtual FText GetDefaultNodeTitle() const override { return NSLOCTEXT("PVPCGRemoveBranches", "NodeTitle", "Remove Branches"); }
+	virtual FText GetNodeTooltipText() const override { return NSLOCTEXT("PVPCGRemoveBranches", "NodeTooltip", "Remove branches on the basis of selected criteria"); }
 	virtual FLinearColor GetNodeTitleColor() const override { return FLinearColor::Yellow; }
 
 	virtual TArray<EPVRenderType> GetRenderType() const override { return TArray{ EPVRenderType::Skeleton }; }
diff --git a/Engine/Plugins/Experimental/ProceduralVegetationEditor/Source/ProceduralVegetationEditor/Private/ProceduralVegetationEditorModule.cpp b/Engine/Plugins/Experimental/ProceduralVegetationEditor/Source/ProceduralVegetationEditor/Private/ProceduralVegetationEditorModule.cpp
--- a/Engine/Plugins/Experimental/ProceduralVegetationEditor/Source/ProceduralVegetationEditor/Private/ProceduralVegetationEditorModule.cpp
+++ b/Engine/Plugins/Experimental/ProceduralVegetationEditor/Source/ProceduralVegetationEditor/Private/ProceduralVegetationEditorModule.cpp
@@ -68,8 +68,8 @@ void FProceduralVegetationEditorModule::RegisterMenus()
 				if (const UContentBrowserAssetContextMenuContext* Context = UContentBrowserAssetContextMenuContext::FindContextWithAssets(InSection))
 				{
 					{
-						const TAttribute<FText> Label = LOCTEXT("Mesh_OpenInProceduralVegetationEditor", "Open in Procedural Vegetation Editor");
-						const TAttribute<FText> ToolTip = LOCTEXT("Mesh_OpenInProceduralVegetationEditor", "Open the mesh in the Procedural Vegetation Editor.");
+						const TAttribute<FText> Label = LOCTEXT("Mesh_OpenInProceduralVegetationEditorLabel", "Open in Procedural Vegetation Editor");
+						const TAttribute<FText> ToolTip = LOCTEXT("Mesh_OpenInProceduralVegetationEditorToolTip", "Open the mesh in the Procedural Vegetation Editor.");
 						const FSlateIcon Icon = FSlateIcon(FAppStyle::GetAppStyleSetName(), "Icons.Edit");
 
 						FToolUIAction UIAction;
