[PVE] Fix issue where foliage scale is being ignored

diff --git a/Engine/Plugins/Experimental/ProceduralVegetationEditor/Source/ProceduralVegetation/Private/Facades/PVFoliageFacade.cpp b/Engine/Plugins/Experimental/ProceduralVegetationEditor/Source/ProceduralVegetation/Private/Facades/PVFoliageFacade.cpp
--- a/Engine/Plugins/Experimental/ProceduralVegetationEditor/Source/ProceduralVegetation/Private/Facades/PVFoliageFacade.cpp
+++ b/Engine/Plugins/Experimental/ProceduralVegetationEditor/Source/ProceduralVegetation/Private/Facades/PVFoliageFacade.cpp
@@ -270,7 +270,7 @@ namespace PV::Facades
 			ScalesAttribute.ModifyAt(Index, Input);
 		}
 	}
-	
+
 	void FFoliageFacade::SetFoliageName(const int32 Index, const FString& Input)
 	{
 		check(!IsConst());
@@ -279,7 +279,7 @@ namespace PV::Facades
 			NamesAttribute.ModifyAt(Index, Input);
 		}
 	}
-	
+
 	void FFoliageFacade::SetFoliageNames(const TArray<FString>& InputNames)
 	{
 		check(!IsConst());
@@ -354,10 +354,11 @@ namespace PV::Facades
 		float Angle = FMath::Acos(FMath::Clamp(Dot, -1.0f, 1.0f));
 		const FQuat FaceRotationQuat = FQuat(FVector(Data.UpVector), Angle);
 
+		const FTransform ScaleTransform(FQuat::Identity, FVector::ZeroVector, FVector(Data.Scale));
 		const FTransform FaceRotationTransform(FaceRotationQuat);
 		const FTransform UpVectorRotationTransform(UpVectorRotationQuat);
 		const FTransform TranslationTransform(FRotator::ZeroRotator, FVector(Data.PivotPoint));
-		const FTransform FinalTransform = UpVectorRotationTransform * FaceRotationTransform * TranslationTransform;
+		const FTransform FinalTransform = ScaleTransform * UpVectorRotationTransform * FaceRotationTransform * TranslationTransform;
 
 		return FinalTransform;
 	}
