[ProceduralVegetationEditor] Fix for crash when branches are removed.

diff --git a/Engine/Plugins/Experimental/ProceduralVegetationEditor/Source/ProceduralVegetation/Private/Facades/PVTreeFacade.cpp b/Engine/Plugins/Experimental/ProceduralVegetationEditor/Source/ProceduralVegetation/Private/Facades/PVTreeFacade.cpp
--- a/Engine/Plugins/Experimental/ProceduralVegetationEditor/Source/ProceduralVegetation/Private/Facades/PVTreeFacade.cpp
+++ b/Engine/Plugins/Experimental/ProceduralVegetationEditor/Source/ProceduralVegetation/Private/Facades/PVTreeFacade.cpp
@@ -101,7 +101,7 @@ namespace PV::Facades
 		FPointFacade PointFacadeOut(OutCollection);
 		FFoliageFacade FoliageFacadeOut(OutCollection);
 
-		for (int32 i = InBranchesToRemove.Num() - 1; i > 0 ; i--)
+		for (int32 i = InBranchesToRemove.Num() - 1; i >= 0 ; i--)
 		{
 			GatherChildBranches(InBranchFacade, InBranchesToRemove[i], InBranchesToRemove);
 		}
diff --git a/Engine/Plugins/Experimental/ProceduralVegetationEditor/Source/ProceduralVegetation/Private/Implementations/PVMesher.cpp b/Engine/Plugins/Experimental/ProceduralVegetationEditor/Source/ProceduralVegetation/Private/Implementations/PVMesher.cpp
--- a/Engine/Plugins/Experimental/ProceduralVegetationEditor/Source/ProceduralVegetation/Private/Implementations/PVMesher.cpp
+++ b/Engine/Plugins/Experimental/ProceduralVegetationEditor/Source/ProceduralVegetation/Private/Implementations/PVMesher.cpp
@@ -22,6 +22,14 @@ void FPVMesher::GenerateGeometryCollection(const FManagedArrayCollection& InSkel
 	PV::Facades::FBoneFacade::DefineSchema(OutGeometryCollection);
 	PV::Facades::FBoneFacade BoneFacade = PV::Facades::FBoneFacade(OutGeometryCollection);
 
+	PV::Facades::FBranchFacade BranchFacade = PV::Facades::FBranchFacade(OutGeometryCollection);
+
+	if (BranchFacade.GetElementCount() == 0)
+	{
+		UE_LOG(LogProceduralVegetation, Log, TEXT("No branch data available for generating mesh."));
+		return;
+	}
+	
 	SetNjordPixelID(OutGeometryCollection);
 
 	MesherSettings.MaterialSettings.ApplyMaterialSettings(OutGeometryCollection);
