[ProceduralVegetationEditor] Fixed unused branches getting loaded and converted for export,
Fixed Incorrect bones
Fixed wind data not build for non nanite assemblies export

diff --git a/Engine/Plugins/Experimental/ProceduralVegetationEditor/Source/ProceduralVegetation/Private/Facades/PVBoneFacade.cpp b/Engine/Plugins/Experimental/ProceduralVegetationEditor/Source/ProceduralVegetation/Private/Facades/PVBoneFacade.cpp
--- a/Engine/Plugins/Experimental/ProceduralVegetationEditor/Source/ProceduralVegetation/Private/Facades/PVBoneFacade.cpp
+++ b/Engine/Plugins/Experimental/ProceduralVegetationEditor/Source/ProceduralVegetation/Private/Facades/PVBoneFacade.cpp
@@ -101,13 +101,15 @@ namespace PV::Facades
 		return ParentBranchBones;
 	}
 
-	FBoneNode* FindBoneForPoint(const FPointFacade& PointFacade, TArray<FBoneNode>& BoneNodes, int32 BonePoint, float MaxDistance = 1.0)
+	FBoneNode* FindBoneForPoint(const FPointFacade& PointFacade, TArray<FBoneNode>& BoneNodes, int32 BonePoint, FString OutMessage, bool bFallbackToClosestBone = false, float MaxDistance = 1.0)
 	{
 		const TManagedArray<FVector3f>& Positions = PointFacade.GetPositions();
 		FVector BonePosition = FVector(Positions[BonePoint]);
 		float BonePScale = PointFacade.GetPointScale(BonePoint);
 				
 		FBoneNode* FoundBone = nullptr;
+		FBoneNode* ClosestBone = nullptr;
+		float MinDist = FLT_MAX;
 				
 		for (int32 NodeIdx = 0; NodeIdx < BoneNodes.Num(); ++NodeIdx)
 		{
@@ -119,12 +121,24 @@ namespace PV::Facades
 			}
 
 			float Dist = FMath::PointDistToSegment(BonePosition, NodePosition, NextPosition);
+			MinDist = FMath::Min(MinDist, Dist);
 			if (Dist <= BonePScale * 2)
 			{
 				FoundBone = &BoneNodes[NodeIdx];
 			}
+
+			if (Dist == MinDist)
+			{
+				ClosestBone = &BoneNodes[NodeIdx];
+			}
 		}
 
+		if (!FoundBone && bFallbackToClosestBone)
+		{
+			OutMessage = FString::Format(TEXT("FindBoneForPoint fallback to closest bone for point index {0}") , {BonePoint});
+			return ClosestBone;
+		}
+		
 		return FoundBone;
 	}
 
@@ -274,13 +288,19 @@ namespace PV::Facades
 			check(BranchPoints.IsValidIndex(0));
 			int32 ParentIndex = BranchFacade.GetParentBranchIndices(BranchIndex).Last();
 			TArray<FBoneNode> ParentBoneNodes = GetParentBranchBones(Bones, ParentIndex);
-			
-			FBoneNode* FoundBoneNode = FindBoneForPoint(PointFacade, ParentBoneNodes, BranchPoints[0]);
+
+			FString Message;
+			FBoneNode* FoundBoneNode = FindBoneForPoint(PointFacade, ParentBoneNodes, BranchPoints[0], Message, true);
 
 			if (FoundBoneNode == nullptr && i != 0)
 			{
 				UE_LOG(LogProceduralVegetation, Log, TEXT("Parent bone not found for branch (%i) ParentIndex {%i}"), BranchIndex, ParentIndex);
 			}
+
+			if (!Message.IsEmpty())
+			{
+				UE_LOG(LogProceduralVegetation, Log, TEXT("%s Branch (%i) ParentIndex {%i}"), *Message, BranchIndex, ParentIndex);
+			}
 			
 			ParentBoneIndex = FoundBoneNode ? FoundBoneNode->BoneIndex : ParentBoneIndex;
 			PreviousPosition = FoundBoneNode ? FVector(FoundBoneNode->AbsolutePosition) : RootPointPosition;
diff --git a/Engine/Plugins/Experimental/ProceduralVegetationEditor/Source/ProceduralVegetation/Private/Helpers/PVExportHelper.cpp b/Engine/Plugins/Experimental/ProceduralVegetationEditor/Source/ProceduralVegetation/Private/Helpers/PVExportHelper.cpp
--- a/Engine/Plugins/Experimental/ProceduralVegetationEditor/Source/ProceduralVegetation/Private/Helpers/PVExportHelper.cpp
+++ b/Engine/Plugins/Experimental/ProceduralVegetationEditor/Source/ProceduralVegetation/Private/Helpers/PVExportHelper.cpp
@@ -22,7 +22,6 @@
 
 #include "Engine/SkeletalMesh.h"
 #include "Engine/StaticMesh.h"
-#include "Engine/World.h"
 #include "Facades/PVBranchFacade.h"
 #include "Facades/PVPointFacade.h"
 #include "GeometryCollection/GeometryCollection.h"
@@ -34,7 +33,6 @@
 #include "UObject/Package.h"
 #include "ProceduralVegetation.h"
 #include "GeometryCollection/Facades/CollectionConstraintOverrideFacade.h"
-#include "Rendering/RenderCommandPipes.h"
 
 namespace PV::Export::Internal
 {
@@ -280,6 +278,34 @@ namespace PV::Export::Internal
 		}
 	}
 
+	bool ContainsFoliageName(const Facades::FFoliageFacade& FoliageFacade, const FString& InFoliageName)
+	{
+		int32 NumInstances = FoliageFacade.NumFoliageEntries();
+			
+		for(int32 Id = 0; Id < NumInstances; Id++)
+		{
+			FString FoliageName = GetFoliageName(FoliageFacade, Id);
+			if(FoliageName == InFoliageName)
+			{
+				return true;
+			}
+		}
+
+		return false;
+	}
+
+	void GetUsedFoliage(const Facades::FFoliageFacade& FoliageFacade, TArray<FString>& OutFoliageNames)
+	{
+		for(int32 FoliageNameIndex = 0; FoliageNameIndex < FoliageFacade.NumFoliageNames(); ++FoliageNameIndex)
+		{
+			FString FoliageName = FoliageFacade.GetFoliageName(FoliageNameIndex);
+			if (ContainsFoliageName(FoliageFacade, FoliageName))
+			{
+				OutFoliageNames.Add(FoliageName);
+			}
+		}
+	}
+
 #if WITH_EDITOR
 	void BuildNaniteAssemblyData(
 		const FManagedArrayCollection& Collection,
@@ -297,10 +323,12 @@ namespace PV::Export::Internal
 		
 		TMap<FString, int32> MeshNamePartMap;
 		TMap<FString, TArray<TObjectPtr<UMaterialInterface>>> MeshMaterialsMap;
+
+		TArray<FString> UsedFoliage;
+		GetUsedFoliage(FoliageFacade, UsedFoliage);
 		
-		for(int32 FoliageNameIndex = 0; FoliageNameIndex < FoliageFacade.NumFoliageNames(); ++FoliageNameIndex)
+		for (const FString& FoliageName : UsedFoliage)
 		{
-			FString FoliageName = FoliageFacade.GetFoliageName(FoliageNameIndex);
 			TObjectPtr<UStaticMesh> PartStaticMesh = LoadObject<UStaticMesh>(nullptr, *FoliageName);
 			TObjectPtr<USkeletalMesh> PartSkeletalMesh = LoadObject<USkeletalMesh>(nullptr, *FoliageName);
 
@@ -382,10 +410,12 @@ namespace PV::Export::Internal
 
 		TMap<FString, int32> MeshNamePartMap;
 		TMap<FString, TArray<TObjectPtr<UMaterialInterface>>> MeshMaterialsMap;
+
+		TArray<FString> UsedFoliage;
+		GetUsedFoliage(FoliageFacade, UsedFoliage);
 			
-		for(int32 FoliageNameIndex = 0; FoliageNameIndex < FoliageFacade.NumFoliageNames(); ++FoliageNameIndex)
+		for (const FString& FoliageName : UsedFoliage)
 		{
-			FString FoliageName = FoliageFacade.GetFoliageName(FoliageNameIndex);
 			TObjectPtr<UStaticMesh> PartStaticMesh = LoadObject<UStaticMesh>(nullptr, *FoliageName);
 			TObjectPtr<USkeletalMesh> PartSkeletalMesh = LoadObject<USkeletalMesh>(nullptr, *FoliageName);
 
@@ -395,10 +425,10 @@ namespace PV::Export::Internal
 				auto SkeletalMeshPath = PartStaticMesh.GetPath().Replace(*PartStaticMesh.GetName(), *SkeletalMeshName);
 				int32 PartIndex = AssemblyBuilder.FindOrAddPart(SkeletalMeshPath);
 				MeshNamePartMap.FindOrAdd(FoliageName) = PartIndex;
-		
+	
 				TArray<FAssetData> AssetData;
 				IAssetRegistry::Get()->GetAssetsByPath(*PartStaticMesh.GetPackage()->GetName(), AssetData);
-		
+	
 				if(AssetData.Num() <= 0)
 				{
 					FDynamicMesh3 PartDynamicMesh;
@@ -767,22 +797,22 @@ namespace PV::Export::Internal
 			{
 				//Build the nanite assemblies
 				BuildNaniteAssemblyData(ExportedCollection, ExportMesh);
-
-				FDynamicWindSkeletalImportData ImportData;
-				BuildWindImportData(ImportData,ExportedCollection, BoneNodes);
-				
-				if (!DynamicWind::ImportSkeletalData(*ExportMesh, ImportData))
-				{
-					UE_LOG(LogTemp, Warning, TEXT("Failed to build wind data"));
-				}
-				
-				ExportMesh->Build();
 			}
 			else
 			{
 				UE_LOG(LogProceduralVegetation, Warning, TEXT("Failed to build Nanite Assemblies, because neither Nanite Foliage nor Nanite Assemblies are enabled for the project."));
 			}
 		}
+
+		FDynamicWindSkeletalImportData ImportData;
+		BuildWindImportData(ImportData,ExportedCollection, BoneNodes);
+				
+		if (!DynamicWind::ImportSkeletalData(*ExportMesh, ImportData))
+		{
+			UE_LOG(LogTemp, Warning, TEXT("Failed to build wind data"));
+		}
+				
+		ExportMesh->Build();
 #endif
 	}
 
