Add support for multiseed plants

diff --git a/Engine/Plugins/Experimental/ProceduralVegetationEditor/Source/ProceduralVegetation/Private/Facades/PVBoneFacade.cpp b/Engine/Plugins/Experimental/ProceduralVegetationEditor/Source/ProceduralVegetation/Private/Facades/PVBoneFacade.cpp
--- a/Engine/Plugins/Experimental/ProceduralVegetationEditor/Source/ProceduralVegetation/Private/Facades/PVBoneFacade.cpp
+++ b/Engine/Plugins/Experimental/ProceduralVegetationEditor/Source/ProceduralVegetation/Private/Facades/PVBoneFacade.cpp
@@ -4,6 +4,7 @@
 #include "ProceduralVegetationModule.h"
 #include "Facades/PVAttributesNames.h"
 #include "Facades/PVBranchFacade.h"
+#include "Facades/PVPlantFacade.h"
 #include "Facades/PVPointFacade.h"
 #include "GeometryCollection/Facades/CollectionConstraintOverrideFacade.h"
 
@@ -249,139 +250,147 @@ namespace PV::Facades
 	{
 		FPointFacade PointFacade(VertexPointIds.GetConstCollection());
 		FBranchFacade BranchFacade(VertexPointIds.GetConstCollection());
+		FPlantFacade PlantFacade(VertexPointIds.GetConstCollection());
 
 		SetWindSimulationData(*VertexPointIds.GetCollection());
 		
 		const TManagedArray<FVector3f>& Positions = PointFacade.GetPositions();
-		TArray<int32> SortedBranchIndices;
-		BranchFacade.GetSortedBranchIndicesByHierarchy(SortedBranchIndices);
-		
-		if (SortedBranchIndices.Num() == 0)
-		{
-			UE_LOG(LogProceduralVegetation, Log, TEXT("Cannot create bones, no branches found"));
-			return;
-		}
 
-		int32 RootBranch = SortedBranchIndices[0];
-		
-		check(BranchFacade.GetPoints(RootBranch).Num() > 0);
-		const int32 RootPointIndex = BranchFacade.GetPoints(RootBranch)[0];
-		
-		check(Positions.Num() > 0 && Positions.Num() >= RootPointIndex);
-		const FVector RootPointPosition = FVector(Positions[RootPointIndex]);
-		
-		FBoneNode RootBoneNode{ "Root",0, INDEX_NONE, RootPointIndex, RootBranch, 1.0, FTransform{RootPointPosition}};
-		Bones.Emplace(MoveTemp(RootBoneNode));
+		constexpr int RootBoneIndex = 0;
+		FBoneNode RootNode{ "Root",RootBoneIndex, INDEX_NONE, INDEX_NONE, INDEX_NONE, 1.0, FTransform{FVector::Zero()}};
+		Bones.Emplace(RootNode);
 
-		int ParentBoneIndex = 0;
-		int BoneIndex = 1;
-		FVector PreviousPosition = RootPointPosition;
-		
-		for (int32 i = 0; i < SortedBranchIndices.Num(); ++i)
+		int32 CurrentBoneIndex = 1;
+		for (const TMap<int32, int32> PlantNumbersToTrunkIDs = PlantFacade.GetPlantNumbersToTrunkIndicesMap();
+		 const TPair<int32, int32> Pair : PlantNumbersToTrunkIDs)
 		{
-			int32 BranchIndex = SortedBranchIndices[i];
-			int32 BranchNumber = BranchFacade.GetBranchNumber(BranchIndex);
-			UE_LOG(LogProceduralVegetation, Log, TEXT("Sorted i{%i} Branch Index{%i} Branch Number{%i}"), i, BranchIndex, BranchNumber);
-			
-			auto BranchPoints = BranchFacade.GetPoints(BranchIndex);
-			
-			check(BranchPoints.IsValidIndex(0));
-			int32 ParentIndex = BranchFacade.GetParentBranchIndices(BranchIndex).Last();
-			TArray<FBoneNode> ParentBoneNodes = GetParentBranchBones(Bones, ParentIndex);
+			const int32 PlantNumber = Pair.Key;
+			const int32 TrunkIndex = Pair.Value;
 
-			FString Message;
-			FBoneNode* FoundBoneNode = FindBoneForPoint(PointFacade, ParentBoneNodes, BranchPoints[0], Message, true);
+			check(BranchFacade.GetPoints(TrunkIndex).Num() > 0);
+			const int32 RootPointIndex = BranchFacade.GetPoints(TrunkIndex)[0];
+		
+			check(Positions.Num() > 0 && Positions.Num() >= RootPointIndex);
+			const FVector RootPointPosition = FVector(Positions[RootPointIndex]);
+			FBoneNode TrunkBoneNode{ "Root",CurrentBoneIndex, RootBoneIndex, RootPointIndex, TrunkIndex, 1.0, FTransform{RootPointPosition}};
+			CurrentBoneIndex++;
+			Bones.Emplace(MoveTemp(TrunkBoneNode));
 
-			if (FoundBoneNode == nullptr && i != 0)
-			{
-				UE_LOG(LogProceduralVegetation, Log, TEXT("Parent bone not found for branch (%i) ParentIndex {%i}"), BranchIndex, ParentIndex);
-			}
+			int ParentBoneIndex = TrunkBoneNode.BoneIndex;
+			FVector PreviousPosition = RootPointPosition;
 
-			if (!Message.IsEmpty())
+			TArray<int32> SortedBranchIndices = PlantFacade.GetBranchIndicesSortedByHierarchyNumber(PlantNumber);
+			if (SortedBranchIndices.Num() == 0)
 			{
-				UE_LOG(LogProceduralVegetation, Log, TEXT("%s Branch (%i) ParentIndex {%i}"), *Message, BranchIndex, ParentIndex);
+				UE_LOG(LogProceduralVegetation, Log, TEXT("Cannot create bones, no branches found for Plant Number {%i}"), PlantNumber);
+				continue;
 			}
+			for (int32 i = 0; i < SortedBranchIndices.Num(); ++i)
+			{
+				int32 BranchIndex = SortedBranchIndices[i];
+				int32 BranchNumber = BranchFacade.GetBranchNumber(BranchIndex);
+				UE_LOG(LogProceduralVegetation, Log, TEXT("Sorted i{%i} Branch Index{%i} Branch Number{%i}"), i, BranchIndex, BranchNumber);
+			
+				auto BranchPoints = BranchFacade.GetPoints(BranchIndex);
 			
-			ParentBoneIndex = FoundBoneNode ? FoundBoneNode->BoneIndex : ParentBoneIndex;
-			PreviousPosition = FoundBoneNode ? FVector(FoundBoneNode->AbsolutePosition) : RootPointPosition;
+				check(BranchPoints.IsValidIndex(0));
+				int32 ParentIndex = BranchFacade.GetParentBranchIndices(BranchIndex).Last();
+				TArray<FBoneNode> ParentBoneNodes = GetParentBranchBones(Bones, ParentIndex);
+			
+				FString Message;
+				FBoneNode* FoundBoneNode = FindBoneForPoint(PointFacade, ParentBoneNodes, BranchPoints[0], Message, true);
 
-			TArray<int32> ConnectingPoints;
+				if (FoundBoneNode == nullptr && i != 0)
+				{
+					UE_LOG(LogProceduralVegetation, Log, TEXT("Parent bone not found for branch (%i) ParentIndex {%i}"), BranchIndex, ParentIndex);
+				}
 
-			TArray<FConnectingSegment> ConnectingSegments;
-			FConnectingSegment ConnectingSegment;
-			
-			for (int Idx = 0; Idx < BranchPoints.Num(); Idx++)
-			{
-				const int32 PointIdx = BranchPoints[Idx];
-				check(Positions.Num() > PointIdx && PointIdx >= 0);
-				const FVector PointPosition = FVector(Positions[PointIdx]);
-				const float PScale = PointFacade.GetPointScale(PointIdx);
-				
-				FVector NextPosition = PointPosition;
-				if (BranchPoints.IsValidIndex(Idx + 1))
+				if (!Message.IsEmpty())
 				{
-					NextPosition = FVector(BranchPoints[Idx + 1]);
+					UE_LOG(LogProceduralVegetation, Log, TEXT("%s Branch (%i) ParentIndex {%i}"), *Message, BranchIndex, ParentIndex);
 				}
-				
-				float NjordPixelID = PointFacade.GetNjordPixelIndex(PointIdx);
+			
+				ParentBoneIndex = FoundBoneNode ? FoundBoneNode->BoneIndex : ParentBoneIndex;
+				PreviousPosition = FoundBoneNode ? FVector(FoundBoneNode->AbsolutePosition) : RootPointPosition;
+
+				TArray<int32> ConnectingPoints;
 
-				bool IsConnectingPoint = false;
-				for (auto ChildBranch : BranchFacade.GetChildren(BranchIndex))
+				TArray<FConnectingSegment> ConnectingSegments;
+				FConnectingSegment ConnectingSegment;
+			
+				for (int Idx = 0; Idx < BranchPoints.Num(); Idx++)
 				{
-					const int32 ChildBranchIndex = BranchFacade.GetBranchNumbers().Find(ChildBranch);
-					if (ChildBranchIndex != INDEX_NONE)
+					const int32 PointIdx = BranchPoints[Idx];
+					check(Positions.Num() > PointIdx && PointIdx >= 0);
+					const FVector PointPosition = FVector(Positions[PointIdx]);
+					const float PScale = PointFacade.GetPointScale(PointIdx);
+				
+					FVector NextPosition = PointPosition;
+					if (BranchPoints.IsValidIndex(Idx + 1))
 					{
-						const TArray<int32>& ChildPoints = BranchFacade.GetPoints(ChildBranchIndex);
-						check(ChildPoints.Num() > 0);
-						FVector FirstPointPosition = FVector(Positions[ChildPoints[0]]);
+						NextPosition = FVector(BranchPoints[Idx + 1]);
+					}
+				
+					float NjordPixelID = PointFacade.GetNjordPixelIndex(PointIdx);
 
-						float PointDistToSegment = FMath::PointDistToSegment(FirstPointPosition, PointPosition, NextPosition);
-						if (PointDistToSegment <= PScale && !ConnectingPoints.Contains(ChildPoints[0]))
+					bool IsConnectingPoint = false;
+					for (auto ChildBranch : BranchFacade.GetChildren(BranchIndex))
+					{
+						const int32 ChildBranchIndex = BranchFacade.GetBranchNumbers().Find(ChildBranch);
+						if (ChildBranchIndex != INDEX_NONE)
 						{
-							ConnectingPoints.Add(ChildPoints[0]);
-							IsConnectingPoint = true;
+							const TArray<int32>& ChildPoints = BranchFacade.GetPoints(ChildBranchIndex);
+							check(ChildPoints.Num() > 0);
+							FVector FirstPointPosition = FVector(Positions[ChildPoints[0]]);
+
+							float PointDistToSegment = FMath::PointDistToSegment(FirstPointPosition, PointPosition, NextPosition);
+							if (PointDistToSegment <= PScale && !ConnectingPoints.Contains(ChildPoints[0]))
+							{
+								ConnectingPoints.Add(ChildPoints[0]);
+								IsConnectingPoint = true;
+							}
+						}
+						else
+						{
+							UE_LOG(LogProceduralVegetation, Warning, TEXT("Child branch number {%i} not found in facade, for branch {%i}"), ChildBranch, BranchFacade.GetBranchNumber(BranchIndex));
 						}
 					}
-					else
-					{
-						UE_LOG(LogProceduralVegetation, Warning, TEXT("Child branch number {%i} not found in facade, for branch {%i}"), ChildBranch, BranchFacade.GetBranchNumber(BranchIndex));
-					}
-				}
 
-				bool IsLastPoint = !IsConnectingPoint && !(FMath::Frac(NjordPixelID) <= 0.0001f) && Idx == BranchPoints.Num() - 1;
-				//If this point exist in the BoneNodes set it as parent Bone
+					bool IsLastPoint = !IsConnectingPoint && !(FMath::Frac(NjordPixelID) <= 0.0001f) && Idx == BranchPoints.Num() - 1;
+					//If this point exist in the BoneNodes set it as parent Bone
 				
-				//UE_LOG(LogProceduralVegetation, Warning, TEXT("Iterating Branch Points NjordPixelId, Id{%f} PointIndex{%i} PointFacade.GetBudNumber(BranchPoints[PointIndex]){%i} BranchID{%i} ParentBoneIndex{%i} bFound{%i} FMath::Frac(NjordPixelID){%f}"),
-				//	NjordPixelID, PointIdx, PointFacade.GetBudNumber(PointIdx), BranchIndex, ParentBoneIndex, FoundBoneNode ? 1 : 0 , FMath::Frac(NjordPixelID));
+					//UE_LOG(LogProceduralVegetation, Warning, TEXT("Iterating Branch Points NjordPixelId, Id{%f} PointIndex{%i} PointFacade.GetBudNumber(BranchPoints[PointIndex]){%i} BranchID{%i} ParentBoneIndex{%i} bFound{%i} FMath::Frac(NjordPixelID){%f}"),
+					//	NjordPixelID, PointIdx, PointFacade.GetBudNumber(PointIdx), BranchIndex, ParentBoneIndex, FoundBoneNode ? 1 : 0 , FMath::Frac(NjordPixelID));
 				
-				if ((FMath::Frac(NjordPixelID) <= 0.0001f || IsConnectingPoint || IsLastPoint))
-				{
-					FBoneNode BoneNode;
-					BoneNode.BoneName = *("Bone_" + FString::FromInt(BoneIndex));
-					BoneNode.BoneIndex = BoneIndex;
-					BoneNode.PointIndex = PointIdx;
-					BoneNode.ParentBoneIndex = ParentBoneIndex;
-					BoneNode.BranchIndex = BranchIndex;
-					BoneNode.NjordPixelID = NjordPixelID;
-					FVector Position = FVector(PointPosition) - PreviousPosition;
-					BoneNode.AbsolutePosition = PointPosition;
-					BoneNode.BoneTransform = FTransform{Position};
-					//Bones.Emplace(MoveTemp(BoneNode));
-					ParentBoneIndex = BoneIndex;
-					PreviousPosition = FVector(PointPosition);
-					BoneIndex++;
-
-					int AddedNodeIndex = ConnectingSegment.Nodes.Emplace(MoveTemp(BoneNode));
-					if (IsConnectingPoint || IsLastPoint)
+					if ((FMath::Frac(NjordPixelID) <= 0.0001f || IsConnectingPoint || IsLastPoint))
 					{
-						int AddedIndex = ConnectingSegments.Add(MoveTemp(ConnectingSegment));
-						ConnectingSegment = FConnectingSegment();
+						FBoneNode BoneNode;
+						BoneNode.BoneName = *("Bone_" + FString::FromInt(CurrentBoneIndex));
+						BoneNode.BoneIndex = CurrentBoneIndex;
+						BoneNode.PointIndex = PointIdx;
+						BoneNode.ParentBoneIndex = ParentBoneIndex;
+						BoneNode.BranchIndex = BranchIndex;
+						BoneNode.NjordPixelID = NjordPixelID;
+						FVector Position = FVector(PointPosition) - PreviousPosition;
+						BoneNode.AbsolutePosition = PointPosition;
+						BoneNode.BoneTransform = FTransform{Position};
+						//Bones.Emplace(MoveTemp(BoneNode));
+						ParentBoneIndex = CurrentBoneIndex;
+						PreviousPosition = FVector(PointPosition);
+						CurrentBoneIndex++;
+
+						int AddedNodeIndex = ConnectingSegment.Nodes.Emplace(MoveTemp(BoneNode));
+						if (IsConnectingPoint || IsLastPoint)
+						{
+							int AddedIndex = ConnectingSegments.Add(MoveTemp(ConnectingSegment));
+							ConnectingSegment = FConnectingSegment();
+						}
+						//UE_LOG(LogProceduralVegetation, Warning, TEXT("Creating Bone NjordPixelId, Id{%f} BoneId{%i} PointIndex{%i} ParentBoneIndex{%i} BranchID{%i}"),
+						//	BoneNode.NjordPixelID, BoneNode.BoneIndex, BoneNode.PointIndex,BoneNode.ParentBoneIndex,BoneNode.BranchIndex);
 					}
-					//UE_LOG(LogProceduralVegetation, Warning, TEXT("Creating Bone NjordPixelId, Id{%f} BoneId{%i} PointIndex{%i} ParentBoneIndex{%i} BranchID{%i}"),
-					//	BoneNode.NjordPixelID, BoneNode.BoneIndex, BoneNode.PointIndex,BoneNode.ParentBoneIndex,BoneNode.BranchIndex);
 				}
+				ApplyBoneReduction(ConnectingSegments, Bones, ReductionStrength);
 			}
-			ApplyBoneReduction(ConnectingSegments, Bones, ReductionStrength);
 		}
 
 		ReindexBones(Bones);
diff --git a/Engine/Plugins/Experimental/ProceduralVegetationEditor/Source/ProceduralVegetation/Private/Facades/PVBranchFacade.cpp b/Engine/Plugins/Experimental/ProceduralVegetationEditor/Source/ProceduralVegetation/Private/Facades/PVBranchFacade.cpp
--- a/Engine/Plugins/Experimental/ProceduralVegetationEditor/Source/ProceduralVegetation/Private/Facades/PVBranchFacade.cpp
+++ b/Engine/Plugins/Experimental/ProceduralVegetationEditor/Source/ProceduralVegetation/Private/Facades/PVBranchFacade.cpp
@@ -123,16 +123,6 @@ namespace PV::Facades
 		return INDEX_NONE;
 	}
 
-	int32 FBranchFacade::GetTrunkIndex() const
-	{
-		if (IsValid() && BranchNumbers.Num() > 0)
-		{
-			return BranchNumbers.Get().Find(1);
-		}
-
-		return INDEX_NONE;
-	}
-
 	TArray<int32> FBranchFacade::GetParentBranchIndices(int BranchIndex) const
 	{
 		TArray<int32> ParentIndices;
diff --git a/Engine/Plugins/Experimental/ProceduralVegetationEditor/Source/ProceduralVegetation/Private/Facades/PVPlantFacade.cpp b/Engine/Plugins/Experimental/ProceduralVegetationEditor/Source/ProceduralVegetation/Private/Facades/PVPlantFacade.cpp
new file mode 100644
--- /dev/null
+++ b/Engine/Plugins/Experimental/ProceduralVegetationEditor/Source/ProceduralVegetation/Private/Facades/PVPlantFacade.cpp
@@ -0,0 +1,156 @@
+﻿// Copyright Epic Games, Inc. All Rights Reserved.
+
+#include "Facades/PVPlantFacade.h"
+#include "Facades/PVAttributesNames.h"
+
+namespace PV::Facades
+{
+	FPlantFacade::FPlantFacade(FManagedArrayCollection& InCollection)
+		: Collection(&InCollection)
+		, PlantNumbers(InCollection, AttributeNames::PlantNumber, GroupNames::BranchGroup)
+		, BranchNumbers(InCollection, AttributeNames::BranchNumber, GroupNames::BranchGroup)
+		, BranchParentNumbers(InCollection, AttributeNames::BranchParentNumber, GroupNames::BranchGroup)
+		, BranchHierarchyNumbers(InCollection, AttributeNames::BranchHierarchyNumber, GroupNames::BranchGroup)
+	{
+	}
+
+	FPlantFacade::FPlantFacade(const FManagedArrayCollection& InCollection)
+		: Collection(nullptr)
+		, PlantNumbers(InCollection, AttributeNames::PlantNumber, GroupNames::BranchGroup)
+		, BranchNumbers(InCollection, AttributeNames::BranchNumber, GroupNames::BranchGroup)
+		, BranchParentNumbers(InCollection, AttributeNames::BranchParentNumber, GroupNames::BranchGroup)
+		, BranchHierarchyNumbers(InCollection, AttributeNames::BranchHierarchyNumber, GroupNames::BranchGroup)
+	{
+	}
+
+	bool FPlantFacade::IsValid() const
+	{
+		return PlantNumbers.IsValid()
+			&& BranchNumbers.IsValid()
+			&& BranchParentNumbers.IsValid()
+			&& BranchHierarchyNumbers.IsValid();
+	}
+
+	int32 FPlantFacade::GetElementCount() const
+	{
+		TSet<int32> PlantNumberElements;
+		for (int32 i = 0; i < PlantNumbers.Num(); i++)
+		{
+			PlantNumberElements.Add(PlantNumbers[i]);
+		}
+
+		return PlantNumberElements.Num();
+	}
+
+	TArray<int32> FPlantFacade::GetBranchNumbers(const int32 PlantNumber) const
+	{
+		TArray<int32> BranchNumbersResult;
+
+		for (int32 i = 0; i < BranchNumbers.Num(); i++)
+		{
+			if (PlantNumbers[i] == PlantNumber)
+			{
+				BranchNumbersResult.Add(BranchNumbers[i]);
+			}
+		}
+
+		return BranchNumbersResult;
+	}
+
+	TArray<int32> FPlantFacade::GetBranchIndices(const int32 PlantNumber) const
+	{
+		TArray<int32> BranchIndicesResult;
+
+		for (int32 i = 0; i < BranchNumbers.Num(); i++)
+		{
+			if (PlantNumbers[i] == PlantNumber)
+			{
+				BranchIndicesResult.Add(i);
+			}
+		}
+
+		return BranchIndicesResult;
+	}
+
+	TArray<int32> FPlantFacade::GetBranchIndicesSortedByHierarchyNumber(const int32 PlantNumber) const
+	{
+		TArray<int32> BranchIndices = GetBranchIndices(PlantNumber);
+		BranchIndices.Sort([&](const int32 Index1, const int32 Index2)
+			{
+				const int32 Index1Hierarchy = BranchHierarchyNumbers[Index1];
+				const int32 Index2Hierarchy = BranchHierarchyNumbers[Index2];
+
+				if (Index1Hierarchy == Index2Hierarchy)
+				{
+					return Index1 < Index2;
+				}
+
+				return Index1Hierarchy < Index2Hierarchy;
+			});
+
+		return BranchIndices;
+	}
+
+	TArray<int32> FPlantFacade::GetPlantNumbers() const
+	{
+		TSet<int32> PlantNumberElements;
+		for (int32 i = 0; i < PlantNumbers.Num(); i++)
+		{
+			PlantNumberElements.Add(PlantNumbers[i]);
+		}
+
+		return PlantNumberElements.Array();
+	}
+
+	TArray<int32> FPlantFacade::GetTrunkIndices() const
+	{
+		TArray<int32> TrunkIndicesResult;
+
+		for (int32 i = 0; i < BranchNumbers.Num(); i++)
+		{
+			if (BranchParentNumbers[i] == 0)
+			{
+				TrunkIndicesResult.Add(i);
+			}
+		}
+
+		return TrunkIndicesResult;
+	}
+
+	TMap<int32, int32> FPlantFacade::GetPlantNumbersToTrunkIndicesMap() const
+	{
+		TMap<int32, int32> PlantNumberToTrunkIndicesMap;
+		for (int32 i = 0; i < BranchNumbers.Num(); i++)
+		{
+			if (BranchParentNumbers[i] == 0)
+			{
+				PlantNumberToTrunkIndicesMap.Add(PlantNumbers[i], i);
+			}
+		}
+
+		return PlantNumberToTrunkIndicesMap;
+	}
+
+	int32 FPlantFacade::GetTrunkIndex(const int32 PlantNumber) const
+	{
+		for (int32 i = 0; i < BranchParentNumbers.Num(); i++)
+		{
+			if (PlantNumbers[i] == PlantNumber && BranchParentNumbers[i] == 0)
+			{
+				return i;
+			}
+		}
+
+		return INDEX_NONE;
+	}
+
+	bool FPlantFacade::IsTrunkIndex(const int32 BranchIndex) const
+	{
+		if (BranchParentNumbers.IsValidIndex(BranchIndex))
+		{
+			return BranchParentNumbers[BranchIndex] == 0;
+		}
+
+		return false;
+	}
+}
diff --git a/Engine/Plugins/Experimental/ProceduralVegetationEditor/Source/ProceduralVegetation/Private/Implementations/PVCarve.cpp b/Engine/Plugins/Experimental/ProceduralVegetationEditor/Source/ProceduralVegetation/Private/Implementations/PVCarve.cpp
--- a/Engine/Plugins/Experimental/ProceduralVegetationEditor/Source/ProceduralVegetation/Private/Implementations/PVCarve.cpp
+++ b/Engine/Plugins/Experimental/ProceduralVegetationEditor/Source/ProceduralVegetation/Private/Implementations/PVCarve.cpp
@@ -2,13 +2,15 @@
 
 #include "PVCarve.h"
 
+#include "Facades/PVAttributesNames.h"
 #include "Facades/PVFoliageFacade.h"
+#include "Facades/PVPlantFacade.h"
 #include "Facades/PVPointFacade.h"
 #include "Facades/PVTreeFacade.h"
 
 
 void FPVCarve::ApplyCarve(FManagedArrayCollection& OutCollection, const FManagedArrayCollection& SourceCollection,
-                                 const ECarveBasis CarveBasis, const float Carve)
+                          const ECarveBasis CarveBasis, const float Carve)
 {
 	if (Carve == 0.0f)
 	{
@@ -26,8 +28,8 @@ void FPVCarve::ApplyCarve(FManagedArrayCollection& OutCollection, const FManaged
 }
 
 void FPVCarve::UpdatePointScales(PV::Facades::FPointFacade& PointFacadeOut, const PV::Facades::FPointFacade& PointFacadeSource,
-                                        const TArray<int>& BranchPoints, const float LastPointScale, const int32 LastPointIndex,
-                                        const float CarveRatio, const int32 EndIndex, const float FirstPointTargetPScale)
+                                 const TArray<int>& BranchPoints, const float LastPointScale, const int32 LastPointIndex,
+                                 const float CarveRatio, const int32 EndIndex, const float FirstPointTargetPScale)
 {
 	TRACE_CPUPROFILER_EVENT_SCOPE(PVCarve::UpdatePointScales);
 
@@ -48,9 +50,9 @@ void FPVCarve::UpdatePointScales(PV::Facades::FPointFacade& PointFacadeOut, cons
 }
 
 void FPVCarve::UpdateBranchLinkedIndices(PV::Facades::FBranchFacade& BranchFacadeOut,
-                                                PV::Facades::FFoliageFacade& FoliageFacadeOut, const TMap<int32, int32>& PointsOldIDsToNewIDs,
-                                                const TArray<bool>& PointsRemoved, const TMap<int32, int32>& FoliageInstancesOldIDsToNewIDs,
-                                                const TArray<bool>& FoliageInstancesRemoved)
+                                         PV::Facades::FFoliageFacade& FoliageFacadeOut, const TMap<int32, int32>& PointsOldIDsToNewIDs,
+                                         const TArray<bool>& PointsRemoved, const TMap<int32, int32>& FoliageInstancesOldIDsToNewIDs,
+                                         const TArray<bool>& FoliageInstancesRemoved)
 {
 	TRACE_CPUPROFILER_EVENT_SCOPE(PVCarve::UpdateBranchLinkedIndices);
 
@@ -120,87 +122,94 @@ void FPVCarve::UpdateBranchLinkedIndices(PV::Facades::FBranchFacade& BranchFacad
 	}
 }
 
-void FPVCarve::RecomputeAttributes(PV::Facades::FPointFacade& PointFacadeOut,
-                                          const PV::Facades::FPointFacade& PointFacadeSource,
-                                          const PV::Facades::FBranchFacade& BranchFacade,
-                                          const PV::Facades::FBranchFacade& BranchFacadeSource,
-                                          const TMap<int32, int32>& BranchesNewIDsToOldIDs)
+void FPVCarve::RecomputeAttributes(FManagedArrayCollection& OutCollection, const FManagedArrayCollection& SourceCollection,
+                                   const TMap<int32, int32>& BranchesNewIDsToOldIDs)
 {
 	TRACE_CPUPROFILER_EVENT_SCOPE(PVCarve::RecomputeAttributes);
 
-	// Remap Ethylene levels while gathering minimum and maximum PScales for each branch and the whole plant
-	float MinimumPScale = TNumericLimits<float>::Max();
-	float MaximumPScale = TNumericLimits<float>::Min();
-	TArray<TTuple<float, float>> BranchMinMaxPScales;
-	BranchMinMaxPScales.Init(TTuple<float, float>(TNumericLimits<float>::Max(), TNumericLimits<float>::Min()), BranchFacade.GetElementCount());
+	const PV::Facades::FBranchFacade BranchFacadeSource(SourceCollection);
+	const PV::Facades::FPointFacade PointFacadeSource(SourceCollection);
+
+	const PV::Facades::FBranchFacade BranchFacadeOut(OutCollection);
+	const PV::Facades::FPlantFacade PlantFacadeOut(OutCollection);
+	PV::Facades::FPointFacade PointFacadeOut(OutCollection);
 
-	for (int32 BranchIndex = 0; BranchIndex < BranchFacade.GetElementCount(); ++BranchIndex)
+	for (const int32& PlantNumber : PlantFacadeOut.GetPlantNumbers())
 	{
-		const TArray<int32>& BranchPoints = BranchFacade.GetPoints(BranchIndex);
-		const int32 NumBranchPointsCurrent = BranchPoints.Num();
-		const int32 OldBranchIndex = BranchesNewIDsToOldIDs.Contains(BranchIndex)
-			? BranchesNewIDsToOldIDs[BranchIndex]
-			: BranchIndex;
-		const TArray<int32>& SourceBranchPoints = BranchFacadeSource.GetPoints(OldBranchIndex);
-
-		for (int32 Idx = 0; Idx < NumBranchPointsCurrent; ++Idx)
+		// Remap Ethylene levels while gathering minimum and maximum PScales for each branch and the whole plant
+		float MinimumPScale = TNumericLimits<float>::Max();
+		float MaximumPScale = TNumericLimits<float>::Min();
+		TArray<TTuple<float, float>> BranchMinMaxPScales;
+		BranchMinMaxPScales.Init(TTuple<float, float>(TNumericLimits<float>::Max(), TNumericLimits<float>::Min()), BranchFacadeOut.GetElementCount());
+
+		for (const int32& BranchIndex : PlantFacadeOut.GetBranchIndices(PlantNumber))
 		{
-			const float RelativePointIndex = (static_cast<float>(Idx) / NumBranchPointsCurrent) * SourceBranchPoints.Num();
-			const int32 OldCurrentPointIndex = FMath::Clamp(FMath::CeilToInt(RelativePointIndex), 0, SourceBranchPoints.Num() - 1);
-			const int32 OldPreviousPointIndex = OldCurrentPointIndex == 0
-				? 0
-				: OldCurrentPointIndex - 1;
-			float BlendValue = RelativePointIndex - FMath::FloorToInt(RelativePointIndex);
-
-			const TArray<float>& BudHormoneLevelsAtOldPreviousPointIndex = PointFacadeSource.GetBudHormoneLevels(
-				SourceBranchPoints[OldPreviousPointIndex]);
-			check(BudHormoneLevelsAtOldPreviousPointIndex.Num() >= 5);
-			float EthyleneLevelAtOldPreviousPointIndex = BudHormoneLevelsAtOldPreviousPointIndex[4];
-
-			const TArray<float>& BudHormoneLevelsAtOldCurrentPointIndex = PointFacadeSource.GetBudHormoneLevels(
-				SourceBranchPoints[OldCurrentPointIndex]);
-			check(BudHormoneLevelsAtOldCurrentPointIndex.Num() >= 5);
-			float EthyleneLevelAtOldCurrentPointIndex = BudHormoneLevelsAtOldCurrentPointIndex[4];
-
-			const float NewEthyleneLevel = FMath::Lerp(EthyleneLevelAtOldPreviousPointIndex, EthyleneLevelAtOldCurrentPointIndex, BlendValue);
-			TArray<float> CurrentBudHormoneLevels = PointFacadeOut.GetBudHormoneLevels(BranchPoints[Idx]);
-			check(CurrentBudHormoneLevels.Num() >= 4);
-			CurrentBudHormoneLevels[4] = NewEthyleneLevel;
-			PointFacadeOut.SetBudHormoneLevels(BranchPoints[Idx], CurrentBudHormoneLevels);
-
-			const float PointScaleAtIndex = PointFacadeOut.GetPointScale(BranchPoints[Idx]);
-			MinimumPScale = FMath::Min(MinimumPScale, PointScaleAtIndex);
-			MaximumPScale = FMath::Max(MaximumPScale, PointScaleAtIndex);
-
-			BranchMinMaxPScales[BranchIndex].Get<0>() = FMath::Min(BranchMinMaxPScales[BranchIndex].Get<0>(), PointScaleAtIndex);
-			BranchMinMaxPScales[BranchIndex].Get<1>() = FMath::Max(BranchMinMaxPScales[BranchIndex].Get<1>(), PointScaleAtIndex);
+			const TArray<int32>& BranchPoints = BranchFacadeOut.GetPoints(BranchIndex);
+			const int32 NumBranchPointsCurrent = BranchPoints.Num();
+			const int32 OldBranchIndex = BranchesNewIDsToOldIDs.Contains(BranchIndex)
+				? BranchesNewIDsToOldIDs[BranchIndex]
+				: BranchIndex;
+			const TArray<int32>& SourceBranchPoints = BranchFacadeSource.GetPoints(OldBranchIndex);
+
+			for (int32 Idx = 0; Idx < NumBranchPointsCurrent; ++Idx)
+			{
+				const float RelativePointIndex = (static_cast<float>(Idx) / NumBranchPointsCurrent) * SourceBranchPoints.Num();
+				const int32 OldCurrentPointIndex = FMath::Clamp(FMath::CeilToInt(RelativePointIndex), 0, SourceBranchPoints.Num() - 1);
+				const int32 OldPreviousPointIndex = OldCurrentPointIndex == 0
+					? 0
+					: OldCurrentPointIndex - 1;
+				float BlendValue = RelativePointIndex - FMath::FloorToInt(RelativePointIndex);
+
+				const TArray<float>& BudHormoneLevelsAtOldPreviousPointIndex = PointFacadeSource.GetBudHormoneLevels(
+					SourceBranchPoints[OldPreviousPointIndex]);
+				check(BudHormoneLevelsAtOldPreviousPointIndex.Num() >= 5);
+				float EthyleneLevelAtOldPreviousPointIndex = BudHormoneLevelsAtOldPreviousPointIndex[4];
+
+				const TArray<float>& BudHormoneLevelsAtOldCurrentPointIndex = PointFacadeSource.GetBudHormoneLevels(
+					SourceBranchPoints[OldCurrentPointIndex]);
+				check(BudHormoneLevelsAtOldCurrentPointIndex.Num() >= 5);
+				float EthyleneLevelAtOldCurrentPointIndex = BudHormoneLevelsAtOldCurrentPointIndex[4];
+
+				const float NewEthyleneLevel = FMath::Lerp(EthyleneLevelAtOldPreviousPointIndex, EthyleneLevelAtOldCurrentPointIndex, BlendValue);
+				TArray<float> CurrentBudHormoneLevels = PointFacadeOut.GetBudHormoneLevels(BranchPoints[Idx]);
+				check(CurrentBudHormoneLevels.Num() >= 4);
+				CurrentBudHormoneLevels[4] = NewEthyleneLevel;
+				PointFacadeOut.SetBudHormoneLevels(BranchPoints[Idx], CurrentBudHormoneLevels);
+
+				const float PointScaleAtIndex = PointFacadeOut.GetPointScale(BranchPoints[Idx]);
+				MinimumPScale = FMath::Min(MinimumPScale, PointScaleAtIndex);
+				MaximumPScale = FMath::Max(MaximumPScale, PointScaleAtIndex);
+
+				BranchMinMaxPScales[BranchIndex].Get<0>() = FMath::Min(BranchMinMaxPScales[BranchIndex].Get<0>(), PointScaleAtIndex);
+				BranchMinMaxPScales[BranchIndex].Get<1>() = FMath::Max(BranchMinMaxPScales[BranchIndex].Get<1>(), PointScaleAtIndex);
+			}
 		}
-	}
-
-	// Remap BranchGradient and PlantGradient
-	const FVector2f NormalRange(0.0f, 1.0f);
-	const FVector2f PlantPScaleRange(MinimumPScale, MaximumPScale);
 
-	for (int32 BranchIndex = 0; BranchIndex < BranchFacade.GetElementCount(); ++BranchIndex)
-	{
-		const TArray<int32>& BranchPoints = BranchFacade.GetPoints(BranchIndex);
-		const FVector2f BranchPScaleRange(BranchMinMaxPScales[BranchIndex].Get<0>(), BranchMinMaxPScales[BranchIndex].Get<1>());
+		// Remap BranchGradient and PlantGradient
+		const FVector2f NormalRange(0.0f, 1.0f);
+		const FVector2f PlantPScaleRange(MinimumPScale, MaximumPScale);
 
-		for (const int32& PointIndex : BranchPoints)
+		for (const int32& BranchIndex : PlantFacadeOut.GetBranchIndices(PlantNumber))
 		{
-			const float PlantGradientAtPointIndex = FMath::GetMappedRangeValueClamped(PlantPScaleRange, NormalRange,
-				PointFacadeOut.GetPointScale(PointIndex));
-			PointFacadeOut.SetPlantGradient(PointIndex, PlantGradientAtPointIndex);
+			const TArray<int32>& BranchPoints = BranchFacadeOut.GetPoints(BranchIndex);
+			const FVector2f BranchPScaleRange(BranchMinMaxPScales[BranchIndex].Get<0>(), BranchMinMaxPScales[BranchIndex].Get<1>());
+
+			for (const int32& PointIndex : BranchPoints)
+			{
+				const float PlantGradientAtPointIndex = FMath::GetMappedRangeValueClamped(PlantPScaleRange, NormalRange,
+					PointFacadeOut.GetPointScale(PointIndex));
+				PointFacadeOut.SetPlantGradient(PointIndex, PlantGradientAtPointIndex);
 
-			const float BranchGradientAtPointIndex = FMath::GetMappedRangeValueClamped(BranchPScaleRange, NormalRange,
-				PointFacadeOut.GetPointScale(PointIndex));
-			PointFacadeOut.SetBranchGradient(PointIndex, BranchGradientAtPointIndex);
+				const float BranchGradientAtPointIndex = FMath::GetMappedRangeValueClamped(BranchPScaleRange, NormalRange,
+					PointFacadeOut.GetPointScale(PointIndex));
+				PointFacadeOut.SetBranchGradient(PointIndex, BranchGradientAtPointIndex);
+			}
 		}
 	}
 }
 
 void FPVCarve::CarveFromTop(FManagedArrayCollection& OutCollection, const FManagedArrayCollection& SourceCollection, const float Carve,
-                                   const ECarveBasis CarveBasis)
+                            const ECarveBasis CarveBasis)
 {
 	TRACE_CPUPROFILER_EVENT_SCOPE(PVCarve::CarveFromTop);
 
@@ -209,6 +218,7 @@ void FPVCarve::CarveFromTop(FManagedArrayCollection& OutCollection, const FManag
 	const PV::Facades::FBranchFacade BranchFacadeSource(SourceCollection);
 	const PV::Facades::FPointFacade PointFacadeSource(SourceCollection);
 	const PV::Facades::FFoliageFacade FoliageFacadeSource(SourceCollection);
+	const PV::Facades::FPlantFacade PlantFacadeSource(SourceCollection);
 
 	PV::Facades::FBranchFacade BranchFacadeOut(OutCollection);
 	PV::Facades::FPointFacade PointFacadeOut(OutCollection);
@@ -235,9 +245,6 @@ void FPVCarve::CarveFromTop(FManagedArrayCollection& OutCollection, const FManag
 	TArray<float> CarveValues;
 	CarveValues.Reserve(PointFacadeSource.GetElementCount());
 
-	float MinimumCarveValue = TNumericLimits<float>::Max();
-	float MaximumCarveValue = TNumericLimits<float>::Min();
-
 	for (int32 i = 0; i < PointFacadeSource.GetElementCount(); ++i)
 	{
 		switch (CarveBasis)
@@ -258,9 +265,6 @@ void FPVCarve::CarveFromTop(FManagedArrayCollection& OutCollection, const FManag
 			CarveValues.Add(PointFacadeSource.GetPointScale(i));
 			break;
 		}
-
-		MinimumCarveValue = FMath::Min(MinimumCarveValue, CarveValues[i]);
-		MaximumCarveValue = FMath::Max(MaximumCarveValue, CarveValues[i]);
 	}
 
 	FVector2f MappedRange(0.0f, 1.0f);
@@ -269,158 +273,173 @@ void FPVCarve::CarveFromTop(FManagedArrayCollection& OutCollection, const FManag
 		MappedRange = FVector2f(1.0f, 0.0f);
 	}
 
-	for (int32 i = 0; i < CarveValues.Num(); ++i)
+	for (const int32& PlantNumber : PlantFacadeSource.GetPlantNumbers())
 	{
-		CarveValues[i] = FMath::GetMappedRangeValueClamped(FVector2f(MinimumCarveValue, MaximumCarveValue), MappedRange, CarveValues[i]);
-	}
+		float MinimumCarveValue = TNumericLimits<float>::Max();
+		float MaximumCarveValue = TNumericLimits<float>::Min();
 
-	// Compute branches, foliage instances and points to remove
-	const int32 TrunkIndex = BranchFacadeSource.GetTrunkIndex();
-	for (int32 BranchIndex = 0; BranchIndex < NumOfBranchesInSource; ++BranchIndex)
-	{
-		if (BranchesToRemove[BranchIndex])
+		for (const int32& BranchIndex : PlantFacadeSource.GetBranchIndices(PlantNumber))
 		{
-			continue;
-		}
+			for (const int32& PointIndex : BranchFacadeSource.GetPoints(BranchIndex))
+			{
+				MinimumCarveValue = FMath::Min(MinimumCarveValue, CarveValues[PointIndex]);
+				MaximumCarveValue = FMath::Max(MaximumCarveValue, CarveValues[PointIndex]);
+			}
 
-		const TArray<int32>& BranchPoints = BranchFacadeSource.GetPoints(BranchIndex);
-		const int32 NumOfPointsInCurrentBranch = BranchPoints.Num();
+			for (const int32& PointIndex : BranchFacadeSource.GetPoints(BranchIndex))
+			{
+				CarveValues[PointIndex] = FMath::GetMappedRangeValueClamped(
+					FVector2f(MinimumCarveValue, MaximumCarveValue),
+					MappedRange,
+					CarveValues[PointIndex]);
+			}
 
-		if (NumOfPointsInCurrentBranch == 0)
-		{
-			BranchesToRemove[BranchIndex] = true;
-			continue;
-		}
+			// Compute branches, foliage instances and points to remove
+			if (BranchesToRemove[BranchIndex])
+			{
+				continue;
+			}
 
-		// Compute PScale for the point in parent branch that is closest to the first point
-		// (For non-trunk branches)
-		float FirstPointTargetPScale = PointFacadeSource.GetPointScale(BranchPoints[0]);
-		if (BranchIndex != TrunkIndex) [[likely]]
-		{
-			const FVector3f& FirstPointPosition = PointFacadeSource.GetPosition(BranchPoints[0]);
+			const TArray<int32>& BranchPoints = BranchFacadeSource.GetPoints(BranchIndex);
+			const int32 NumOfPointsInCurrentBranch = BranchPoints.Num();
 
-			const int32 ParentBranchNumber = BranchFacadeSource.GetParents(BranchIndex).Last();
-			const int32 ParentBranchIndex = BranchNumbersToBranchIDs[ParentBranchNumber];
-			const TArray<int32>& ParentBranchPoints = BranchFacadeSource.GetPoints(ParentBranchIndex);
+			if (NumOfPointsInCurrentBranch == 0)
+			{
+				BranchesToRemove[BranchIndex] = true;
+				continue;
+			}
 
-			int32 ClosestPointIndex = 0;
-			float MinimumDistance = FLT_MAX;
-			for (int32 PointIndex = 0; PointIndex < ParentBranchPoints.Num(); ++PointIndex)
+			// Compute PScale for the point in parent branch that is closest to the first point
+			// (For non-trunk branches)
+			float FirstPointTargetPScale = PointFacadeSource.GetPointScale(BranchPoints[0]);
+			if (!PlantFacadeSource.IsTrunkIndex(BranchIndex)) [[likely]]
 			{
-				const FVector3f& ParentPointPosition = PointFacadeSource.GetPosition(ParentBranchPoints[PointIndex]);
-				if (const float Distance = FVector3f::Distance(FirstPointPosition, ParentPointPosition); Distance < MinimumDistance)
+				const FVector3f& FirstPointPosition = PointFacadeSource.GetPosition(BranchPoints[0]);
+
+				const int32 ParentBranchNumber = BranchFacadeSource.GetParents(BranchIndex).Last();
+				const int32 ParentBranchIndex = BranchNumbersToBranchIDs[ParentBranchNumber];
+				const TArray<int32>& ParentBranchPoints = BranchFacadeSource.GetPoints(ParentBranchIndex);
+
+				int32 ClosestPointIndex = 0;
+				float MinimumDistance = FLT_MAX;
+				for (int32 PointIndex = 0; PointIndex < ParentBranchPoints.Num(); ++PointIndex)
 				{
-					MinimumDistance = Distance;
-					ClosestPointIndex = PointIndex;
+					const FVector3f& ParentPointPosition = PointFacadeSource.GetPosition(ParentBranchPoints[PointIndex]);
+					if (const float Distance = FVector3f::Distance(FirstPointPosition, ParentPointPosition); Distance < MinimumDistance)
+					{
+						MinimumDistance = Distance;
+						ClosestPointIndex = PointIndex;
+					}
 				}
-			}
-
-			FirstPointTargetPScale = PointFacadeOut.GetPointScale(ParentBranchPoints[ClosestPointIndex]);
-		}
 
-		int EndIndex = 1;
-		while (EndIndex < NumOfPointsInCurrentBranch && (CarveValues[BranchPoints[EndIndex]] < CarveDistance + 0.0001f))
-		{
-			EndIndex++;
-		}
-		EndIndex--;
+				FirstPointTargetPScale = PointFacadeOut.GetPointScale(ParentBranchPoints[ClosestPointIndex]);
+			}
 
-		if (EndIndex == NumOfPointsInCurrentBranch - 1)
-		{
-			const int32 LastPointIndex = BranchPoints[EndIndex];
-			const float LastPointScale = PointFacadeSource.GetPointScale(LastPointIndex);
-			UpdatePointScales(PointFacadeOut, PointFacadeSource, BranchPoints, LastPointScale, LastPointIndex, 1.0f, EndIndex,
-				FirstPointTargetPScale);
-			continue;
-		}
+			int EndIndex = 1;
+			while (EndIndex < NumOfPointsInCurrentBranch && (CarveValues[BranchPoints[EndIndex]] < CarveDistance + 0.0001f))
+			{
+				EndIndex++;
+			}
+			EndIndex--;
 
-		// Remove all branches, child branches and foliage instances beyond last point
-		for (const int32& BranchNumber : BranchFacadeSource.GetChildren(BranchIndex))
-		{
-			if (BranchNumbersToLengthFromRoots.Contains(BranchNumber) && BranchNumbersToLengthFromRoots[BranchNumber] > PointFacadeSource.
-				GetLengthFromRoot(BranchPoints[EndIndex]))
+			if (EndIndex == NumOfPointsInCurrentBranch - 1)
 			{
-				const int32 Index = BranchNumbersToBranchIDs[BranchNumber];
-				BranchesToRemove[Index] = true;
-				for (const int32& P : BranchFacadeSource.GetPoints(Index))
-				{
-					PointsToRemove[P] = true;
-				}
-				for (const int32& FId : FoliageFacadeSource.GetFoliageEntryIdsForBranch(Index))
-				{
-					FoliageInstancesToRemove[FId] = true;
-				}
+				const int32 LastPointIndex = BranchPoints[EndIndex];
+				const float LastPointScale = PointFacadeSource.GetPointScale(LastPointIndex);
+				UpdatePointScales(PointFacadeOut, PointFacadeSource, BranchPoints, LastPointScale, LastPointIndex, 1.0f, EndIndex,
+					FirstPointTargetPScale);
+				continue;
+			}
 
-				const TArray<int32>& BranchChildren = BranchFacadeSource.GetChildren(Index);
-				for (int32 k = 0; k < BranchChildren.Num(); ++k)
+			// Remove all branches, child branches and foliage instances beyond last point
+			for (const int32& BranchNumber : BranchFacadeSource.GetChildren(BranchIndex))
+			{
+				if (BranchNumbersToLengthFromRoots.Contains(BranchNumber) && BranchNumbersToLengthFromRoots[BranchNumber] > PointFacadeSource.
+					GetLengthFromRoot(BranchPoints[EndIndex]))
 				{
-					if (!BranchNumbersToBranchIDs.Contains(BranchChildren[k]))
-					{
-						continue;
-					}
-
-					const int32 ChildIndex = BranchNumbersToBranchIDs[BranchChildren[k]];
-					BranchesToRemove[ChildIndex] = true;
-					for (const int32& P : BranchFacadeSource.GetPoints(ChildIndex))
+					const int32 Index = BranchNumbersToBranchIDs[BranchNumber];
+					BranchesToRemove[Index] = true;
+					for (const int32& P : BranchFacadeSource.GetPoints(Index))
 					{
 						PointsToRemove[P] = true;
 					}
-					for (const int32& FId : FoliageFacadeSource.GetFoliageEntryIdsForBranch(ChildIndex))
+					for (const int32& FId : FoliageFacadeSource.GetFoliageEntryIdsForBranch(Index))
 					{
 						FoliageInstancesToRemove[FId] = true;
 					}
+
+					const TArray<int32>& BranchChildren = BranchFacadeSource.GetChildren(Index);
+					for (int32 k = 0; k < BranchChildren.Num(); ++k)
+					{
+						if (!BranchNumbersToBranchIDs.Contains(BranchChildren[k]))
+						{
+							continue;
+						}
+
+						const int32 ChildIndex = BranchNumbersToBranchIDs[BranchChildren[k]];
+						BranchesToRemove[ChildIndex] = true;
+						for (const int32& P : BranchFacadeSource.GetPoints(ChildIndex))
+						{
+							PointsToRemove[P] = true;
+						}
+						for (const int32& FId : FoliageFacadeSource.GetFoliageEntryIdsForBranch(ChildIndex))
+						{
+							FoliageInstancesToRemove[FId] = true;
+						}
+					}
 				}
 			}
-		}
 
-		for (const int32& FoliageIndex : FoliageFacadeSource.GetFoliageEntryIdsForBranch(BranchIndex))
-		{
-			if (FoliageFacadeSource.GetLengthFromRoot(FoliageIndex) > PointFacadeSource.GetLengthFromRoot(BranchPoints[EndIndex]))
+			for (const int32& FoliageIndex : FoliageFacadeSource.GetFoliageEntryIdsForBranch(BranchIndex))
 			{
-				FoliageInstancesToRemove[FoliageIndex] = true;
+				if (FoliageFacadeSource.GetLengthFromRoot(FoliageIndex) > PointFacadeSource.GetLengthFromRoot(BranchPoints[EndIndex]))
+				{
+					FoliageInstancesToRemove[FoliageIndex] = true;
+				}
 			}
-		}
 
-		// Keep next point to adjust its position, remove all others
-		++EndIndex;
-		for (int32 j = EndIndex + 1; j < NumOfPointsInCurrentBranch; ++j)
-		{
-			PointsToRemove[BranchPoints[j]] = true;
-		}
-
-		TArray<int32> UpdatedBranchPoints = TArray<int32>(&BranchPoints[0], EndIndex + 1);
-		BranchFacadeOut.SetPoints(BranchIndex, UpdatedBranchPoints);
-
-		// Adjust last point
-		const int32 LastPointIndex = BranchPoints[EndIndex];
-		const int32 PreviousPointIndex = BranchPoints[EndIndex - 1];
-		const FVector3f& LastPointPosition = PointFacadeSource.GetPosition(LastPointIndex);
-		const FVector3f& PreviousPointPosition = PointFacadeSource.GetPosition(PreviousPointIndex);
-		const float LastPointCarveValue = CarveValues[LastPointIndex];
-		const float PreviousPointCarveValue = CarveValues[PreviousPointIndex];
-		const float BlendValue = FMath::Clamp(
-			FMath::GetMappedRangeValueClamped(
-				FVector2f(PreviousPointCarveValue, LastPointCarveValue),
-				FVector2f(0.0f, 1.0f),
-				CarveDistance),
-			0.1f,
-			1.0f);
-		const FVector3f Position = FMath::Lerp(PreviousPointPosition, LastPointPosition, BlendValue);
-		const float LengthFromRoot = FMath::Lerp(PointFacadeSource.GetLengthFromRoot(PreviousPointIndex),
-			PointFacadeSource.GetLengthFromRoot(LastPointIndex),
-			BlendValue);
-		const float PointScale = FMath::Lerp(PointFacadeSource.GetPointScale(PreviousPointIndex),
-			PointFacadeSource.GetPointScale(LastPointIndex),
-			BlendValue);
-		PointFacadeOut.ModifyPositions()[LastPointIndex] = Position;
-		PointFacadeOut.ModifyLengthFromRoots()[LastPointIndex] = LengthFromRoot;
-		PointFacadeOut.ModifyLengthFromSeeds()[LastPointIndex] = LengthFromRoot;
-		PointFacadeOut.ModifyPointScales()[LastPointIndex] = PointScale;
+			// Keep next point to adjust its position, remove all others
+			++EndIndex;
+			for (int32 j = EndIndex + 1; j < NumOfPointsInCurrentBranch; ++j)
+			{
+				PointsToRemove[BranchPoints[j]] = true;
+			}
 
+			TArray<int32> UpdatedBranchPoints = TArray(&BranchPoints[0], EndIndex + 1);
+			BranchFacadeOut.SetPoints(BranchIndex, UpdatedBranchPoints);
 
-		// Update point scales
-		const float CarveRatio = (EndIndex + BlendValue) / NumOfPointsInCurrentBranch;
-		UpdatePointScales(PointFacadeOut, PointFacadeSource, BranchPoints, PointScale, LastPointIndex, CarveRatio, EndIndex, FirstPointTargetPScale);
+			// Adjust last point
+			const int32 LastPointIndex = BranchPoints[EndIndex];
+			const int32 PreviousPointIndex = BranchPoints[EndIndex - 1];
+			const FVector3f& LastPointPosition = PointFacadeSource.GetPosition(LastPointIndex);
+			const FVector3f& PreviousPointPosition = PointFacadeSource.GetPosition(PreviousPointIndex);
+			const float LastPointCarveValue = CarveValues[LastPointIndex];
+			const float PreviousPointCarveValue = CarveValues[PreviousPointIndex];
+			const float BlendValue = FMath::Clamp(
+				FMath::GetMappedRangeValueClamped(
+					FVector2f(PreviousPointCarveValue, LastPointCarveValue),
+					FVector2f(0.0f, 1.0f),
+					CarveDistance),
+				0.1f,
+				1.0f);
+			const FVector3f Position = FMath::Lerp(PreviousPointPosition, LastPointPosition, BlendValue);
+			const float LengthFromRoot = FMath::Lerp(PointFacadeSource.GetLengthFromRoot(PreviousPointIndex),
+				PointFacadeSource.GetLengthFromRoot(LastPointIndex),
+				BlendValue);
+			const float PointScale = FMath::Lerp(PointFacadeSource.GetPointScale(PreviousPointIndex),
+				PointFacadeSource.GetPointScale(LastPointIndex),
+				BlendValue);
+			PointFacadeOut.ModifyPositions()[LastPointIndex] = Position;
+			PointFacadeOut.ModifyLengthFromRoots()[LastPointIndex] = LengthFromRoot;
+			PointFacadeOut.ModifyLengthFromSeeds()[LastPointIndex] = LengthFromRoot;
+			PointFacadeOut.ModifyPointScales()[LastPointIndex] = PointScale;
+
+
+			// Update point scales
+			const float CarveRatio = (EndIndex + BlendValue) / NumOfPointsInCurrentBranch;
+			UpdatePointScales(PointFacadeOut, PointFacadeSource, BranchPoints, PointScale, LastPointIndex, CarveRatio, EndIndex,
+				FirstPointTargetPScale);
+		}
 	}
 
 	RemoveEntriesAndRecomputeAttributes(OutCollection, SourceCollection, PointsToRemove, BranchesToRemove, FoliageInstancesToRemove);
@@ -433,6 +452,8 @@ void FPVCarve::CarveFromBottom(FManagedArrayCollection& OutCollection, const FMa
 	const PV::Facades::FBranchFacade BranchFacadeSource(SourceCollection);
 	const PV::Facades::FPointFacade PointFacadeSource(SourceCollection);
 	const PV::Facades::FFoliageFacade FoliageFacadeSource(SourceCollection);
+	const PV::Facades::FPlantFacade PlantFacadeSource(SourceCollection);
+
 	PV::Facades::FBranchFacade BranchFacadeOut(OutCollection);
 	PV::Facades::FPointFacade PointFacadeOut(OutCollection);
 	PV::Facades::FFoliageFacade FoliageFacadeOut(OutCollection);
@@ -442,11 +463,10 @@ void FPVCarve::CarveFromBottom(FManagedArrayCollection& OutCollection, const FMa
 		return;
 	}
 
-	const int32 NumOfBranchesInSource = BranchFacadeSource.GetElementCount();
 	TArray<bool> PointsToRemove;
 	PointsToRemove.Init(false, PointFacadeSource.GetElementCount());
 	TArray<bool> BranchesToRemove;
-	BranchesToRemove.Init(false, NumOfBranchesInSource);
+	BranchesToRemove.Init(false, BranchFacadeSource.GetElementCount());
 	TArray<bool> FoliageInstancesToRemove;
 	FoliageInstancesToRemove.Init(false, FoliageFacadeSource.NumFoliageEntries());
 
@@ -456,138 +476,148 @@ void FPVCarve::CarveFromBottom(FManagedArrayCollection& OutCollection, const FMa
 
 	// Carve along trunk based on length from root
 	// Compute branches, foliage instances and points to remove
-	const int32 TrunkIndex = BranchFacadeSource.GetTrunkIndex();
-	const TArray<int32>& TrunkPoints = BranchFacadeSource.GetPoints(TrunkIndex);
-	const int32 NumOfPointsInTrunk = TrunkPoints.Num();
-	if (TrunkPoints.Num() < 2)
+	for (const TMap<int32, int32> PlantNumbersToTrunkIDs = PlantFacadeSource.GetPlantNumbersToTrunkIndicesMap();
+	     const TPair<int32, int32> Pair : PlantNumbersToTrunkIDs)
 	{
-		return;
-	}
+		const int32 PlantNumber = Pair.Key;
+		const int32 TrunkIndex = Pair.Value;
+		const TArray<int32>& TrunkPoints = BranchFacadeSource.GetPoints(TrunkIndex);
+		const int32 NumOfPointsInTrunk = TrunkPoints.Num();
 
-	const float LengthFromRootOfFirstPoint = PointFacadeSource.GetLengthFromRoot(TrunkPoints[0]);
-	const float LengthFromRootOfLastPoint = PointFacadeSource.GetLengthFromRoot(TrunkPoints[NumOfPointsInTrunk - 1]);
-	const float CarveLengthFromRoot = FMath::GetMappedRangeValueClamped(FVector2f(0.0f, 1.0f),
-		FVector2f(LengthFromRootOfFirstPoint, LengthFromRootOfLastPoint), Carve);
+		if (NumOfPointsInTrunk < 2)
+		{
+			BranchesToRemove[TrunkIndex] = true;
+			continue;
+		}
 
-	int EndIndex = 0;
-	while (EndIndex < NumOfPointsInTrunk && PointFacadeSource.GetLengthFromRoot(TrunkPoints[EndIndex]) < CarveLengthFromRoot)
-	{
-		EndIndex++;
-	}
-	EndIndex--;
+		const float LengthFromRootOfFirstPoint = PointFacadeSource.GetLengthFromRoot(TrunkPoints[0]);
+		const float LengthFromRootOfLastPoint = PointFacadeSource.GetLengthFromRoot(TrunkPoints[NumOfPointsInTrunk - 1]);
+		const float CarveLengthFromRoot = FMath::GetMappedRangeValueClamped(FVector2f(0.0f, 1.0f),
+			FVector2f(LengthFromRootOfFirstPoint, LengthFromRootOfLastPoint), Carve);
 
-	// Remove all branches, child branches and foliage instances before current point
-	for (const int32& BranchNumber : BranchFacadeSource.GetChildren(TrunkIndex))
-	{
-		if (BranchNumbersToLengthFromRoots.Contains(BranchNumber) && BranchNumbersToLengthFromRoots[BranchNumber] < CarveLengthFromRoot)
+		int EndIndex = 0;
+		while (EndIndex < NumOfPointsInTrunk && PointFacadeSource.GetLengthFromRoot(TrunkPoints[EndIndex]) < CarveLengthFromRoot)
 		{
-			const int32 Index = BranchNumbersToBranchIDs[BranchNumber];
-			BranchesToRemove[Index] = true;
-			for (const int32& P : BranchFacadeSource.GetPoints(Index))
-			{
-				PointsToRemove[P] = true;
-			}
-			for (const int32& FId : FoliageFacadeSource.GetFoliageEntryIdsForBranch(Index))
-			{
-				FoliageInstancesToRemove[FId] = true;
-			}
+			EndIndex++;
+		}
+		EndIndex--;
 
-			const TArray<int32>& BranchChildren = BranchFacadeSource.GetChildren(Index);
-			for (int32 k = 0; k < BranchChildren.Num(); ++k)
+		// Remove all branches, child branches and foliage instances before current point
+		for (const int32& BranchNumber : BranchFacadeSource.GetChildren(TrunkIndex))
+		{
+			if (BranchNumbersToLengthFromRoots.Contains(BranchNumber) && BranchNumbersToLengthFromRoots[BranchNumber] < CarveLengthFromRoot)
 			{
-				if (!BranchNumbersToBranchIDs.Contains(BranchChildren[k]))
-				{
-					continue;
-				}
-
-				const int32 ChildIndex = BranchNumbersToBranchIDs[BranchChildren[k]];
-				BranchesToRemove[ChildIndex] = true;
-				for (const int32& P : BranchFacadeSource.GetPoints(ChildIndex))
+				const int32 Index = BranchNumbersToBranchIDs[BranchNumber];
+				BranchesToRemove[Index] = true;
+				for (const int32& P : BranchFacadeSource.GetPoints(Index))
 				{
 					PointsToRemove[P] = true;
 				}
-				for (const int32& FId : FoliageFacadeSource.GetFoliageEntryIdsForBranch(ChildIndex))
+				for (const int32& FId : FoliageFacadeSource.GetFoliageEntryIdsForBranch(Index))
 				{
 					FoliageInstancesToRemove[FId] = true;
 				}
+
+				const TArray<int32>& BranchChildren = BranchFacadeSource.GetChildren(Index);
+				for (int32 k = 0; k < BranchChildren.Num(); ++k)
+				{
+					if (!BranchNumbersToBranchIDs.Contains(BranchChildren[k]))
+					{
+						continue;
+					}
+
+					const int32 ChildIndex = BranchNumbersToBranchIDs[BranchChildren[k]];
+					BranchesToRemove[ChildIndex] = true;
+					for (const int32& P : BranchFacadeSource.GetPoints(ChildIndex))
+					{
+						PointsToRemove[P] = true;
+					}
+					for (const int32& FId : FoliageFacadeSource.GetFoliageEntryIdsForBranch(ChildIndex))
+					{
+						FoliageInstancesToRemove[FId] = true;
+					}
+				}
 			}
 		}
-	}
 
-	for (const int32& FoliageIndex : FoliageFacadeSource.GetFoliageEntryIdsForBranch(TrunkIndex))
-	{
-		if (FoliageFacadeSource.GetLengthFromRoot(FoliageIndex) < CarveLengthFromRoot)
+		for (const int32& FoliageIndex : FoliageFacadeSource.GetFoliageEntryIdsForBranch(TrunkIndex))
 		{
-			FoliageInstancesToRemove[FoliageIndex] = true;
+			if (FoliageFacadeSource.GetLengthFromRoot(FoliageIndex) < CarveLengthFromRoot)
+			{
+				FoliageInstancesToRemove[FoliageIndex] = true;
+			}
 		}
-	}
 
-	// Keep current point to adjust its position, remove all others
-	for (int32 j = 0; j < EndIndex; ++j)
-	{
-		PointsToRemove[TrunkPoints[j]] = true;
-	}
+		// Keep current point to adjust its position, remove all others
+		for (int32 j = 0; j < EndIndex; ++j)
+		{
+			PointsToRemove[TrunkPoints[j]] = true;
+		}
 
-	TArray<int32> UpdatedTrunkPoints = TArray<int32>(&TrunkPoints[EndIndex], NumOfPointsInTrunk - EndIndex);
-	BranchFacadeOut.SetPoints(TrunkIndex, UpdatedTrunkPoints);
-
-	// Adjust End point
-	const int32 CurrentPointIndex = TrunkPoints[EndIndex];
-	const int32 NextPointIndex = TrunkPoints[EndIndex + 1];
-	const FVector3f& PositionOfCurrentPoint = PointFacadeSource.GetPosition(CurrentPointIndex);
-	const FVector3f& PositionOfNextPoint = PointFacadeSource.GetPosition(NextPointIndex);
-	const float LengthFromRootOfCurrentPoint = PointFacadeSource.GetLengthFromRoot(CurrentPointIndex);
-	const float LengthFromRootOfNextPoint = PointFacadeSource.GetLengthFromRoot(NextPointIndex);
-
-	const float BlendValue = FMath::Clamp(
-		FMath::GetMappedRangeValueClamped(
-			FVector2f(LengthFromRootOfCurrentPoint, LengthFromRootOfNextPoint),
-			FVector2f(0.0f, 1.0f),
-			CarveLengthFromRoot),
-		0.1f,
-		1.0f);
-	const FVector3f Position = FMath::Lerp(PositionOfCurrentPoint, PositionOfNextPoint, BlendValue);
-	const float PointScale = FMath::Lerp(PointFacadeSource.GetPointScale(CurrentPointIndex),
-		PointFacadeSource.GetPointScale(NextPointIndex),
-		BlendValue);
-	PointFacadeOut.ModifyPositions()[CurrentPointIndex] = Position;
-	PointFacadeOut.ModifyLengthFromRoots()[CurrentPointIndex] = CarveLengthFromRoot;
-	PointFacadeOut.ModifyLengthFromSeeds()[CurrentPointIndex] = CarveLengthFromRoot;
-	PointFacadeOut.ModifyPointScales()[CurrentPointIndex] = PointScale;
+		TArray<int32> UpdatedTrunkPoints = TArray(&TrunkPoints[EndIndex], NumOfPointsInTrunk - EndIndex);
+		BranchFacadeOut.SetPoints(TrunkIndex, UpdatedTrunkPoints);
 
-	RemoveEntriesAndRecomputeAttributes(OutCollection, SourceCollection, PointsToRemove, BranchesToRemove, FoliageInstancesToRemove);
+		// Adjust End point
+		const int32 CurrentPointIndex = TrunkPoints[EndIndex];
+		const int32 NextPointIndex = TrunkPoints[EndIndex + 1];
+		const FVector3f& PositionOfCurrentPoint = PointFacadeSource.GetPosition(CurrentPointIndex);
+		const FVector3f& PositionOfNextPoint = PointFacadeSource.GetPosition(NextPointIndex);
+		const float LengthFromRootOfCurrentPoint = PointFacadeSource.GetLengthFromRoot(CurrentPointIndex);
+		const float LengthFromRootOfNextPoint = PointFacadeSource.GetLengthFromRoot(NextPointIndex);
 
-	// Move all points down
-	const FVector3f Delta(0.0f, 0.0f, Position.Z);
+		const float BlendValue = FMath::Clamp(
+			FMath::GetMappedRangeValueClamped(
+				FVector2f(LengthFromRootOfCurrentPoint, LengthFromRootOfNextPoint),
+				FVector2f(0.0f, 1.0f),
+				CarveLengthFromRoot),
+			0.1f,
+			1.0f);
+		const FVector3f Position = FMath::Lerp(PositionOfCurrentPoint, PositionOfNextPoint, BlendValue);
+		const float PointScale = FMath::Lerp(PointFacadeSource.GetPointScale(CurrentPointIndex),
+			PointFacadeSource.GetPointScale(NextPointIndex),
+			BlendValue);
+		PointFacadeOut.ModifyPositions()[CurrentPointIndex] = Position;
+		PointFacadeOut.ModifyLengthFromRoots()[CurrentPointIndex] = CarveLengthFromRoot;
+		PointFacadeOut.ModifyLengthFromSeeds()[CurrentPointIndex] = CarveLengthFromRoot;
+		PointFacadeOut.ModifyPointScales()[CurrentPointIndex] = PointScale;
 
-	for (int32 PointIndex = 0; PointIndex < PointFacadeOut.GetElementCount(); ++PointIndex)
-	{
-		PointFacadeOut.ModifyPositions()[PointIndex] = PointFacadeOut.GetPosition(PointIndex) - Delta;
-	}
+		// Move all points down and recompute length from roots
+		const FVector3f Delta(0.0f, 0.0f, Position.Z);
 
-	for (int32 FoliageIndex = 0; FoliageIndex < FoliageFacadeOut.GetElementCount(); ++FoliageIndex)
-	{
-		const FVector3f UpdatedPivotPoint = FoliageFacadeOut.GetPivotPoint(FoliageIndex) - Delta;
-		FoliageFacadeOut.SetPivotPoint(FoliageIndex, UpdatedPivotPoint);
-	}
+		for (const int32& BranchIndex : PlantFacadeSource.GetBranchIndices(PlantNumber))
+		{
+			if (BranchesToRemove[BranchIndex])
+			{
+				continue;
+			}
 
-	// Recompute length from root
-	for (int32 PointIndex = 0; PointIndex < PointFacadeOut.GetElementCount(); ++PointIndex)
-	{
-		PointFacadeOut.ModifyLengthFromRoots()[PointIndex] = FMath::Max(PointFacadeOut.GetLengthFromRoot(PointIndex) - CarveLengthFromRoot, 0.0f);
-		PointFacadeOut.ModifyLengthFromSeeds()[PointIndex] = FMath::Max(PointFacadeOut.GetLengthFromSeed(PointIndex) - CarveLengthFromRoot, 0.0f);
-	}
+			for (const int32& PointIndex : BranchFacadeOut.GetPoints(BranchIndex))
+			{
+				PointFacadeOut.ModifyPositions()[PointIndex] = PointFacadeOut.GetPosition(PointIndex) - Delta;
 
-	for (int32 FoliageIndex = 0; FoliageIndex < FoliageFacadeOut.GetElementCount(); ++FoliageIndex)
-	{
-		const float UpdatedLengthFromRoot = FMath::Max(FoliageFacadeOut.GetLengthFromRoot(FoliageIndex) - CarveLengthFromRoot);
-		FoliageFacadeOut.SetLengthFromRoot(FoliageIndex, UpdatedLengthFromRoot);
+				PointFacadeOut.ModifyLengthFromRoots()[PointIndex] = FMath::Max(PointFacadeOut.GetLengthFromRoot(PointIndex) - CarveLengthFromRoot,
+					0.0f);
+				PointFacadeOut.ModifyLengthFromSeeds()[PointIndex] = FMath::Max(PointFacadeOut.GetLengthFromSeed(PointIndex) - CarveLengthFromRoot,
+					0.0f);
+			}
+
+			for (const int32& FoliageIndex : FoliageFacadeOut.GetFoliageEntryIdsForBranch(BranchIndex))
+			{
+				const FVector3f UpdatedPivotPoint = FoliageFacadeOut.GetPivotPoint(FoliageIndex) - Delta;
+				FoliageFacadeOut.SetPivotPoint(FoliageIndex, UpdatedPivotPoint);
+
+				const float UpdatedLengthFromRoot = FMath::Max(FoliageFacadeOut.GetLengthFromRoot(FoliageIndex) - CarveLengthFromRoot);
+				FoliageFacadeOut.SetLengthFromRoot(FoliageIndex, UpdatedLengthFromRoot);
+			}
+		}
 	}
+
+	RemoveEntriesAndRecomputeAttributes(OutCollection, SourceCollection, PointsToRemove, BranchesToRemove, FoliageInstancesToRemove);
 }
 
 void FPVCarve::ComputeMetadata(TMap<int32, int32>& OutBranchNumbersToBranchIDs, TMap<int32, float>& OutBranchNumbersToLengthFromRoots,
-                                      const PV::Facades::FBranchFacade& BranchFacadeSource,
-                                      const PV::Facades::FPointFacade& PointFacadeSource)
+                               const PV::Facades::FBranchFacade& BranchFacadeSource,
+                               const PV::Facades::FPointFacade& PointFacadeSource)
 {
 	TRACE_CPUPROFILER_EVENT_SCOPE(PVCarve::ComputeMetadata);
 
@@ -598,25 +628,30 @@ void FPVCarve::ComputeMetadata(TMap<int32, int32>& OutBranchNumbersToBranchIDs,
 	for (int32 BranchIndex = 0; BranchIndex < NumOfBranchesInSource; ++BranchIndex)
 	{
 		const int32 BranchNumber = BranchFacadeSource.GetBranchNumber(BranchIndex);
-		const int32 FirstPointIndex = BranchFacadeSource.GetPoints(BranchIndex)[0];
-		const float BranchLengthFromRoot = PointFacadeSource.GetLengthFromRoot(FirstPointIndex);
 		OutBranchNumbersToBranchIDs.Add(BranchNumber, BranchIndex);
-		OutBranchNumbersToLengthFromRoots.Add(BranchNumber, BranchLengthFromRoot);
+
+		const TArray<int32>& BranchPoints = BranchFacadeSource.GetPoints(BranchIndex);
+		check(BranchPoints.Num() > 0);
+
+		if (BranchPoints.Num() > 0)
+		{
+			const int32 FirstPointIndex = BranchFacadeSource.GetPoints(BranchIndex)[0];
+			const float BranchLengthFromRoot = PointFacadeSource.GetLengthFromRoot(FirstPointIndex);
+			OutBranchNumbersToLengthFromRoots.Add(BranchNumber, BranchLengthFromRoot);
+		}
+		else
+		{
+			OutBranchNumbersToLengthFromRoots.Add(BranchNumber, 0.0f);
+		}
 	}
 }
 
 void FPVCarve::RemoveEntriesAndRecomputeAttributes(FManagedArrayCollection& OutCollection, const FManagedArrayCollection& SourceCollection,
-                                                          TArray<bool>& PointsToRemove, TArray<bool>& BranchesToRemove,
-                                                          TArray<bool>& FoliageInstancesToRemove)
+                                                   TArray<bool>& PointsToRemove, TArray<bool>& BranchesToRemove,
+                                                   TArray<bool>& FoliageInstancesToRemove)
 {
 	TRACE_CPUPROFILER_EVENT_SCOPE(PVCarve::RemoveEntriesAndRecomputeAttributes);
 
-	const PV::Facades::FBranchFacade BranchFacadeSource(SourceCollection);
-	const PV::Facades::FPointFacade PointFacadeSource(SourceCollection);
-
-	PV::Facades::FBranchFacade BranchFacadeOut(OutCollection);
-	PV::Facades::FPointFacade PointFacadeOut(OutCollection);
-
 	PV::Facades::FRemoveEntriesResult RemoveEntriesResult = PV::Facades::FTreeFacade::RemoveEntriesAndReIndexAttributes(
 		OutCollection, PointsToRemove, BranchesToRemove, FoliageInstancesToRemove);
 
@@ -627,5 +662,5 @@ void FPVCarve::RemoveEntriesAndRecomputeAttributes(FManagedArrayCollection& OutC
 		BranchesNewIDsToOldIDs.Add(Pair.Value, Pair.Key);
 	}
 
-	RecomputeAttributes(PointFacadeOut, PointFacadeSource, BranchFacadeOut, BranchFacadeSource, BranchesNewIDsToOldIDs);
+	RecomputeAttributes(OutCollection, SourceCollection, BranchesNewIDsToOldIDs);
 }
diff --git a/Engine/Plugins/Experimental/ProceduralVegetationEditor/Source/ProceduralVegetation/Private/Implementations/PVCarve.h b/Engine/Plugins/Experimental/ProceduralVegetationEditor/Source/ProceduralVegetation/Private/Implementations/PVCarve.h
--- a/Engine/Plugins/Experimental/ProceduralVegetationEditor/Source/ProceduralVegetation/Private/Implementations/PVCarve.h
+++ b/Engine/Plugins/Experimental/ProceduralVegetationEditor/Source/ProceduralVegetation/Private/Implementations/PVCarve.h
@@ -45,9 +45,8 @@ private:
 	                                      const TMap<int32, int32>& PointsOldIDsToNewIDs, const TArray<bool>& PointsRemoved,
 	                                      const TMap<int32, int32>& FoliageInstancesOldIDsToNewIDs, const TArray<bool>& FoliageInstancesRemoved);
 
-	static void RecomputeAttributes(PV::Facades::FPointFacade& PointFacadeOut, const PV::Facades::FPointFacade& PointFacadeSource,
-	                                const PV::Facades::FBranchFacade& BranchFacade,
-	                                const PV::Facades::FBranchFacade& BranchFacadeSource, const TMap<int32, int32>& BranchesNewIDsToOldIDs);
+	static void RecomputeAttributes(FManagedArrayCollection& OutCollection, const FManagedArrayCollection& SourceCollection,
+	                                const TMap<int32, int32>& BranchesNewIDsToOldIDs);
 
 	static void CarveFromTop(FManagedArrayCollection& OutCollection, const FManagedArrayCollection& SourceCollection, const float Carve,
 	                         const ECarveBasis CarveBasis);
@@ -55,8 +54,7 @@ private:
 	static void CarveFromBottom(FManagedArrayCollection& OutCollection, const FManagedArrayCollection& SourceCollection, const float Carve);
 
 	static void ComputeMetadata(TMap<int32, int32>& OutBranchNumbersToBranchIDs, TMap<int32, float>& OutBranchNumbersToLengthFromRoots,
-	                            const PV::Facades::FBranchFacade& BranchFacadeSource,
-	                            const PV::Facades::FPointFacade& PointFacadeSource);
+	                            const PV::Facades::FBranchFacade& BranchFacadeSource, const PV::Facades::FPointFacade& PointFacadeSource);
 
 	static void RemoveEntriesAndRecomputeAttributes(FManagedArrayCollection& OutCollection, const FManagedArrayCollection& SourceCollection,
 	                                                TArray<bool>& PointsToRemove, TArray<bool>& BranchesToRemove,
diff --git a/Engine/Plugins/Experimental/ProceduralVegetationEditor/Source/ProceduralVegetation/Private/Implementations/PVGravity.cpp b/Engine/Plugins/Experimental/ProceduralVegetationEditor/Source/ProceduralVegetation/Private/Implementations/PVGravity.cpp
--- a/Engine/Plugins/Experimental/ProceduralVegetationEditor/Source/ProceduralVegetation/Private/Implementations/PVGravity.cpp
+++ b/Engine/Plugins/Experimental/ProceduralVegetationEditor/Source/ProceduralVegetation/Private/Implementations/PVGravity.cpp
@@ -4,21 +4,23 @@
 #include "Facades/PVBranchFacade.h"
 #include "Facades/PVBudVectorsFacade.h"
 #include "Facades/PVFoliageFacade.h"
+#include "Facades/PVPlantFacade.h"
 #include "Facades/PVPointFacade.h"
 
-void FPVGravity::GeneratePhototropicData(const FPVGravitySettings& InGravitySettings, const PV::Facades::FBranchFacade& InBranchFacade, const PV::Facades::FPointFacade& InPointFacade, TArray<FVector3f>& OutPhototropicDirections)
+void FPVGravity::GeneratePhototropicData(const FPVGravitySettings& InGravitySettings, const PV::Facades::FBranchFacade& InBranchFacade,
+                                         const PV::Facades::FPointFacade& InPointFacade, TArray<FVector3f>& OutPhototropicDirections)
 {
 	FVector3f PrevVec = FVector3f::UpVector;
-		
-	OutPhototropicDirections.Init(PrevVec,InPointFacade.GetElementCount());
+
+	OutPhototropicDirections.Init(PrevVec, InPointFacade.GetElementCount());
 
 	static constexpr int BudDirectionsLightOptimalIndex = 2;
 	static constexpr int BudDirectionsLightSubOptimal = 3;
-		
+
 	for (int BranchIndex = 0; BranchIndex < InBranchFacade.GetElementCount(); BranchIndex++)
 	{
 		const TArray<int32>& BranchPoints = InBranchFacade.GetPoints(BranchIndex);
-			
+
 		for (int i = BranchPoints.Num() - 1; i > 0; i--)
 		{
 			const auto CurrentPointIndex = BranchPoints[i];
@@ -29,7 +31,7 @@ void FPVGravity::GeneratePhototropicData(const FPVGravitySettings& InGravitySett
 				const FVector3f& LightOptimal = BudDirections[BudDirectionsLightOptimalIndex];
 				const FVector3f& LightSubOptimal = BudDirections[BudDirectionsLightSubOptimal];
 
-				FVector3f PhototropicDirection = FMath::Lerp(LightOptimal,LightSubOptimal, InGravitySettings.PhototropicBias);
+				FVector3f PhototropicDirection = FMath::Lerp(LightOptimal, LightSubOptimal, InGravitySettings.PhototropicBias);
 				PhototropicDirection.Normalize();
 
 				if (PhototropicDirection.Length() < 0.001)
@@ -46,31 +48,32 @@ void FPVGravity::GeneratePhototropicData(const FPVGravitySettings& InGravitySett
 					OutPhototropicDirections[CurrentPointIndex] = PhototropicDirection;
 				}
 			}
-		}	
+		}
 	}
 }
 
 void FPVGravity::ApplyGravity(const FPVGravitySettings& InGravitySettings, FManagedArrayCollection& OutCollection)
 {
-	PV::Facades::FBranchFacade BranchFacade = PV::Facades::FBranchFacade(OutCollection);
-	PV::Facades::FPointFacade PointFacade = PV::Facades::FPointFacade(OutCollection);
-		
-	if(!PointFacade.IsValid() || !BranchFacade.IsValid())
+	const PV::Facades::FBranchFacade BranchFacade(OutCollection);
+	const PV::Facades::FPointFacade PointFacade(OutCollection);
+	const PV::Facades::FPlantFacade PlantFacade(OutCollection);
+
+	if (!PointFacade.IsValid() || !BranchFacade.IsValid() || !PlantFacade.IsValid())
 	{
 		return;
 	}
 
-	const TManagedArray<int32>& BranchNumbers = BranchFacade.GetBranchNumbers();
-	const int32 RootBranchIndex = BranchNumbers.Find(1);
-
 	TArray<FVector3f> OutPhototropicDirection;
-	
+
 	if (InGravitySettings.Mode == EGravityMode::Phototropic)
 	{
 		GeneratePhototropicData(InGravitySettings, BranchFacade, PointFacade, OutPhototropicDirection);
 	}
 
-	ApplyGravity(RootBranchIndex, InGravitySettings, OutPhototropicDirection, OutCollection);
+	for (const int32 TrunkIndex : PlantFacade.GetTrunkIndices())
+	{
+		ApplyGravity(TrunkIndex, InGravitySettings, OutPhototropicDirection, OutCollection);
+	}
 }
 
 // Please note that LFR (length from root) for both skeleton points and foliage points are not being scaled by
@@ -87,29 +90,35 @@ void FPVGravity::ApplyGravity(
 	FVector3f PivotPosition
 )
 {
-	static const auto CalculateDownForce = [](const FVector3f InRelativeDirection, const int InPointIndex, const float InPScaleGradient, const FPVGravitySettings& InGravitySettings, const TArray<FVector3f>& InPhototropicDirections) -> FQuat4f
-	{
-		const FVector3f PhototropicDirection = InPhototropicDirections.Num() > InPointIndex ? InPhototropicDirections[InPointIndex] : FVector3f::UpVector;
-		const FVector3f GravityDirection = 	InGravitySettings.Mode == EGravityMode::Phototropic ? PhototropicDirection : InGravitySettings.Direction; 
-
-		float YDot = 1.0f - FMath::Abs(FVector3f::DotProduct(InRelativeDirection,InGravitySettings.Direction));
-		YDot = FMath::GetMappedRangeValueClamped(FVector2f(0, 1),
-		FVector2f(FMath::Max(0.1f, 1 - InGravitySettings.AngleCorrection), 1),
-		          YDot);
-		
-		const float GravityImpact = InGravitySettings.Gravity * (1 - InPScaleGradient) * YDot;
-		const FVector3f GravityForce = FMath::Lerp(InRelativeDirection, GravityDirection, GravityImpact).GetUnsafeNormal();
-		return FQuat4f::FindBetween(InRelativeDirection, GravityForce);
-	};
+	static const auto CalculateDownForce = [](const FVector3f InRelativeDirection, const int InPointIndex, const float InPScaleGradient,
+	                                          const FPVGravitySettings& InGravitySettings,
+	                                          const TArray<FVector3f>& InPhototropicDirections) -> FQuat4f
+		{
+			const FVector3f PhototropicDirection = InPhototropicDirections.Num() > InPointIndex
+				? InPhototropicDirections[InPointIndex]
+				: FVector3f::UpVector;
+			const FVector3f GravityDirection = InGravitySettings.Mode == EGravityMode::Phototropic
+				? PhototropicDirection
+				: InGravitySettings.Direction;
+
+			float YDot = 1.0f - FMath::Abs(FVector3f::DotProduct(InRelativeDirection, InGravitySettings.Direction));
+			YDot = FMath::GetMappedRangeValueClamped(FVector2f(0, 1),
+				FVector2f(FMath::Max(0.1f, 1 - InGravitySettings.AngleCorrection), 1),
+				YDot);
+
+			const float GravityImpact = InGravitySettings.Gravity * (1 - InPScaleGradient) * YDot;
+			const FVector3f GravityForce = FMath::Lerp(InRelativeDirection, GravityDirection, GravityImpact).GetUnsafeNormal();
+			return FQuat4f::FindBetween(InRelativeDirection, GravityForce);
+		};
 
 	PV::Facades::FBranchFacade BranchFacade = PV::Facades::FBranchFacade(OutCollection);
 	PV::Facades::FPointFacade PointFacade = PV::Facades::FPointFacade(OutCollection);
 	PV::Facades::FBudVectorsFacade BudVectorFacade = PV::Facades::FBudVectorsFacade(OutCollection);
 	PV::Facades::FFoliageFacade FoliageFacade = PV::Facades::FFoliageFacade(OutCollection);
-	
+
 	TManagedArray<FVector3f>& PointPositions = PointFacade.ModifyPositions();
 	TManagedArray<TArray<FVector3f>>& PointBudDirections = BudVectorFacade.ModifyBudDirections();
-	
+
 	const TArray<int32>& CurrentBranchPoints = BranchFacade.GetPoints(BranchIndex);
 	const TArray<int32>& CurrentBranchChildren = BranchFacade.GetChildren(BranchIndex);
 	const TManagedArray<int32>& BranchNumbers = BranchFacade.GetBranchNumbers();
@@ -121,7 +130,7 @@ void FPVGravity::ApplyGravity(
 	{
 		const int32 ChildIndex = BranchNumbers.Find(Child);
 		const TArray<int32>& Parents = BranchFacade.GetParents(ChildIndex);
-		
+
 		if (ChildIndex != INDEX_NONE && !Parents.IsEmpty() && Parents[Parents.Num() - 1] == BranchNumbers[BranchIndex])
 		{
 			ChildrenIndices.Add(ChildIndex);
@@ -137,7 +146,8 @@ void FPVGravity::ApplyGravity(
 		const float& PScaleGradient = PointFacade.GetPointScaleGradient(PointIndex);
 
 		const FVector3f RelativeDirection = TotalDownForce * (Position - PreviousPosition);
-		const FQuat4f DownForce = CalculateDownForce(RelativeDirection.GetUnsafeNormal(), PointIndex, PScaleGradient, GravitySettings, PhototropicDirections);
+		const FQuat4f DownForce = CalculateDownForce(RelativeDirection.GetUnsafeNormal(), PointIndex, PScaleGradient, GravitySettings,
+			PhototropicDirections);
 
 		TotalDownForce = DownForce * TotalDownForce;
 		PreviousPosition = Position;
@@ -170,7 +180,8 @@ void FPVGravity::ApplyGravity(
 		const FVector3f RelativeDirection = TotalDownForce * (Position - PreviousPosition);
 		const FVector3f Tangent = RelativeDirection.GetUnsafeNormal();
 
-		const FQuat4f DownForce = CalculateDownForce(Tangent, PointIndex, PointFacade.GetPointScaleGradient(PointIndex), GravitySettings, PhototropicDirections);
+		const FQuat4f DownForce = CalculateDownForce(Tangent, PointIndex, PointFacade.GetPointScaleGradient(PointIndex), GravitySettings,
+			PhototropicDirections);
 		TotalDownForce = DownForce * TotalDownForce;
 
 		// Update attachment points for all foliage elements for current branch between the current and previous point 
@@ -189,15 +200,15 @@ void FPVGravity::ApplyGravity(
 			}
 		}
 
-		for (const int32 ChildrenIndex : ChildrenIndices)
+		for (const int32 ChildIndex : ChildrenIndices)
 		{
-			const int32 FirstPointIndex = BranchFacade.GetPoints(ChildrenIndex)[0];
+			const int32 FirstPointIndex = BranchFacade.GetPoints(ChildIndex)[0];
 			const float ChildLengthFromRoot = PointFacade.GetLengthFromRoot(FirstPointIndex);
 			if (ChildLengthFromRoot <= LengthFromRoot && ChildLengthFromRoot > PreviousLengthFromRoot)
 			[[unlikely]]
 			{
 				ApplyGravity(
-					ChildrenIndex,
+					ChildIndex,
 					GravitySettings,
 					PhototropicDirections,
 					OutCollection,
diff --git a/Engine/Plugins/Experimental/ProceduralVegetationEditor/Source/ProceduralVegetation/Private/Implementations/PVMesher.cpp b/Engine/Plugins/Experimental/ProceduralVegetationEditor/Source/ProceduralVegetation/Private/Implementations/PVMesher.cpp
--- a/Engine/Plugins/Experimental/ProceduralVegetationEditor/Source/ProceduralVegetation/Private/Implementations/PVMesher.cpp
+++ b/Engine/Plugins/Experimental/ProceduralVegetationEditor/Source/ProceduralVegetation/Private/Implementations/PVMesher.cpp
@@ -10,10 +10,11 @@
 #include "Facades/PVBranchFacade.h"
 #include "Facades/PVBudVectorsFacade.h"
 #include "Facades/PVFoliageFacade.h"
+#include "Facades/PVPlantFacade.h"
 #include "Facades/PVPointFacade.h"
 
 void FPVMesher::GenerateGeometryCollection(const FManagedArrayCollection& InSkeletonCollection, const FPVMesherSettings& MesherSettings,
-                                                  FGeometryCollection& OutGeometryCollection)
+                                           FGeometryCollection& OutGeometryCollection)
 {
 	InSkeletonCollection.CopyTo(&OutGeometryCollection);
 
@@ -185,7 +186,7 @@ void FPVMesher::SetNjordPixelID(FGeometryCollection& OutGeometryCollection)
 }
 
 void FPVMesher::GenerateDynamicMesh(FManagedArrayCollection& Collection, const FPVMesherSettings& MesherSettings,
-                                           TObjectPtr<UDynamicMesh>& OutMesh)
+                                    TObjectPtr<UDynamicMesh>& OutMesh)
 {
 	TArray<FLocalDynamicMeshData> MeshesDataCollection;
 	GetLocalDynamicMeshData(Collection, MesherSettings, MeshesDataCollection);
@@ -208,12 +209,10 @@ void FPVMesher::GenerateDynamicMesh(FManagedArrayCollection& Collection, const F
 		EDynamicMeshChangeType::MeshChange);
 }
 
-TSet<int32> FPVMesher::CollectHardPoints(const PV::Facades::FBranchFacade& BranchFacade,
-                                                const PV::Facades::FPointFacade& PointFacade)
+TSet<int32> FPVMesher::CollectHardPoints(const PV::Facades::FBranchFacade& BranchFacade, const PV::Facades::FPointFacade& PointFacade,
+                                         const PV::Facades::FPlantFacade& PlantFacade)
 {
 	TSet<int32> HardPoints;
-	const int32 MainTrunkIndex = BranchFacade.GetTrunkIndex();
-	const int32 SeedPointIndex = BranchFacade.GetPoints(MainTrunkIndex)[0];
 	TSet<int32> BranchSourceBudNumbers;
 	BranchSourceBudNumbers.Reserve(BranchFacade.GetElementCount());
 
@@ -235,11 +234,12 @@ TSet<int32> FPVMesher::CollectHardPoints(const PV::Facades::FBranchFacade& Branc
 			const int32 PointBudNumber = PointFacade.GetBudNumber(PointIndex);
 			const float NjordPixelIndex = PointFacade.GetNjordPixelIndex(PointIndex);
 
+			bool bIsSeedPoint = PlantFacade.IsTrunkIndex(BranchIndex) && PointIndex == FirstPointIndex;
 			bool bIsBranchJoint = (NjordPixelIndex == FMath::FloorToFloat(NjordPixelIndex)) && BranchSourceBudNumbers.Contains(PointBudNumber);
 
-			if (PointIndex == SeedPointIndex ||
-				PointIndex == FirstPointIndex ||
+			if (PointIndex == FirstPointIndex ||
 				PointIndex == LastPointIndex ||
+				bIsSeedPoint ||
 				bIsBranchJoint
 			)
 			{
@@ -252,9 +252,9 @@ TSet<int32> FPVMesher::CollectHardPoints(const PV::Facades::FBranchFacade& Branc
 }
 
 TSet<int32> FPVMesher::ComputePointGradients(const PV::Facades::FPointFacade& PointFacade,
-                                                    const FPVMesherSettings& MesherSettings, const TSet<int32>& HardPoints,
-                                                    const float MaxPointScale, TArray<float>& OutMeshDivisionsGradients,
-                                                    TArray<float>& OutDeltaModifiers)
+                                             const FPVMesherSettings& MesherSettings, const TSet<int32>& HardPoints,
+                                             const float MaxPointScale, TArray<float>& OutMeshDivisionsGradients,
+                                             TArray<float>& OutDeltaModifiers)
 {
 	TSet<int32> PointsToRemove;
 	OutMeshDivisionsGradients.Reserve(PointFacade.GetElementCount());
@@ -334,8 +334,8 @@ TSet<int32> FPVMesher::ComputePointGradients(const PV::Facades::FPointFacade& Po
 }
 
 void FPVMesher::PerformPathSimplification(const PV::Facades::FBranchFacade& BranchFacade,
-                                                 const PV::Facades::FPointFacade& PointFacade, const FPVMesherSettings& MesherSettings,
-                                                 const TSet<int32>& HardPoints, const TArray<float>& DeltaModifiers, TSet<int32>& InOutPointsToRemove)
+                                          const PV::Facades::FPointFacade& PointFacade, const FPVMesherSettings& MesherSettings,
+                                          const TSet<int32>& HardPoints, const TArray<float>& DeltaModifiers, TSet<int32>& InOutPointsToRemove)
 {
 	const int32 Iterations = MesherSettings.Accuracy;
 	// Epsilon is the maximum allowed curvature error
@@ -428,8 +428,8 @@ void FPVMesher::PerformPathSimplification(const PV::Facades::FBranchFacade& Bran
 }
 
 TArray<int32> FPVMesher::ComputeMeshDivisions(const PV::Facades::FBranchFacade& BranchFacade,
-                                                     const FPVMesherSettings& MesherSettings, const TArray<float>& MeshDivisionsGradients,
-                                                     const int32 PointCount)
+                                              const FPVMesherSettings& MesherSettings, const TArray<float>& MeshDivisionsGradients,
+                                              const int32 PointCount)
 {
 	TArray<int32> TargetMeshDivisions;
 	static constexpr int32 MinimumDivisions = 3;
@@ -462,7 +462,7 @@ TArray<int32> FPVMesher::ComputeMeshDivisions(const PV::Facades::FBranchFacade&
 }
 
 void FPVMesher::TriangulateRings(const TArray<int32>& PreviousIndices, const TArray<int32>& CurrentIndices, int32& InOutPolyGroupIndex,
-                                        FLocalDynamicMeshData& OutMeshData)
+                                 FLocalDynamicMeshData& OutMeshData)
 {
 	if (!PreviousIndices.IsEmpty())
 	[[likely]]
@@ -564,9 +564,9 @@ TMap<int32, TArray<int32>> FPVMesher::GetPointsIndicesToFoliageIndicesMap(const
 }
 
 void FPVMesher::GenerateBranchMeshData(const bool bPrimitiveIsTrunk, const int32 GenerationNumber, const TArray<int32>& PrimitivePoints,
-                                              const FDisplacementData& DisplacementData, const FPVMesherSettings& MesherSettings,
-                                              const TArray<int32>& TargetMeshDivisions, const FManagedArrayCollection& Collection,
-                                              FLocalDynamicMeshData& OutLocalMeshData)
+                                       const FDisplacementData& DisplacementData, const FPVMesherSettings& MesherSettings,
+                                       const TArray<int32>& TargetMeshDivisions, const FManagedArrayCollection& Collection,
+                                       FLocalDynamicMeshData& OutLocalMeshData)
 {
 	const PV::Facades::FPointFacade PointFacade(Collection);
 	const PV::Facades::FBudVectorsFacade BudVectorFacade(Collection);
@@ -782,10 +782,11 @@ void FPVMesher::UpdateFoliagePivotPoints(const TSet<int32>& PointsToRemove, FMan
 }
 
 void FPVMesher::GetLocalDynamicMeshData(FManagedArrayCollection& Collection, const FPVMesherSettings& MesherSettings,
-                                               TArray<FLocalDynamicMeshData>& MeshesData)
+                                        TArray<FLocalDynamicMeshData>& MeshesData)
 {
 	const PV::Facades::FBranchFacade BranchFacade(Collection);
 	const PV::Facades::FPointFacade PointFacade(Collection);
+	const PV::Facades::FPlantFacade PlantFacade(Collection);
 
 	if (!BranchFacade.IsValid())
 	{
@@ -796,7 +797,7 @@ void FPVMesher::GetLocalDynamicMeshData(FManagedArrayCollection& Collection, con
 	const TManagedArray<float>& PointScales = PointFacade.GetPointScales();
 	const float MaxPointScale = *Algo::MaxElement(PointScales);
 
-	const TSet<int32> HardPoints = CollectHardPoints(BranchFacade, PointFacade);
+	const TSet<int32> HardPoints = CollectHardPoints(BranchFacade, PointFacade, PlantFacade);
 
 	TArray<float> MeshDivisionsGradientsForPoints;
 	TArray<float> DeltaModifiersForPoints;
@@ -845,7 +846,7 @@ void FPVMesher::GetLocalDynamicMeshData(FManagedArrayCollection& Collection, con
 							return !PointsToRemove.Contains(PointIndex);
 						});
 
-				const bool PrimitiveIsTrunk = BranchFacade.GetTrunkIndex() == PrimitiveIndex;
+				const bool PrimitiveIsTrunk = PlantFacade.IsTrunkIndex(PrimitiveIndex);
 				const int32 GenerationNumber = BranchFacade.GetGenerationNumber(PrimitiveIndex);
 
 				MeshesData[PrimitiveIndex].MaterialID = BranchFacade.GetBranchUVMaterial(PrimitiveIndex);
diff --git a/Engine/Plugins/Experimental/ProceduralVegetationEditor/Source/ProceduralVegetation/Private/Implementations/PVMesher.h b/Engine/Plugins/Experimental/ProceduralVegetationEditor/Source/ProceduralVegetation/Private/Implementations/PVMesher.h
--- a/Engine/Plugins/Experimental/ProceduralVegetationEditor/Source/ProceduralVegetation/Private/Implementations/PVMesher.h
+++ b/Engine/Plugins/Experimental/ProceduralVegetationEditor/Source/ProceduralVegetation/Private/Implementations/PVMesher.h
@@ -13,6 +13,7 @@
 #include "PVMaterialSettings.h"
 #include "Curves/CurveFloat.h"
 #include "Engine/Texture2D.h"
+#include "Facades/PVPlantFacade.h"
 
 #include "PVMesher.generated.h"
 
@@ -298,8 +299,8 @@ struct FPVMesher
 	static void ExtractDisplacementData(const TObjectPtr<UTexture2D>& Texture, TArray<float>& OutValues);
 
 private:
-	static TSet<int32> CollectHardPoints(const PV::Facades::FBranchFacade& BranchFacade,
-	                                     const PV::Facades::FPointFacade& PointFacade);
+	static TSet<int32> CollectHardPoints(const PV::Facades::FBranchFacade& BranchFacade, const PV::Facades::FPointFacade& PointFacade,
+	                                     const PV::Facades::FPlantFacade& PlantFacade);
 
 	static TSet<int32> ComputePointGradients(const PV::Facades::FPointFacade& PointFacade, const FPVMesherSettings& MesherSettings,
 	                                         const TSet<int32>& HardPoints, const float MaxPointScale, TArray<float>& OutMeshDivisionsGradients,
diff --git a/Engine/Plugins/Experimental/ProceduralVegetationEditor/Source/ProceduralVegetation/Private/Implementations/PVRemoveBranches.cpp b/Engine/Plugins/Experimental/ProceduralVegetationEditor/Source/ProceduralVegetation/Private/Implementations/PVRemoveBranches.cpp
--- a/Engine/Plugins/Experimental/ProceduralVegetationEditor/Source/ProceduralVegetation/Private/Implementations/PVRemoveBranches.cpp
+++ b/Engine/Plugins/Experimental/ProceduralVegetationEditor/Source/ProceduralVegetation/Private/Implementations/PVRemoveBranches.cpp
@@ -2,104 +2,126 @@
 
 #include "PVRemoveBranches.h"
 #include "Facades/PVBranchFacade.h"
-#include "Facades/PVBudVectorsFacade.h"
-#include "Facades/PVFoliageFacade.h"
+#include "Facades/PVPlantFacade.h"
 #include "Facades/PVPointFacade.h"
 #include "Facades/PVTreeFacade.h"
 
 namespace PV::Facades
 {
 	static constexpr int BranchBudLightDetectedIndex = 2;
-	static constexpr int BudDevelopementGenerationIndex = 0;
-	static constexpr int BudDevelopementAgeIndex = 2;
-	
-	void GatherBranchRemovalData(const FBranchFacade& InBranchFacade, const FPointFacade& InPointFacade,const ERemoveBranchesBasis InBasis, TArray<float>& OutRemovalData)
+	static constexpr int BudDevelopmentGenerationIndex = 0;
+	static constexpr int BudDevelopmentAgeIndex = 2;
+
+	void GatherBranchRemovalData(const FBranchFacade& InBranchFacade, const FPointFacade& InPointFacade, const FPlantFacade& InPlantFacade,
+	                             const ERemoveBranchesBasis InBasis, TArray<float>& OutRemovalData)
 	{
-		const int32 NumBranches = InBranchFacade.GetElementCount();
-		const int32 RootBranchIndex = InBranchFacade.GetTrunkIndex();
+		const int32 TotalNumberOfBranches = InBranchFacade.GetElementCount();
+		OutRemovalData.Init(-1.0f, TotalNumberOfBranches);
 
-		check(RootBranchIndex !=  INDEX_NONE);
-		
-		if (RootBranchIndex ==  INDEX_NONE)
-			return;
-		
-		const TArray<int32>& RootBranchPoints = InBranchFacade.GetPoints(RootBranchIndex);
+		for (const TMap<int32, int32> PlantNumbersToTrunkIDs = InPlantFacade.GetPlantNumbersToTrunkIndicesMap();
+		     const TPair<int32, int32> Pair : PlantNumbersToTrunkIDs)
+		{
+			const int32 PlantNumber = Pair.Key;
+			const int32 TrunkIndex = Pair.Value;
+			const TArray<int32> BranchIndices = InPlantFacade.GetBranchIndices(PlantNumber);
 
-		check(RootBranchPoints.Num() > 0);
-		
-		if (RootBranchPoints.Num() == 0)
-			return;
+			check(TrunkIndex != INDEX_NONE);
 
-		OutRemovalData.AddUninitialized(NumBranches);
+			if (TrunkIndex == INDEX_NONE)
+				return;
 
-		const TArray<int>& RootBudDevelopment = InPointFacade.GetBudDevelopment(RootBranchPoints[0]);
+			const TArray<int32>& TrunkPoints = InBranchFacade.GetPoints(TrunkIndex);
 
-		float MaxValue = std::numeric_limits<float>::lowest();
-		
-		for (int32 BranchIndex = 0; BranchIndex < NumBranches; ++BranchIndex)
-		{
-			const TArray<int32>& BranchPoints = InBranchFacade.GetPoints(BranchIndex);
+			check(TrunkPoints.Num() > 0);
 
-			if (InBasis == ERemoveBranchesBasis::Length)
-			{
-				const float BranchRootPointLFR = InPointFacade.GetLengthFromRoot(BranchPoints[0]);
-				const float BranchLastPointLFR = InPointFacade.GetLengthFromRoot(BranchPoints[BranchPoints.Num() - 1]);
+			if (TrunkPoints.Num() == 0)
+				return;
 
-				OutRemovalData[BranchIndex] = BranchLastPointLFR - BranchRootPointLFR;
-			}
-			else if (InBasis == ERemoveBranchesBasis::Radius)
-			{
-				OutRemovalData[BranchIndex] = InPointFacade.GetPointScale(BranchPoints[0]);
-			}
-			else if (InBasis == ERemoveBranchesBasis::Light)
-			{
-				const TArray<float>& PointBudLightDetected = InPointFacade.GetBudLightDetected(BranchPoints[BranchPoints.Num() - 2]);
+			const int32 RootPointIndex = TrunkPoints[0];
 
-				// TODO: Some entries dont have complete data. Need to update the JSON to fix it. For now, setting -1 value to keep track of inaccurate data. 
-				OutRemovalData[BranchIndex] = PointBudLightDetected.Num() > BranchBudLightDetectedIndex ? PointBudLightDetected[BranchBudLightDetectedIndex] : -1.0f;
-			}
-			else if (InBasis == ERemoveBranchesBasis::Age ||  InBasis == ERemoveBranchesBasis::Generation)
+			const TArray<int>& RootBudDevelopment = InPointFacade.GetBudDevelopment(RootPointIndex);
+
+			float MaxValue = std::numeric_limits<float>::lowest();
+
+			for (const int32 BranchIndex : BranchIndices)
 			{
-				const TArray<int>& PointBudDevelopment = InPointFacade.GetBudDevelopment(BranchPoints[BranchPoints.Num() - 1]);
-				
-				if (InBasis == ERemoveBranchesBasis::Age)
+				const TArray<int32>& BranchPoints = InBranchFacade.GetPoints(BranchIndex);
+
+				if (BranchPoints.Num() == 0)
 				{
-					OutRemovalData[BranchIndex] = (float)PointBudDevelopment[BudDevelopementAgeIndex]/RootBudDevelopment[BudDevelopementAgeIndex];
+					continue;
 				}
-				else if (InBasis == ERemoveBranchesBasis::Generation)
+
+				if (InBasis == ERemoveBranchesBasis::Length)
 				{
-					OutRemovalData[BranchIndex] = (float)(RootBudDevelopment[BudDevelopementGenerationIndex])/PointBudDevelopment[BudDevelopementGenerationIndex];
+					const float BranchRootPointLFR = InPointFacade.GetLengthFromRoot(BranchPoints[0]);
+					const float BranchLastPointLFR = InPointFacade.GetLengthFromRoot(BranchPoints[BranchPoints.Num() - 1]);
+
+					OutRemovalData[BranchIndex] = BranchLastPointLFR - BranchRootPointLFR;
 				}
-			}
+				else if (InBasis == ERemoveBranchesBasis::Radius)
+				{
+					OutRemovalData[BranchIndex] = InPointFacade.GetPointScale(BranchPoints[0]);
+				}
+				else if (InBasis == ERemoveBranchesBasis::Light)
+				{
+					if (BranchPoints.Num() < 2)
+					{
+						continue;
+					}
 
-			MaxValue = FMath::Max(OutRemovalData[BranchIndex], MaxValue);
-		}
+					const TArray<float>& PointBudLightDetected = InPointFacade.GetBudLightDetected(BranchPoints[BranchPoints.Num() - 2]);
 
-		if (MaxValue > 1)
-		{
-			for (int i = 0; i < OutRemovalData.Num(); i++)
+					// TODO: Some entries dont have complete data. Need to update the JSON to fix it. For now, setting -1 value to keep track of inaccurate data. 
+					OutRemovalData[BranchIndex] = PointBudLightDetected.Num() > BranchBudLightDetectedIndex
+						? PointBudLightDetected[BranchBudLightDetectedIndex]
+						: -1.0f;
+				}
+				else if (InBasis == ERemoveBranchesBasis::Age || InBasis == ERemoveBranchesBasis::Generation)
+				{
+					const TArray<int>& PointBudDevelopment = InPointFacade.GetBudDevelopment(BranchPoints[BranchPoints.Num() - 1]);
+
+					if (InBasis == ERemoveBranchesBasis::Age)
+					{
+						const int RootBudAge = RootBudDevelopment[BudDevelopmentAgeIndex];
+						OutRemovalData[BranchIndex] = RootBudAge > 0
+							? static_cast<float>(PointBudDevelopment[BudDevelopmentAgeIndex]) / static_cast<float>(RootBudAge)
+							: -1.0f;
+					}
+					else if (InBasis == ERemoveBranchesBasis::Generation)
+					{
+						const int PointBudGeneration = PointBudDevelopment[BudDevelopmentGenerationIndex];
+						OutRemovalData[BranchIndex] = PointBudGeneration > 0
+							? static_cast<float>(RootBudDevelopment[BudDevelopmentGenerationIndex]) / static_cast<float>(PointBudGeneration)
+							: -1.0f;
+					}
+				}
+
+				MaxValue = FMath::Max(OutRemovalData[BranchIndex], MaxValue);
+			}
+
+			if (MaxValue > 1)
 			{
-				OutRemovalData[i] = OutRemovalData[i] / MaxValue;
+				for (const int32 BranchIndex : BranchIndices)
+				{
+					OutRemovalData[BranchIndex] = OutRemovalData[BranchIndex] / MaxValue;
+				}
 			}
 		}
 	}
 
-	void FindBrancehesToRemove(const FBranchFacade& InBranchFacade, const TArray<float>& RemovalData, const float& InThreshold, TArray<int>& OutBranchesToRemove)
+	void FindBranchesToRemove(const FBranchFacade& InBranchFacade, const TArray<float>& RemovalData, const float& InThreshold,
+	                          TArray<int>& OutBranchesToRemove)
 	{
 		const int32 NumBranches = InBranchFacade.GetElementCount();
-		
+		OutBranchesToRemove.Reserve(NumBranches / 2);
+
 		for (int32 BranchIndex = 0; BranchIndex < NumBranches; ++BranchIndex)
 		{
-			float Val = RemovalData[BranchIndex];
-			
-			if (Val != -1.0 && Val < InThreshold)
+			if (const float Val = RemovalData[BranchIndex];
+				Val > 0.0f && Val < InThreshold)
 			{
-				if (!OutBranchesToRemove.Contains(BranchIndex))
-				{
-					OutBranchesToRemove.Add(BranchIndex);
-				}
-				
-				UE_LOG(LogTemp, Log, TEXT("PVRemoveBranches :: RemoveBranches : Removing branch index : %i"), BranchIndex);
+				OutBranchesToRemove.Add(BranchIndex);
 			}
 		}
 	}
@@ -107,20 +129,20 @@ namespace PV::Facades
 
 void FPVRemoveBranches::ApplyRemoveBranches(const ERemoveBranchesBasis InBasis, const float InThreshold, FManagedArrayCollection& OutCollection)
 {
-	PV::Facades::FBranchFacade BranchFacade = PV::Facades::FBranchFacade(OutCollection);
-	PV::Facades::FPointFacade PointFacade = PV::Facades::FPointFacade(OutCollection);
-	PV::Facades::FFoliageFacade FoliageFacade = PV::Facades::FFoliageFacade(OutCollection);
-		
-	if(!PointFacade.IsValid() || !BranchFacade.IsValid())
+	const PV::Facades::FBranchFacade BranchFacade(OutCollection);
+	const PV::Facades::FPointFacade PointFacade(OutCollection);
+	const PV::Facades::FPlantFacade PlantFacade(OutCollection);
+
+	if (!PointFacade.IsValid() || !BranchFacade.IsValid())
 	{
 		return;
 	}
-	
+
 	TArray<float> OutRemovalData;
-	GatherBranchRemovalData(BranchFacade, PointFacade, InBasis,OutRemovalData);
+	GatherBranchRemovalData(BranchFacade, PointFacade, PlantFacade, InBasis, OutRemovalData);
 
 	TArray<int> OutBranchesToRemove;
-	FindBrancehesToRemove(BranchFacade, OutRemovalData, InThreshold, OutBranchesToRemove);
+	FindBranchesToRemove(BranchFacade, OutRemovalData, InThreshold, OutBranchesToRemove);
 
 	PV::Facades::FTreeFacade::RemoveBranches(BranchFacade, OutBranchesToRemove, OutCollection);
 }
diff --git a/Engine/Plugins/Experimental/ProceduralVegetationEditor/Source/ProceduralVegetation/Private/Implementations/PVRemoveBranches.h b/Engine/Plugins/Experimental/ProceduralVegetationEditor/Source/ProceduralVegetation/Private/Implementations/PVRemoveBranches.h
--- a/Engine/Plugins/Experimental/ProceduralVegetationEditor/Source/ProceduralVegetation/Private/Implementations/PVRemoveBranches.h
+++ b/Engine/Plugins/Experimental/ProceduralVegetationEditor/Source/ProceduralVegetation/Private/Implementations/PVRemoveBranches.h
@@ -19,5 +19,5 @@ enum class ERemoveBranchesBasis : uint8
 
 struct FPVRemoveBranches
 {
-	static void ApplyRemoveBranches(const ERemoveBranchesBasis InBasis, const float Threshold, FManagedArrayCollection& OutCollection);
+	static void ApplyRemoveBranches(const ERemoveBranchesBasis InBasis, const float InThreshold, FManagedArrayCollection& OutCollection);
 };
diff --git a/Engine/Plugins/Experimental/ProceduralVegetationEditor/Source/ProceduralVegetation/Private/Implementations/PVScale.cpp b/Engine/Plugins/Experimental/ProceduralVegetationEditor/Source/ProceduralVegetation/Private/Implementations/PVScale.cpp
--- a/Engine/Plugins/Experimental/ProceduralVegetationEditor/Source/ProceduralVegetation/Private/Implementations/PVScale.cpp
+++ b/Engine/Plugins/Experimental/ProceduralVegetationEditor/Source/ProceduralVegetation/Private/Implementations/PVScale.cpp
@@ -5,75 +5,81 @@
 #include "Facades/PVBranchFacade.h"
 #include "Facades/PVBudVectorsFacade.h"
 #include "Facades/PVFoliageFacade.h"
+#include "Facades/PVPlantFacade.h"
 #include "Facades/PVPointFacade.h"
 
 void FPVScale::ApplyScale(const float InManualScale, FManagedArrayCollection& OutCollection)
 {
-	PV::Facades::FBranchFacade BranchFacade = PV::Facades::FBranchFacade(OutCollection);
-	PV::Facades::FPointFacade PointFacade = PV::Facades::FPointFacade(OutCollection);
-		
-	if(!PointFacade.IsValid() || !BranchFacade.IsValid())
+	PV::Facades::FBranchFacade BranchFacade(OutCollection);
+	PV::Facades::FPointFacade PointFacade(OutCollection);
+	PV::Facades::FPlantFacade PlantFacade(OutCollection);
+	PV::Facades::FFoliageFacade FoliageFacade(OutCollection);
+
+	if (!PointFacade.IsValid() || !BranchFacade.IsValid() || !PlantFacade.IsValid())
 	{
 		return;
 	}
 
-	const int32 RootBranchIndex = BranchFacade.GetTrunkIndex();
-	const int32 NumBranches = BranchFacade.GetElementCount();
+	for (const TMap<int32, int32> PlantNumbersToTrunkIDs = PlantFacade.GetPlantNumbersToTrunkIndicesMap();
+	     const TPair<int32, int32> Pair : PlantNumbersToTrunkIDs)
+	{
+		const int32 PlantNumber = Pair.Key;
+		const int32 TrunkIndex = Pair.Value;
+		const TArray<int32> BranchIndices = PlantFacade.GetBranchIndices(PlantNumber);
 
-	check(RootBranchIndex !=  INDEX_NONE);
-	
-	if (RootBranchIndex ==  INDEX_NONE)
-		return;
-	
-	const TArray<int32>& RootBranchPoints = BranchFacade.GetPoints(RootBranchIndex);
+		check(TrunkIndex != INDEX_NONE);
 
-	check(RootBranchPoints.Num() > 0);
-	
-	if (RootBranchPoints.Num() == 0)
-		return;
-	
-	const int32 RootPointIndex = RootBranchPoints[0];
-	const FVector3f& PivotPos = PointFacade.GetPosition(RootPointIndex);
-	const float RootPointLFR = PointFacade.GetLengthFromRoot(RootPointIndex);
-	
-	for (int32 BranchIndex = 0; BranchIndex < NumBranches; ++BranchIndex)
-	{
-		const TArray<int32>& BranchPoints = BranchFacade.GetPoints(BranchIndex);
+		if (TrunkIndex == INDEX_NONE)
+			return;
+
+		const TArray<int32>& TrunkPoints = BranchFacade.GetPoints(TrunkIndex);
+
+		check(TrunkPoints.Num() > 0);
 
-		for (const int32& PointIndex : BranchPoints)
+		if (TrunkPoints.Num() == 0)
+			return;
+
+		const int32 RootPointIndex = TrunkPoints[0];
+		const FVector3f& PivotPos = PointFacade.GetPosition(RootPointIndex);
+		const float RootPointLFR = PointFacade.GetLengthFromRoot(RootPointIndex);
+
+		for (const int32 BranchIndex : BranchIndices)
 		{
-			FVector3f PointPos = PointFacade.GetPosition(PointIndex);
-			float LFS = PointFacade.GetLengthFromSeed(PointIndex);
+			for (const TArray<int32>& BranchPoints = BranchFacade.GetPoints(BranchIndex);
+			     const int32& PointIndex : BranchPoints)
+			{
+				FVector3f PointPos = PointFacade.GetPosition(PointIndex);
+				float LFS = PointFacade.GetLengthFromSeed(PointIndex);
 
-			PointPos -= PivotPos;
-			PointPos *= InManualScale;
-			PointPos += PivotPos;
+				PointPos -= PivotPos;
+				PointPos *= InManualScale;
+				PointPos += PivotPos;
 
-			LFS *= InManualScale;
+				LFS *= InManualScale;
 
-			float LFR = LFS + RootPointLFR;
+				float LFR = LFS + RootPointLFR;
 
-			PointFacade.ModifyPositions()[PointIndex] = PointPos;
-			PointFacade.ModifyLengthFromSeeds()[PointIndex] = LFS;
-			PointFacade.ModifyLengthFromRoots()[PointIndex] = LFR;
-			PointFacade.ModifyPointScales()[PointIndex] *= InManualScale;
-		}
-	}
+				PointFacade.ModifyPositions()[PointIndex] = PointPos;
+				PointFacade.ModifyLengthFromSeeds()[PointIndex] = LFS;
+				PointFacade.ModifyLengthFromRoots()[PointIndex] = LFR;
+				PointFacade.ModifyPointScales()[PointIndex] *= InManualScale;
+			}
 
-	PV::Facades::FFoliageFacade FoliageFacade(OutCollection);
-	int32 NumInstances = FoliageFacade.NumFoliageEntries();
+			const TArray<int32>& FoliageEntryIds = FoliageFacade.GetFoliageEntryIdsForBranch(BranchIndex);
 
-	for(int32 Id = 0; Id < NumInstances; Id++)
-	{
-		PV::Facades::FFoliageEntryData Data = FoliageFacade.GetFoliageEntry(Id);
+			for (const int32 FoliageIndex : FoliageEntryIds)
+			{
+				PV::Facades::FFoliageEntryData Data = FoliageFacade.GetFoliageEntry(FoliageIndex);
 
-		Data.PivotPoint -= PivotPos;
-		Data.PivotPoint *= InManualScale;
-		Data.PivotPoint += PivotPos;
-		
-		Data.Scale *= InManualScale;
-		Data.LengthFromRoot = Data.LengthFromRoot * InManualScale;
+				Data.PivotPoint -= PivotPos;
+				Data.PivotPoint *= InManualScale;
+				Data.PivotPoint += PivotPos;
 
-		FoliageFacade.SetFoliageEntry(Id, Data);
+				Data.Scale *= InManualScale;
+				Data.LengthFromRoot = Data.LengthFromRoot * InManualScale;
+
+				FoliageFacade.SetFoliageEntry(FoliageIndex, Data);
+			}
+		}
 	}
 }
diff --git a/Engine/Plugins/Experimental/ProceduralVegetationEditor/Source/ProceduralVegetation/Public/Facades/PVAttributesNames.h b/Engine/Plugins/Experimental/ProceduralVegetationEditor/Source/ProceduralVegetation/Public/Facades/PVAttributesNames.h
--- a/Engine/Plugins/Experimental/ProceduralVegetationEditor/Source/ProceduralVegetation/Public/Facades/PVAttributesNames.h
+++ b/Engine/Plugins/Experimental/ProceduralVegetationEditor/Source/ProceduralVegetation/Public/Facades/PVAttributesNames.h
@@ -29,6 +29,8 @@ namespace PV::AttributeNames
 	inline static const FName TrunkMaterialPath = FName(TEXT("TrunkMaterialPath"));
 	inline static const FName BranchHierarchyNumber = FName(TEXT("branchHierarchyNumber"));
 	inline static const FName BranchSimulationGroupIndex = FName(TEXT("branchSimulationGroupIndex"));
+	inline static const FName PlantNumber = FName(TEXT("plantNumber"));
+	inline static const FName BranchParentNumber = FName(TEXT("branchParentNumber"));
 
 	//Point Attributes
 	inline static const FName PointPosition = FName(TEXT("Position"));
diff --git a/Engine/Plugins/Experimental/ProceduralVegetationEditor/Source/ProceduralVegetation/Public/Facades/PVBranchFacade.h b/Engine/Plugins/Experimental/ProceduralVegetationEditor/Source/ProceduralVegetation/Public/Facades/PVBranchFacade.h
--- a/Engine/Plugins/Experimental/ProceduralVegetationEditor/Source/ProceduralVegetation/Public/Facades/PVBranchFacade.h
+++ b/Engine/Plugins/Experimental/ProceduralVegetationEditor/Source/ProceduralVegetation/Public/Facades/PVBranchFacade.h
@@ -41,8 +41,6 @@ namespace PV::Facades
 
 		int32 GetBranchSourceBudNumber(const int32 Index) const;
 
-		int32 GetTrunkIndex() const;
-
 		TArray<int32> GetParentBranchIndices(int BranchIndex) const;
 
 		int32 GetGenerationNumber(const int32 Index) const;
diff --git a/Engine/Plugins/Experimental/ProceduralVegetationEditor/Source/ProceduralVegetation/Public/Facades/PVPlantFacade.h b/Engine/Plugins/Experimental/ProceduralVegetationEditor/Source/ProceduralVegetation/Public/Facades/PVPlantFacade.h
new file mode 100644
--- /dev/null
+++ b/Engine/Plugins/Experimental/ProceduralVegetationEditor/Source/ProceduralVegetation/Public/Facades/PVPlantFacade.h
@@ -0,0 +1,52 @@
+﻿// Copyright Epic Games, Inc. All Rights Reserved.
+
+#pragma once
+
+#include "CoreMinimal.h"
+#include "IShrinkableFacade.h"
+#include "GeometryCollection/ManagedArrayAccessor.h"
+
+namespace PV::Facades
+{
+	/**
+	 * FPlantFacade is used to access and manipulate the data from Plants Group with in the MegaPlant's FManagedArrayCollection
+	 * A MegaPlant asset can consist of multiple plants. Each entry in the group references that plant's associated branch
+	 * numbers and indices.
+	 */
+	class PROCEDURALVEGETATION_API FPlantFacade
+	{
+	public:
+		FPlantFacade(FManagedArrayCollection& InCollection);
+		FPlantFacade(const FManagedArrayCollection& InCollection);
+
+		bool IsConst() const { return Collection == nullptr; }
+
+		bool IsValid() const;
+
+		int32 GetElementCount() const;
+
+		TArray<int32> GetBranchNumbers(const int32 PlantNumber) const;
+
+		TArray<int32> GetBranchIndices(const int32 PlantNumber) const;
+
+		TArray<int32> GetBranchIndicesSortedByHierarchyNumber(const int32 PlantNumber) const;
+
+		TArray<int32> GetPlantNumbers() const;
+
+		TArray<int32> GetTrunkIndices() const;
+
+		TMap<int32, int32> GetPlantNumbersToTrunkIndicesMap() const;
+
+		int32 GetTrunkIndex(const int32 PlantNumber) const;
+
+		bool IsTrunkIndex(const int32 BranchIndex) const;
+
+	private:
+		FManagedArrayCollection* Collection = nullptr;
+
+		TManagedArrayAccessor<int32> PlantNumbers;
+		TManagedArrayAccessor<int32> BranchNumbers;
+		TManagedArrayAccessor<int32> BranchParentNumbers;
+		TManagedArrayAccessor<int32> BranchHierarchyNumbers;
+	};
+}
