[ProceduralVegetationEditor] Implemented core framework for debug Visualizations
- Ability to visualize different attributes as:
  - Text
  - Direction
  - Point
- Ability to inspect multiple data sets at the same time.
- For now, scalar (float, int) and Vector (FVector3f) are supported, but we can easily add other options.

diff --git a/Engine/Plugins/Experimental/ProceduralVegetationEditor/Source/ProceduralVegetation/Private/Facades/PVFoliageFacade.cpp b/Engine/Plugins/Experimental/ProceduralVegetationEditor/Source/ProceduralVegetation/Private/Facades/PVFoliageFacade.cpp
--- a/Engine/Plugins/Experimental/ProceduralVegetationEditor/Source/ProceduralVegetation/Private/Facades/PVFoliageFacade.cpp
+++ b/Engine/Plugins/Experimental/ProceduralVegetationEditor/Source/ProceduralVegetation/Private/Facades/PVFoliageFacade.cpp
@@ -362,4 +362,9 @@ namespace PV::Facades
 
 		return FinalTransform;
 	}
+
+	const TManagedArray<FVector3f>& FFoliageFacade::GetPivotPositions() const
+	{
+		return PivotPointsAttribute.Get();
+	}
 }
diff --git a/Engine/Plugins/Experimental/ProceduralVegetationEditor/Source/ProceduralVegetation/Private/Nodes/PVPCGBaseSettings.cpp b/Engine/Plugins/Experimental/ProceduralVegetationEditor/Source/ProceduralVegetation/Private/Nodes/PVPCGBaseSettings.cpp
--- a/Engine/Plugins/Experimental/ProceduralVegetationEditor/Source/ProceduralVegetation/Private/Nodes/PVPCGBaseSettings.cpp
+++ b/Engine/Plugins/Experimental/ProceduralVegetationEditor/Source/ProceduralVegetation/Private/Nodes/PVPCGBaseSettings.cpp
@@ -49,6 +49,10 @@ void FPVPCGBaseElement::PostExecuteInternal(FPCGContext* InContext) const
 			if (const UPVPCGData* ProceduralVegetationData = Cast<UPVPCGData>(OutputData.Data))
 			{
 				ProceduralVegetationData->SetRenderType(Settings->GetRenderType());
+
+#if PV_SHOW_DEBUG_VISUALIZATION
+				ProceduralVegetationData->SetDebugSettings(Settings->GetDebugVisualizationSettings());
+#endif
 			}
 		}
 	}
diff --git a/Engine/Plugins/Experimental/ProceduralVegetationEditor/Source/ProceduralVegetation/Private/PVPCGData.cpp b/Engine/Plugins/Experimental/ProceduralVegetationEditor/Source/ProceduralVegetation/Private/PVPCGData.cpp
--- a/Engine/Plugins/Experimental/ProceduralVegetationEditor/Source/ProceduralVegetation/Private/PVPCGData.cpp
+++ b/Engine/Plugins/Experimental/ProceduralVegetationEditor/Source/ProceduralVegetation/Private/PVPCGData.cpp
@@ -33,6 +33,11 @@ void UPVPCGData::AddToCrc(FArchiveCrc32& Ar, bool bFullDataCrc) const
 		Ar << Collection;
 #if WITH_EDITOR
 		Ar << RenderType;
+
+#if PV_SHOW_DEBUG_VISUALIZATION 
+		Ar << DebugSettings;
+#endif
+		
 #endif
 	}
 	else
@@ -70,6 +75,11 @@ UPCGSpatialData* UPVPCGData::CopyInternal(FPCGContext* Context) const
 
 #if WITH_EDITOR
 	NewProceduralVegetationData->RenderType = RenderType;
+
+#if PV_SHOW_DEBUG_VISUALIZATION 
+	NewProceduralVegetationData->DebugSettings = DebugSettings;
+#endif
+	
 #endif
 	
 	return NewProceduralVegetationData;
diff --git a/Engine/Plugins/Experimental/ProceduralVegetationEditor/Source/ProceduralVegetation/Public/Facades/PVFoliageFacade.h b/Engine/Plugins/Experimental/ProceduralVegetationEditor/Source/ProceduralVegetation/Public/Facades/PVFoliageFacade.h
--- a/Engine/Plugins/Experimental/ProceduralVegetationEditor/Source/ProceduralVegetation/Public/Facades/PVFoliageFacade.h
+++ b/Engine/Plugins/Experimental/ProceduralVegetationEditor/Source/ProceduralVegetation/Public/Facades/PVFoliageFacade.h
@@ -84,6 +84,8 @@ namespace PV::Facades
 		
 		FTransform GetFoliageTransform(int32 Id) const;
 
+		const TManagedArray<FVector3f>& GetPivotPositions() const;
+
 	protected:
 		void DefineSchema(const int32 InitialSize = 0);
 
diff --git a/Engine/Plugins/Experimental/ProceduralVegetationEditor/Source/ProceduralVegetation/Public/Helpers/PVUtilities.h b/Engine/Plugins/Experimental/ProceduralVegetationEditor/Source/ProceduralVegetation/Public/Helpers/PVUtilities.h
--- a/Engine/Plugins/Experimental/ProceduralVegetationEditor/Source/ProceduralVegetation/Public/Helpers/PVUtilities.h
+++ b/Engine/Plugins/Experimental/ProceduralVegetationEditor/Source/ProceduralVegetation/Public/Helpers/PVUtilities.h
@@ -7,6 +7,10 @@ class UMeshComponent;
 class FFoliageFacade;
 class AActor;
 
+#ifndef PV_SHOW_DEBUG_VISUALIZATION
+	#define PV_SHOW_DEBUG_VISUALIZATION WITH_EDITOR && UE_BUILD_DEBUG 
+#endif
+
 namespace PV::Utilities
 {
 	DECLARE_DELEGATE_OneParam(FFoliageComponentCreatedCallback, UMeshComponent*);
diff --git a/Engine/Plugins/Experimental/ProceduralVegetationEditor/Source/ProceduralVegetation/Public/Nodes/PVPCGBaseSettings.h b/Engine/Plugins/Experimental/ProceduralVegetationEditor/Source/ProceduralVegetation/Public/Nodes/PVPCGBaseSettings.h
--- a/Engine/Plugins/Experimental/ProceduralVegetationEditor/Source/ProceduralVegetation/Public/Nodes/PVPCGBaseSettings.h
+++ b/Engine/Plugins/Experimental/ProceduralVegetationEditor/Source/ProceduralVegetation/Public/Nodes/PVPCGBaseSettings.h
@@ -4,6 +4,7 @@
 #include "CoreMinimal.h"
 #include "PCGSettings.h"
 #include "PVPCGData.h"
+#include "Helpers/PVUtilities.h"
 #include "PVPCGBaseSettings.generated.h"
 
 UCLASS(BlueprintType, Abstract, HideCategories=(Debug), ClassGroup = (Procedural))
@@ -17,10 +18,18 @@ public:
 	virtual TArray<EPVRenderType> GetRenderType() const { return TArray<EPVRenderType>(); }
 #endif
 
+#if PV_SHOW_DEBUG_VISUALIZATION
+	FPVDebugSettings GetDebugVisualizationSettings() const { return DebugVisualizationSettings; }
+#endif
+
 protected:
 	virtual TArray<FPCGPinProperties> InputPinProperties() const override;
 	virtual TArray<FPCGPinProperties> OutputPinProperties() const override;
 	virtual FPCGElementPtr CreateElement() const override;
+private:
+	
+	UPROPERTY(EditAnywhere, Category = DebugSettings, meta=(EditCondition = "bDebug", EditConditionHides))
+	FPVDebugSettings DebugVisualizationSettings;
 };
 
 class FPVPCGBaseElement : public IPCGElement
@@ -28,5 +37,4 @@ class FPVPCGBaseElement : public IPCGElement
 protected:
 	virtual void PostExecuteInternal(FPCGContext* InContext) const override;
 	virtual bool ExecuteInternal(FPCGContext* Context) const override;
-};
-
+};
\ No newline at end of file
diff --git a/Engine/Plugins/Experimental/ProceduralVegetationEditor/Source/ProceduralVegetation/Public/PVPCGData.h b/Engine/Plugins/Experimental/ProceduralVegetationEditor/Source/ProceduralVegetation/Public/PVPCGData.h
--- a/Engine/Plugins/Experimental/ProceduralVegetationEditor/Source/ProceduralVegetation/Public/PVPCGData.h
+++ b/Engine/Plugins/Experimental/ProceduralVegetationEditor/Source/ProceduralVegetation/Public/PVPCGData.h
@@ -4,6 +4,7 @@
 
 #include "Data/PCGSpatialData.h"
 #include "GeometryCollection/ManagedArrayCollection.h"
+#include "Helpers/PVUtilities.h"
 #include "PVPCGData.generated.h"
 
 struct FPCGContext;
@@ -23,6 +24,72 @@ enum class EPVRenderType : uint8
 
 ENUM_CLASS_FLAGS(EPVRenderType)
 
+UENUM()
+enum class EPVDebugType
+{
+	Point,
+	Branches,
+	Foliage,
+	Custom
+};
+
+UENUM()
+enum class EPVDebugValueVisualizationMode
+{
+	Text,
+	Direction,
+	Point
+};
+
+USTRUCT()
+struct FPVVisualizationSettings
+{
+	GENERATED_BODY()
+
+	UPROPERTY(EditAnywhere, Category = DebugSettings)
+	EPVDebugType DebugType = EPVDebugType::Point;
+
+	UPROPERTY(EditAnywhere, Category = DebugSettings)
+	bool bShowAnchorPoints = true;
+
+	UPROPERTY(EditAnywhere, Category = DebugSettings)
+	FName AttributeToFilter;
+
+	UPROPERTY(EditAnywhere, Category = DebugSettings)
+	EPVDebugValueVisualizationMode VisualizationMode = EPVDebugValueVisualizationMode::Text;
+
+	UPROPERTY(EditAnywhere, Category = DebugSettings)
+	FColor Color = FColor::Black;
+	
+	FName GetDebugTypeString() const
+	{
+		switch (DebugType)
+		{
+		case EPVDebugType::Point:
+			return "Points";
+
+		case EPVDebugType::Branches:
+			return "Primitives";
+
+		case EPVDebugType::Foliage:
+			return "Foliage";
+
+		case EPVDebugType::Custom:
+			return "Custom";
+		}
+		return NAME_None;
+	}
+};
+
+USTRUCT()
+struct FPVDebugSettings
+{
+	GENERATED_BODY()
+
+	UPROPERTY(EditAnywhere, Category = DebugSettings)
+	TArray<FPVVisualizationSettings> VisualizationSettings;
+};
+
 UCLASS(BlueprintType, ClassGroup = (Procedural))
 class PROCEDURALVEGETATION_API UPVPCGData : public UPCGSpatialData
 {
@@ -45,11 +112,17 @@ public:
 	//~End UPCGSpatialData interface
 
 	const FManagedArrayCollection& GetCollection() const { return Collection; }
-
+	FManagedArrayCollection& GetCollection() { return Collection; }
+	
 #if WITH_EDITOR
 	const TArray<EPVRenderType>& GetRenderType() const { return RenderType; }
 	void SetRenderType(TArray<EPVRenderType> InRenderType) const { RenderType = InRenderType; }
 #endif
+
+#if PV_SHOW_DEBUG_VISUALIZATION 
+	void SetDebugSettings(const FPVDebugSettings& InSettings) const { DebugSettings = InSettings; }
+	const FPVDebugSettings& GetDebugSettings() const { return DebugSettings; }
+#endif
 	
 protected:
 
@@ -57,6 +130,11 @@ protected:
 
 #if WITH_EDITOR
 	mutable TArray<EPVRenderType> RenderType;
+
+#if PV_SHOW_DEBUG_VISUALIZATION 
+	mutable FPVDebugSettings DebugSettings;
+#endif
+	
 #endif
 	
 	// ~Begin UPCGData interface
diff --git a/Engine/Plugins/Experimental/ProceduralVegetationEditor/Source/ProceduralVegetationEditor/Private/Visualizations/DebugVisualization/PVDataAttributeValue.cpp b/Engine/Plugins/Experimental/ProceduralVegetationEditor/Source/ProceduralVegetationEditor/Private/Visualizations/DebugVisualization/PVDataAttributeValue.cpp
new file mode 100644
--- /dev/null
+++ b/Engine/Plugins/Experimental/ProceduralVegetationEditor/Source/ProceduralVegetationEditor/Private/Visualizations/DebugVisualization/PVDataAttributeValue.cpp
@@ -0,0 +1,104 @@
+// Copyright Epic Games, Inc. All Rights Reserved.
+
+#include "PVDataAttributeValue.h"
+
+#if PV_SHOW_DEBUG_VISUALIZATION 
+
+#include "GeometryCollection/ManagedArrayCollection.h"
+
+FPVAttributeScalarValue::FPVAttributeScalarValue(const FName InGroupName, const FName InAttributeName, EManagedArrayType InAttributeType,
+                                                               const FManagedArrayCollection& InCollection)
+	: FPVAttributeValueInterface(InGroupName, InAttributeName, InAttributeType, InCollection)
+{
+	if (InCollection.HasAttribute(InAttributeName, InGroupName))
+	{
+		if (InAttributeType == EManagedArrayType::FInt32Type)
+		{
+			const TManagedArray<int>& FloatAttributesArray = InCollection.GetAttribute<int>(InAttributeName, InGroupName);
+
+			for (int i = 0; i < FloatAttributesArray.Num(); i++)
+			{
+				ValueArray.Add(FloatAttributesArray[i]);
+			}
+		}
+		else if (InAttributeType == EManagedArrayType::FFloatType)
+		{
+			const TManagedArray<float>& FloatAttributesArray = InCollection.GetAttribute<float>(InAttributeName, InGroupName);
+
+			for (int i = 0; i < FloatAttributesArray.Num(); i++)
+			{
+				ValueArray.Add(FloatAttributesArray[i]);
+			}	
+		}
+		else
+		{
+			check(false);
+		}
+	}
+}
+
+TArray<FText> FPVAttributeScalarValue::ToText()
+{
+	TArray<FText> TextArray;
+	
+	for (int i = 0; i < ValueArray.Num(); i++)
+	{
+		TextArray.Add(FText::AsNumber(ValueArray[i]));
+	}
+
+	return TextArray;
+}
+
+TArray<FVector> FPVAttributeScalarValue::ToVectors()
+{
+	static TArray<FVector> VecArray;
+	return VecArray;
+}
+
+FPVAttributeVectorValue::FPVAttributeVectorValue(const FName InGroupName, const FName InAttributeName,
+                                                               EManagedArrayType InAttributeType, const FManagedArrayCollection& InCollection)
+: FPVAttributeValueInterface(InGroupName, InAttributeName, InAttributeType, InCollection)
+{
+	if (InCollection.HasAttribute(InAttributeName, InGroupName))
+	{
+		if (InAttributeType == EManagedArrayType::FVectorType)
+		{
+			const TManagedArray<FVector3f>& FloatAttributesArray = InCollection.GetAttribute<FVector3f>(InAttributeName, InGroupName);
+
+			for (int i = 0; i < FloatAttributesArray.Num(); i++)
+			{
+				ValueArray.Add(FVector(FloatAttributesArray[i].X, FloatAttributesArray[i].Y, FloatAttributesArray[i].Z));
+			}
+		}
+		else
+		{
+			check(false);
+		}
+	}
+}
+
+TArray<FText> FPVAttributeVectorValue::ToText()
+{
+	TArray<FText> TextArray;
+	
+	for (int i = 0; i < ValueArray.Num(); i++)
+	{
+		TextArray.Add(ValueArray[i].ToText());
+	}
+
+	return TextArray;
+}
+
+TArray<FVector> FPVAttributeVectorValue::ToVectors()
+{
+	TArray<FVector> Vectors;
+	
+	for (int i = 0; i < ValueArray.Num(); i++)
+	{
+		Vectors.Add(ValueArray[i]);
+	}
+
+	return Vectors;
+}
+
+#endif
\ No newline at end of file
diff --git a/Engine/Plugins/Experimental/ProceduralVegetationEditor/Source/ProceduralVegetationEditor/Private/Visualizations/DebugVisualization/PVDataAttributeValue.h b/Engine/Plugins/Experimental/ProceduralVegetationEditor/Source/ProceduralVegetationEditor/Private/Visualizations/DebugVisualization/PVDataAttributeValue.h
new file mode 100644
--- /dev/null
+++ b/Engine/Plugins/Experimental/ProceduralVegetationEditor/Source/ProceduralVegetationEditor/Private/Visualizations/DebugVisualization/PVDataAttributeValue.h
@@ -0,0 +1,34 @@
+// Copyright Epic Games, Inc. All Rights Reserved.
+
+#pragma once
+
+#include "Helpers/PVUtilities.h"
+
+#if PV_SHOW_DEBUG_VISUALIZATION 
+
+#include "PVDataAttributeValueInterface.h"
+
+class FPVAttributeScalarValue : public FPVAttributeValueInterface
+{
+public:
+	FPVAttributeScalarValue(const FName InGroupName, const FName InAttributeName, EManagedArrayType InAttributeType, const FManagedArrayCollection& InCollection);
+	virtual TArray<FText> ToText() override;
+	virtual TArray<FVector> ToVectors() override;
+
+private:
+	TArray<float> ValueArray;
+};
+
+class FPVAttributeVectorValue : public FPVAttributeValueInterface
+{
+public:
+	FPVAttributeVectorValue(const FName InGroupName, const FName InAttributeName, EManagedArrayType InAttributeType, const FManagedArrayCollection& InCollection);
+	virtual TArray<FText> ToText() override;
+	virtual bool IsScalar() override { return false;}
+	virtual TArray<FVector> ToVectors() override;
+
+private:
+	TArray<FVector> ValueArray;
+};
+
+#endif
\ No newline at end of file
diff --git a/Engine/Plugins/Experimental/ProceduralVegetationEditor/Source/ProceduralVegetationEditor/Private/Visualizations/DebugVisualization/PVDataAttributeValueInterface.cpp b/Engine/Plugins/Experimental/ProceduralVegetationEditor/Source/ProceduralVegetationEditor/Private/Visualizations/DebugVisualization/PVDataAttributeValueInterface.cpp
new file mode 100644
--- /dev/null
+++ b/Engine/Plugins/Experimental/ProceduralVegetationEditor/Source/ProceduralVegetationEditor/Private/Visualizations/DebugVisualization/PVDataAttributeValueInterface.cpp
@@ -0,0 +1,24 @@
+// Copyright Epic Games, Inc. All Rights Reserved.
+
+#include "PVDataAttributeValueInterface.h"
+
+#if PV_SHOW_DEBUG_VISUALIZATION 
+
+#include "GeometryCollection/ManagedArrayCollection.h"
+#include "PVDataAttributeValue.h"
+
+FPVAttributeValuePtr FPVAttributeValueInterface::Create(const FName InGroupName, const FName InAttributeName,
+	EManagedArrayType InAttributeType, const FManagedArrayCollection& InCollection)
+{
+	switch (InAttributeType)
+	{
+	case EManagedArrayType::FInt32Type:
+	case EManagedArrayType::FFloatType:
+			return std::make_shared<FPVAttributeScalarValue>(InGroupName, InAttributeName, InAttributeType,InCollection);
+	case EManagedArrayType::FVectorType:
+			return std::make_shared<FPVAttributeVectorValue>(InGroupName, InAttributeName, InAttributeType,InCollection);
+	default:
+		return nullptr;
+	}
+}
+#endif
\ No newline at end of file
diff --git a/Engine/Plugins/Experimental/ProceduralVegetationEditor/Source/ProceduralVegetationEditor/Private/Visualizations/DebugVisualization/PVDataAttributeValueInterface.h b/Engine/Plugins/Experimental/ProceduralVegetationEditor/Source/ProceduralVegetationEditor/Private/Visualizations/DebugVisualization/PVDataAttributeValueInterface.h
new file mode 100644
--- /dev/null
+++ b/Engine/Plugins/Experimental/ProceduralVegetationEditor/Source/ProceduralVegetationEditor/Private/Visualizations/DebugVisualization/PVDataAttributeValueInterface.h
@@ -0,0 +1,32 @@
+// Copyright Epic Games, Inc. All Rights Reserved.
+
+#pragma once
+
+#include "Helpers/PVUtilities.h"
+
+#if PV_SHOW_DEBUG_VISUALIZATION
+
+#include "GeometryCollection/ManagedArray.h"
+#include "GeometryCollection/ManagedArrayTypes.h"
+
+struct FManagedArrayCollection;
+
+class FPVAttributeScalarValue;
+class FPVAttributeValueInterface;
+
+typedef std::shared_ptr<FPVAttributeValueInterface> FPVAttributeValuePtr;
+
+class FPVAttributeValueInterface
+{
+public:
+	virtual ~FPVAttributeValueInterface() = default;
+	FPVAttributeValueInterface(const FName InGroupName, const FName InAttributeName, EManagedArrayType InAttributeType, const FManagedArrayCollection& InCollection) {}
+
+	virtual TArray<FText> ToText() = 0;
+	virtual TArray<FVector> ToVectors() = 0;
+	virtual bool IsScalar() { return true; }
+
+	static FPVAttributeValuePtr Create(const FName InGroupName, const FName InAttributeName, EManagedArrayType InAttributeType,const FManagedArrayCollection& InCollection);
+};
+
+#endif
\ No newline at end of file
diff --git a/Engine/Plugins/Experimental/ProceduralVegetationEditor/Source/ProceduralVegetationEditor/Private/Visualizations/DebugVisualization/PVDebugVisualizationBase.cpp b/Engine/Plugins/Experimental/ProceduralVegetationEditor/Source/ProceduralVegetationEditor/Private/Visualizations/DebugVisualization/PVDebugVisualizationBase.cpp
new file mode 100644
--- /dev/null
+++ b/Engine/Plugins/Experimental/ProceduralVegetationEditor/Source/ProceduralVegetationEditor/Private/Visualizations/DebugVisualization/PVDebugVisualizationBase.cpp
@@ -0,0 +1,209 @@
+// Copyright Epic Games, Inc. All Rights Reserved.
+
+#include "PVDebugVisualizationBase.h"
+
+#if PV_SHOW_DEBUG_VISUALIZATION 
+
+#include "AdvancedPreviewScene.h"
+#include "PVDataAttributeValueInterface.h"
+#include "PVDebugVisualizer.h"
+#include "PVPCGData.h"
+#include "Components/InstancedStaticMeshComponent.h"
+#include "Components/TextRenderComponent.h"
+#include "Visualizations/PVLineBatchComponent.h"
+
+#define LOCTEXT_NAMESPACE "MegaPlantDebugVisualization"
+
+DEFINE_LOG_CATEGORY(LogMegaPlantDebugVisualization);
+
+const float FPVDebugVisualizationBase::DefaultDirectionStrength = 20;
+const float FPVDebugVisualizationBase::DefaultTextScale = 2.5;
+const float FPVDebugVisualizationBase::DefaultPointScale = 0.005f;
+
+void FPVDebugVisualizationBase::Draw(FVisualizerDrawContext& InContext)
+{
+	DrawPivotPoints(InContext);
+	DrawAttributes(InContext);
+}
+
+void FPVDebugVisualizationBase::DrawPivotPoints(const FVisualizerDrawContext& InContext)
+{
+	auto Collection = InContext.Collection;
+	auto Settings = InContext.VisualizationSettings;
+
+	if (Settings.bShowAnchorPoints)
+	{
+		const TArray<FVector3f>& Positions = GetPivotPositions(Collection);
+	
+		for (int32 i = 0; i < Positions.Num(); i++)
+		{
+			const FVector Pos(Positions[i]);
+		
+			AddDrawAsPoint(Pos);
+		}	
+	}
+}
+
+void FPVDebugVisualizationBase::DrawAttributes(FVisualizerDrawContext& InContext)
+{
+	auto Collection = InContext.Collection;
+	
+	auto VisualizationSettings = InContext.VisualizationSettings;
+	const FName AttributeName = VisualizationSettings.AttributeToFilter;
+	const FName GroupName = VisualizationSettings.GetDebugTypeString();
+
+	if (!GroupName.IsNone() && !AttributeName.IsNone() && Collection.HasGroup(GroupName))
+	{
+		if (Collection.HasAttribute(AttributeName, GroupName))
+		{
+			int NumElements = Collection.NumElements(GroupName);
+			const FManagedArrayCollection::EArrayType AttributeType = Collection.GetAttributeType(AttributeName, GroupName);
+
+			FPVAttributeValuePtr AttributeValue = FPVAttributeValueInterface::Create(GroupName, AttributeName, AttributeType, Collection);
+
+			check(AttributeValue);
+			
+			if (AttributeValue)
+			{
+				if (AttributeValue->IsScalar() || VisualizationSettings.VisualizationMode == EPVDebugValueVisualizationMode::Text)
+				{
+					auto TextValueArray = AttributeValue->ToText();
+		
+					for (int32 i = 0; i < NumElements; i++)
+					{
+						FVector3f OutPos;
+						float OutScale;
+						GetPivot(Collection, i, OutPos, OutScale);
+						
+						check(TextValueArray.IsValidIndex(i));
+
+						DrawAsText(InContext.SceneSetupParams, TextValueArray[i], OutPos, OutScale, VisualizationSettings.Color);
+					}
+				}
+				else if (VisualizationSettings.VisualizationMode == EPVDebugValueVisualizationMode::Direction)
+				{
+					auto VectorsAttribute = AttributeValue->ToVectors();
+
+					for (int32 i = 0; i < NumElements; i++)
+					{
+						FVector3f OutPos;
+						float OutScale;
+						GetPivot(Collection, i, OutPos, OutScale);
+
+						FVector StartPos(OutPos);
+						FVector EndPos = StartPos + VectorsAttribute[i].GetSafeNormal() * DefaultDirectionStrength;
+						check(VectorsAttribute.IsValidIndex(i));
+
+						DrawAsDirection(InContext.SceneSetupParams, StartPos, EndPos, VisualizationSettings.Color);
+					}
+				}
+				else if (VisualizationSettings.VisualizationMode == EPVDebugValueVisualizationMode::Point)
+				{
+					auto Vectors = AttributeValue->ToVectors();
+
+					for (int32 i = 0; i < NumElements; i++)
+					{
+						check(Vectors.IsValidIndex(i));
+
+						AddDrawAsPoint(Vectors[i]);
+					}
+				}	
+			}
+		}
+	}
+
+	DrawAsPoint(InContext.SceneSetupParams, PointsToDraw);
+}
+
+void FPVDebugVisualizationBase::AddDrawAsPoint(const FVector& InPos)
+{
+	auto Transform = FTransform(InPos);
+	Transform.SetScale3D(FVector(DefaultPointScale, DefaultPointScale, DefaultPointScale));
+
+	PointsToDraw.Add(Transform);
+}
+
+void FPVDebugVisualizationBase::DrawAsText(FPCGSceneSetupParams& InOutParams, const FText& InTextToDraw, const FVector3f& InPos, const float& InScale, const FLinearColor& InColor)
+{
+	const FVector3f& Pos = InPos + (FVector3f::RightVector * -1 * InScale);
+
+	TObjectPtr<UTextRenderComponent> Component = NewObject<UTextRenderComponent>(GetTransientPackage(), NAME_None, RF_Transient);
+	Component->SetText(InTextToDraw);
+	Component->SetTextRenderColor(InColor.ToFColor(false));
+	Component->SetWorldSize(DefaultTextScale);
+	Component->SetGenerateOverlapEvents(false);
+						
+	InOutParams.ManagedResources.Add(Component);
+	InOutParams.Scene->AddComponent(Component, FTransform(FVector(Pos.X, Pos.Y, Pos.Z)));
+}
+
+void FPVDebugVisualizationBase::DrawAsDirection(FPCGSceneSetupParams& InOutParams, const FVector& InStartPos, const FVector& InEndPos,
+	const FLinearColor& InColor)
+{
+	UPVLineBatchComponent* LineComponent = GetOrCreateLineComponent(InOutParams);
+
+	check(LineComponent)
+	
+	if (LineComponent)
+	{
+		LineComponent->AddLine(InStartPos, InEndPos, InColor);
+	}
+}
+
+void FPVDebugVisualizationBase::DrawAsPoint(FPCGSceneSetupParams& InOutParams, const TArray<FTransform>& InTransforms, const FLinearColor& InColor)
+{
+	if (InTransforms.Num() == 0)
+		return;
+	
+	UInstancedStaticMeshComponent* Component = GetOrCreatePointComponent(InOutParams);
+
+	check(Component)
+	
+	if (Component)
+	{
+		Component->AddInstances(InTransforms, false);
+	}	
+}
+
+UPVLineBatchComponent* FPVDebugVisualizationBase::GetOrCreateLineComponent(FPCGSceneSetupParams& InOutParams)
+{
+	UPVLineBatchComponent* LineComponent = nullptr;
+	InOutParams.ManagedResources.FindItemByClass<UPVLineBatchComponent>(&LineComponent);
+
+	if (!LineComponent)
+	{
+		LineComponent = NewObject<UPVLineBatchComponent>(GetTransientPackage(), NAME_None, RF_Transient);
+		LineComponent->InitBounds();
+
+		AddToScene(InOutParams, LineComponent);
+	}
+
+	return LineComponent;
+}
+
+UInstancedStaticMeshComponent* FPVDebugVisualizationBase::GetOrCreatePointComponent(FPCGSceneSetupParams& InOutParams)
+{
+	UInstancedStaticMeshComponent* InstancedStaticMeshComponent = nullptr;
+	InOutParams.ManagedResources.FindItemByClass<UInstancedStaticMeshComponent>(&InstancedStaticMeshComponent);
+
+	if (!InstancedStaticMeshComponent)
+	{
+		UStaticMesh* PointSphere = LoadObject<UStaticMesh>(nullptr, TEXT("/Engine/EditorMeshes/EditorSphere"));
+		
+		InstancedStaticMeshComponent = NewObject<UInstancedStaticMeshComponent>(GetTransientPackage(), NAME_None);
+		InstancedStaticMeshComponent->SetStaticMesh(PointSphere);
+
+		AddToScene(InOutParams, InstancedStaticMeshComponent);
+	}
+
+	return InstancedStaticMeshComponent;
+}
+
+void FPVDebugVisualizationBase::AddToScene(FPCGSceneSetupParams& InOutParams, UPrimitiveComponent* InComponent)
+{
+	InOutParams.ManagedResources.Add(InComponent);
+	InOutParams.Scene->AddComponent(InComponent, FTransform::Identity);
+}
+
+#undef LOCTEXT_NAMESPACE
+#endif
\ No newline at end of file
diff --git a/Engine/Plugins/Experimental/ProceduralVegetationEditor/Source/ProceduralVegetationEditor/Private/Visualizations/DebugVisualization/PVDebugVisualizationBase.h b/Engine/Plugins/Experimental/ProceduralVegetationEditor/Source/ProceduralVegetationEditor/Private/Visualizations/DebugVisualization/PVDebugVisualizationBase.h
new file mode 100644
--- /dev/null
+++ b/Engine/Plugins/Experimental/ProceduralVegetationEditor/Source/ProceduralVegetationEditor/Private/Visualizations/DebugVisualization/PVDebugVisualizationBase.h
@@ -0,0 +1,52 @@
+// Copyright Epic Games, Inc. All Rights Reserved.
+
+#pragma once
+
+#include "Helpers/PVUtilities.h"
+
+#if PV_SHOW_DEBUG_VISUALIZATION 
+
+#include "GeometryCollection/ManagedArray.h"
+#include "Templates/SharedPointer.h"
+
+class UPVLineBatchComponent;
+struct FVisualizerDrawContext;
+class UPVSkeletonVisualizerComponent;
+enum class EManagedArrayType : uint8;
+class FPVAttributeValueInterface;
+struct FPCGSceneSetupParams;
+class UPVPCGData;
+
+class FPVDebugVisualizationBase
+{
+public:
+	virtual ~FPVDebugVisualizationBase() = default;
+	FPVDebugVisualizationBase(){};
+	virtual void Draw(FVisualizerDrawContext& InContext);
+protected:
+	void DrawPivotPoints(const FVisualizerDrawContext& InContext);
+	void DrawAttributes(FVisualizerDrawContext& InContext);
+	void AddDrawAsPoint(const FVector& InPos);
+	virtual TArray<FVector3f> GetPivotPositions(const FManagedArrayCollection& InCollection) = 0;
+	virtual void GetPivot(const FManagedArrayCollection& InCollection, const int InIndex, FVector3f& OutPos, float& OutScale) = 0;
+	
+	static void DrawAsText(FPCGSceneSetupParams& InOutParams, const FText& InTextToDraw, const FVector3f& InPos, const float& InScale, const FLinearColor& InColor = FLinearColor::White);
+	static void DrawAsDirection( FPCGSceneSetupParams& InOutParams, const FVector& InStartPos, const FVector& InEndPos, const FLinearColor& InColor = FLinearColor::White);
+	static void DrawAsPoint(FPCGSceneSetupParams& InOutParams,const TArray<FTransform>& InTransforms, const FLinearColor& InColor = FLinearColor::White);
+	
+	static UPVLineBatchComponent* GetOrCreateLineComponent(FPCGSceneSetupParams& InOutParams);
+	static UInstancedStaticMeshComponent* GetOrCreatePointComponent(FPCGSceneSetupParams& InOutParams);
+
+	static void AddToScene(FPCGSceneSetupParams& InOutParams, UPrimitiveComponent* InComponent);
+	
+	TArray<FTransform> PointsToDraw;
+
+	static const float DefaultDirectionStrength;
+	static const float DefaultTextScale;
+	static const float DefaultPointScale;
+};
+typedef TSharedPtr<FPVDebugVisualizationBase> FPVDebugVisualizationPtr;
+
+DECLARE_LOG_CATEGORY_EXTERN(LogMegaPlantDebugVisualization, Log, All);
+
+#endif
\ No newline at end of file
diff --git a/Engine/Plugins/Experimental/ProceduralVegetationEditor/Source/ProceduralVegetationEditor/Private/Visualizations/DebugVisualization/PVDebugVisualizations.cpp b/Engine/Plugins/Experimental/ProceduralVegetationEditor/Source/ProceduralVegetationEditor/Private/Visualizations/DebugVisualization/PVDebugVisualizations.cpp
new file mode 100644
--- /dev/null
+++ b/Engine/Plugins/Experimental/ProceduralVegetationEditor/Source/ProceduralVegetationEditor/Private/Visualizations/DebugVisualization/PVDebugVisualizations.cpp
@@ -0,0 +1,116 @@
+// Copyright Epic Games, Inc. All Rights Reserved.
+
+#include "PVDebugVisualizations.h"
+
+#if PV_SHOW_DEBUG_VISUALIZATION 
+
+#include "Facades/PVBranchFacade.h"
+#include "Facades/PVFoliageFacade.h"
+#include "Facades/PVPointFacade.h"
+
+TArray<FVector3f> FMegaPlantPointDebugVisualization::GetPivotPositions(const FManagedArrayCollection& InCollection)
+{
+	const PV::Facades::FPointFacade PointFacade(InCollection);
+	const TManagedArray<FVector3f>& PointPositions = PointFacade.GetPositions();
+
+	TArray<FVector3f> Positions;
+	for (int32 i = 0; i < PointPositions.Num(); i++)
+	{
+		const FVector3f& Pos = PointPositions[i];
+
+		Positions.Add(Pos);
+	}
+
+	return Positions;
+}
+
+void FMegaPlantPointDebugVisualization::GetPivot(const FManagedArrayCollection& InCollection, const int InIndex, FVector3f& OutPos, float& OutScale)
+{
+	const PV::Facades::FPointFacade PointFacade(InCollection);
+
+	OutPos = PointFacade.GetPosition(InIndex);
+	OutScale = PointFacade.GetPointScale(InIndex);
+}
+
+TArray<FVector3f> FMegaPlantFoliageDebugVisualization::GetPivotPositions(const FManagedArrayCollection& InCollection)
+{
+	const PV::Facades::FFoliageFacade FoliageFacade(InCollection);
+	const TManagedArray<FVector3f>& PointPositions =  FoliageFacade.GetPivotPositions();
+	
+	TArray<FVector3f> Positions;
+	for (int32 i = 0; i < PointPositions.Num(); i++)
+	{
+		const FVector3f& Pos = PointPositions[i];
+
+		Positions.Add(Pos);
+	}
+
+	return Positions;
+
+}
+
+void FMegaPlantFoliageDebugVisualization::GetPivot(const FManagedArrayCollection& InCollection, const int InIndex, FVector3f& OutPos, float& OutScale)
+{
+	const PV::Facades::FFoliageFacade FoliageFacade(InCollection);
+	const PV::Facades::FBranchFacade BranchFacade(InCollection);
+	const PV::Facades::FPointFacade PointFacade(InCollection);
+	
+	PV::Facades::FFoliageEntryData Data =  FoliageFacade.GetFoliageEntry(InIndex);
+	OutPos = Data.PivotPoint;
+
+	const TArray<int32>& BranchPoints = BranchFacade.GetPoints(Data.BranchId);
+
+	if (BranchPoints.Num() > 0)
+	{
+		OutScale = PointFacade.GetPointScale(BranchPoints[0]);
+	}
+	else
+	{
+		OutScale = 1;
+	}
+}
+
+TArray<FVector3f> FMegaPlantBranchDebugVisualization::GetPivotPositions(const FManagedArrayCollection& InCollection)
+{
+	const PV::Facades::FBranchFacade BranchFacade(InCollection);
+	const PV::Facades::FPointFacade PointFacade(InCollection);
+
+	TArray<FVector3f> Positions;
+	for (int32 BranchIndex = 0; BranchIndex < BranchFacade.GetElementCount(); ++BranchIndex)
+	{
+		const TArray<int32>& BranchPoints = BranchFacade.GetPoints(BranchIndex);
+
+		if (BranchPoints.Num() > 0)
+		{
+			FVector3f RootBranchPointPos = PointFacade.GetPosition(BranchPoints[0]);
+			Positions.Add(RootBranchPointPos);
+		}
+	}
+
+	return Positions;
+}
+
+void FMegaPlantBranchDebugVisualization::GetPivot(const FManagedArrayCollection& InCollection, const int InIndex, FVector3f& OutPos, float& OutScale)
+{
+	const PV::Facades::FBranchFacade BranchFacade(InCollection);
+	const PV::Facades::FPointFacade PointFacade(InCollection);
+
+	TArray<FVector3f> Positions;
+
+	if (BranchFacade.GetElementCount() > InIndex)
+	{
+		const TArray<int32>& BranchPoints = BranchFacade.GetPoints(InIndex);
+
+		if (BranchPoints.Num() > 0)
+		{
+			const int RootPointIndex = BranchPoints[0];
+
+			if(PointFacade.GetElementCount() > RootPointIndex)
+			{
+				OutPos = PointFacade.GetPosition(RootPointIndex);
+				OutScale = PointFacade.GetPointScale(RootPointIndex);
+			}
+		}
+	}
+}
+#endif
\ No newline at end of file
diff --git a/Engine/Plugins/Experimental/ProceduralVegetationEditor/Source/ProceduralVegetationEditor/Private/Visualizations/DebugVisualization/PVDebugVisualizations.h b/Engine/Plugins/Experimental/ProceduralVegetationEditor/Source/ProceduralVegetationEditor/Private/Visualizations/DebugVisualization/PVDebugVisualizations.h
new file mode 100644
--- /dev/null
+++ b/Engine/Plugins/Experimental/ProceduralVegetationEditor/Source/ProceduralVegetationEditor/Private/Visualizations/DebugVisualization/PVDebugVisualizations.h
@@ -0,0 +1,41 @@
+// Copyright Epic Games, Inc. All Rights Reserved.
+
+#pragma once
+
+#include "Helpers/PVUtilities.h"
+
+#if PV_SHOW_DEBUG_VISUALIZATION
+
+#include "PVDebugVisualizationBase.h"
+#include "GeometryCollection/ManagedArray.h"
+
+struct FManagedArrayCollection;
+
+class FMegaPlantPointDebugVisualization final : public FPVDebugVisualizationBase
+{
+public:
+	FMegaPlantPointDebugVisualization(){};
+protected:
+	virtual TArray<FVector3f> GetPivotPositions(const FManagedArrayCollection& InCollection) override;
+	virtual void GetPivot(const FManagedArrayCollection& InCollection, const int InIndex, FVector3f& OutPos, float& OutScale) override;
+};
+
+class FMegaPlantFoliageDebugVisualization final : public FPVDebugVisualizationBase
+{
+public:
+	FMegaPlantFoliageDebugVisualization(){};
+protected:
+	virtual TArray<FVector3f> GetPivotPositions(const FManagedArrayCollection& InCollection) override;
+	virtual void GetPivot(const FManagedArrayCollection& InCollection, const int InIndex, FVector3f& OutPos, float& OutScale) override;
+};
+
+class FMegaPlantBranchDebugVisualization final : public FPVDebugVisualizationBase
+{
+public:
+	FMegaPlantBranchDebugVisualization(){};
+protected:
+	virtual TArray<FVector3f> GetPivotPositions(const FManagedArrayCollection& InCollection) override;
+	virtual void GetPivot(const FManagedArrayCollection& InCollection, const int InIndex, FVector3f& OutPos, float& OutScale) override;
+};
+
+#endif
\ No newline at end of file
diff --git a/Engine/Plugins/Experimental/ProceduralVegetationEditor/Source/ProceduralVegetationEditor/Private/Visualizations/DebugVisualization/PVDebugVisualizer.cpp b/Engine/Plugins/Experimental/ProceduralVegetationEditor/Source/ProceduralVegetationEditor/Private/Visualizations/DebugVisualization/PVDebugVisualizer.cpp
new file mode 100644
--- /dev/null
+++ b/Engine/Plugins/Experimental/ProceduralVegetationEditor/Source/ProceduralVegetationEditor/Private/Visualizations/DebugVisualization/PVDebugVisualizer.cpp
@@ -0,0 +1,56 @@
+// Copyright Epic Games, Inc. All Rights Reserved.
+
+#include "PVDebugVisualizer.h"
+
+#if PV_SHOW_DEBUG_VISUALIZATION 
+
+#include "PVDebugVisualizationBase.h"
+#include "PVDebugVisualizations.h"
+#include "PVPCGData.h"
+
+FVisualizerDrawContext::FVisualizerDrawContext(const FManagedArrayCollection& InCollection, const FPVVisualizationSettings& InSettings,
+	FPCGSceneSetupParams& InSceneParams)
+	: Collection(InCollection)
+	, VisualizationSettings(InSettings)
+	, SceneSetupParams(InSceneParams)
+{
+}
+
+void FPVDebugVisualizer::Draw(const UPVPCGData* InData, FPCGSceneSetupParams& InOutParams)
+{
+	check(InData);
+	
+	TArray<FPVVisualizationSettings> VisualizationSettings = InData->GetDebugSettings().VisualizationSettings;
+
+	for (int32 i = 0; i < VisualizationSettings.Num(); i++)
+	{
+		auto Settings = VisualizationSettings[i];
+
+		FVisualizerDrawContext Context(InData->GetCollection(), Settings,InOutParams);
+
+		if (FPVDebugVisualizationPtr VisualizationPtr = CreateVisualizer(Settings.DebugType))
+		{
+			VisualizationPtr->Draw(Context);
+		}	
+	}
+}
+
+FPVDebugVisualizationPtr FPVDebugVisualizer::CreateVisualizer(EPVDebugType InDebugType)
+{
+	switch (InDebugType)
+	{
+		case EPVDebugType::Point:
+			return MakeShared<FMegaPlantPointDebugVisualization>();
+
+		case EPVDebugType::Foliage:
+			return MakeShared<FMegaPlantFoliageDebugVisualization>();
+
+		case EPVDebugType::Branches:
+			return MakeShared<FMegaPlantBranchDebugVisualization>();
+
+		default:
+			return nullptr;
+	}
+}
+
+#endif
\ No newline at end of file
diff --git a/Engine/Plugins/Experimental/ProceduralVegetationEditor/Source/ProceduralVegetationEditor/Private/Visualizations/DebugVisualization/PVDebugVisualizer.h b/Engine/Plugins/Experimental/ProceduralVegetationEditor/Source/ProceduralVegetationEditor/Private/Visualizations/DebugVisualization/PVDebugVisualizer.h
new file mode 100644
--- /dev/null
+++ b/Engine/Plugins/Experimental/ProceduralVegetationEditor/Source/ProceduralVegetationEditor/Private/Visualizations/DebugVisualization/PVDebugVisualizer.h
@@ -0,0 +1,33 @@
+// Copyright Epic Games, Inc. All Rights Reserved.
+#pragma once
+
+#include "Helpers/PVUtilities.h"
+
+#if PV_SHOW_DEBUG_VISUALIZATION 
+
+#include "PVDebugVisualizationBase.h"
+
+struct FManagedArrayCollection;
+struct FPVVisualizationSettings;
+class FPVDebugVisualizationBase;
+struct FPCGSceneSetupParams;
+class UPVPCGData;
+enum class EPVDebugType;
+
+struct FVisualizerDrawContext
+{
+	FVisualizerDrawContext(const FManagedArrayCollection& InCollection, const FPVVisualizationSettings& InSettings, FPCGSceneSetupParams& InSceneParams);
+	const FManagedArrayCollection& Collection;
+	const FPVVisualizationSettings& VisualizationSettings;
+	FPCGSceneSetupParams& SceneSetupParams;
+};
+
+class FPVDebugVisualizer
+{
+public:
+	static void Draw(const UPVPCGData* InData,FPCGSceneSetupParams& InOutParams);
+private:
+	static FPVDebugVisualizationPtr CreateVisualizer(EPVDebugType InDebugType);
+};
+
+#endif
\ No newline at end of file
diff --git a/Engine/Plugins/Experimental/ProceduralVegetationEditor/Source/ProceduralVegetationEditor/Private/Visualizations/PVLineBatchComponent.cpp b/Engine/Plugins/Experimental/ProceduralVegetationEditor/Source/ProceduralVegetationEditor/Private/Visualizations/PVLineBatchComponent.cpp
--- a/Engine/Plugins/Experimental/ProceduralVegetationEditor/Source/ProceduralVegetationEditor/Private/Visualizations/PVLineBatchComponent.cpp
+++ b/Engine/Plugins/Experimental/ProceduralVegetationEditor/Source/ProceduralVegetationEditor/Private/Visualizations/PVLineBatchComponent.cpp
@@ -4,7 +4,7 @@
 #include "MeshElementCollector.h"
 #include "PrimitiveDrawInterface.h"
 
-FPVLineSceneProxy::FPVLineSceneProxy(const UPVBatchComponent* InComponent)
+FPVLineSceneProxy::FPVLineSceneProxy(const UPVLineBatchComponent* InComponent)
 	: FPrimitiveSceneProxy(InComponent)
 	, LineComponent(InComponent)
 {}
@@ -55,7 +55,7 @@ bool FPVLineSceneProxy::CanBeOccluded() const
 	return false;
 }
 
-UPVBatchComponent::UPVBatchComponent()
+UPVLineBatchComponent::UPVLineBatchComponent()
 {
 	bAutoActivate = true;
 	bTickInEditor = true;
@@ -68,12 +68,12 @@ UPVBatchComponent::UPVBatchComponent()
 	Bounds = FBoxSphereBounds(FVector::ZeroVector, FVector::OneVector, 1);
 }
 
-void UPVBatchComponent::InitBounds()
+void UPVLineBatchComponent::InitBounds()
 {
 	BBox = FBox(ForceInit);
 }
 
-void UPVBatchComponent::AddLine(const FVector& InStartPos, const FVector& InEndPos, const FLinearColor& InColor)
+void UPVLineBatchComponent::AddLine(const FVector& InStartPos, const FVector& InEndPos, const FLinearColor& InColor)
 {
 	Lines.Emplace(InStartPos, InEndPos, InColor);
 
@@ -81,7 +81,7 @@ void UPVBatchComponent::AddLine(const FVector& InStartPos, const FVector& InEndP
 	BBox += InEndPos;
 }
 
-void UPVBatchComponent::Draw(const FSceneView* View, FPrimitiveDrawInterface* PDI) const
+void UPVLineBatchComponent::Draw(const FSceneView* View, FPrimitiveDrawInterface* PDI) const
 {
 	for (const FPVLineInfo& Line : Lines)
 	{
@@ -97,13 +97,13 @@ void UPVBatchComponent::Draw(const FSceneView* View, FPrimitiveDrawInterface* PD
 	}
 }
 
-FPrimitiveSceneProxy* UPVBatchComponent::CreateSceneProxy()
+FPrimitiveSceneProxy* UPVLineBatchComponent::CreateSceneProxy()
 {
 	FPVLineSceneProxy* Proxy = new FPVLineSceneProxy(this);
 	return Proxy;
 }
 
-FBoxSphereBounds UPVBatchComponent::CalcBounds(const FTransform& LocalToWorld) const
+FBoxSphereBounds UPVLineBatchComponent::CalcBounds(const FTransform& LocalToWorld) const
 {
 	return FBoxSphereBounds(BBox).TransformBy(LocalToWorld); 
 }
diff --git a/Engine/Plugins/Experimental/ProceduralVegetationEditor/Source/ProceduralVegetationEditor/Private/Visualizations/PVLineBatchComponent.h b/Engine/Plugins/Experimental/ProceduralVegetationEditor/Source/ProceduralVegetationEditor/Private/Visualizations/PVLineBatchComponent.h
--- a/Engine/Plugins/Experimental/ProceduralVegetationEditor/Source/ProceduralVegetationEditor/Private/Visualizations/PVLineBatchComponent.h
+++ b/Engine/Plugins/Experimental/ProceduralVegetationEditor/Source/ProceduralVegetationEditor/Private/Visualizations/PVLineBatchComponent.h
@@ -10,12 +10,12 @@
 #include "PVLineBatchComponent.generated.h"
 
 struct FManagedArrayCollection;
-class UPVBatchComponent;
+class UPVLineBatchComponent;
 
 class FPVLineSceneProxy final : public FPrimitiveSceneProxy
 {
 public:
-	FPVLineSceneProxy(const UPVBatchComponent* InComponent);
+	FPVLineSceneProxy(const UPVLineBatchComponent* InComponent);
 
 	virtual FPrimitiveViewRelevance GetViewRelevance(const FSceneView* View) const override;
 	virtual void GetDynamicMeshElements(
@@ -31,7 +31,7 @@ public:
 	virtual bool CanBeOccluded() const override;
 
 private:
-	const TObjectPtr<const UPVBatchComponent> LineComponent;
+	const TObjectPtr<const UPVLineBatchComponent> LineComponent;
 };
 
 USTRUCT()
@@ -45,12 +45,12 @@ struct FPVLineInfo
 };
 
 UCLASS(Blueprintable, BlueprintType, meta=(BlueprintSpawnableComponent), MinimalAPI)
-class UPVBatchComponent : public UPrimitiveComponent
+class UPVLineBatchComponent : public UPrimitiveComponent
 {
 	GENERATED_BODY()
 
 public:
-	UPVBatchComponent();
+	UPVLineBatchComponent();
 
 	void InitBounds();
 	void AddLine(const FVector& InStartPos, const FVector& InEndPos, const FLinearColor& InColor = FLinearColor::White);
diff --git a/Engine/Plugins/Experimental/ProceduralVegetationEditor/Source/ProceduralVegetationEditor/Private/Visualizations/PVPCGDataVisualization.cpp b/Engine/Plugins/Experimental/ProceduralVegetationEditor/Source/ProceduralVegetationEditor/Private/Visualizations/PVPCGDataVisualization.cpp
--- a/Engine/Plugins/Experimental/ProceduralVegetationEditor/Source/ProceduralVegetationEditor/Private/Visualizations/PVPCGDataVisualization.cpp
+++ b/Engine/Plugins/Experimental/ProceduralVegetationEditor/Source/ProceduralVegetationEditor/Private/Visualizations/PVPCGDataVisualization.cpp
@@ -1,9 +1,7 @@
 // Copyright Epic Games, Inc. All Rights Reserved.
 
 #include "PVPCGDataVisualization.h"
-
 #include "AdvancedPreviewScene.h"
-#include "EditorViewportClient.h"
 #include "GeometryCollectionToDynamicMesh.h"
 #include "PVBoneComponent.h"
 #include "PCGComponent.h"
@@ -20,6 +18,7 @@
 #include "Helpers/PVUtilities.h"
 #include "PVPCGData.h"
 #include "Components/TextRenderComponent.h"
+#include "DebugVisualization/PVDebugVisualizer.h"
 #include "Visualizations/PVSkeletonVisualizerComponent.h"
 
 #define LOCTEXT_NAMESPACE "PCGManagedArrayCollectionDataVisualization"
@@ -47,17 +46,12 @@ void FPVPCGDataVisualization::RegisterVisualizations()
 void FPVPCGDataVisualization::ExecuteDebugDisplay(FPCGContext* Context, const UPCGSettingsInterface* SettingsInterface, const UPCGData* Data,
                                                          AActor* TargetActor) const
 {
-	const UPVPCGData* ManagedArrayCollectionData = CastChecked<UPVPCGData>(Data);
-	UPCGComponent* SourcePCGComponent = Cast<UPCGComponent>(Context->ExecutionSource.Get());
-	if (!ManagedArrayCollectionData || !SourcePCGComponent)
-	{
-		return;
-	}
+	return;
 }
 
 FPCGSetupSceneFunc FPVPCGDataVisualization::GetViewportSetupFunc(const UPCGSettingsInterface* SettingsInterface, const UPCGData* Data) const
 {
-	return [this, WeakData=TWeakObjectPtr<const UPVPCGData>(Cast<UPVPCGData>(Data))](
+	return [this, WeakData=TWeakObjectPtr<const UPVPCGData>(Cast<UPVPCGData>(Data)), WeakSettings = TWeakObjectPtr<const UPCGSettingsInterface>(Cast<UPCGSettingsInterface>(SettingsInterface))](
 		FPCGSceneSetupParams& InOutParams)
 	{
 		check(InOutParams.Scene);
@@ -109,6 +103,13 @@ FPCGSetupSceneFunc FPVPCGDataVisualization::GetViewportSetupFunc(const UPCGSetti
 			InOutParams.Scene->AddComponent(MannequinComponent, FTransform(FVector(-Bounds.BoxExtent.X * 2.0, 0, 0)));
 		}
 
+#if PV_SHOW_DEBUG_VISUALIZATION
+			if (WeakSettings.IsValid() && WeakSettings->bDebug)
+			{
+				FPVDebugVisualizer::Draw(WeakData.Get(), InOutParams);
+			}
+#endif
+
 		InOutParams.Scene->SetFloorVisibility(true);
 		InOutParams.Scene->SetFloorOffset(1.0f);
 
diff --git a/Engine/Plugins/Experimental/ProceduralVegetationEditor/Source/ProceduralVegetationEditor/Private/Visualizations/PVSkeletonVisualizerComponent.h b/Engine/Plugins/Experimental/ProceduralVegetationEditor/Source/ProceduralVegetationEditor/Private/Visualizations/PVSkeletonVisualizerComponent.h
--- a/Engine/Plugins/Experimental/ProceduralVegetationEditor/Source/ProceduralVegetationEditor/Private/Visualizations/PVSkeletonVisualizerComponent.h
+++ b/Engine/Plugins/Experimental/ProceduralVegetationEditor/Source/ProceduralVegetationEditor/Private/Visualizations/PVSkeletonVisualizerComponent.h
@@ -9,7 +9,7 @@
 struct FManagedArrayCollection;
 
 UCLASS(Blueprintable, BlueprintType, meta=(BlueprintSpawnableComponent), MinimalAPI)
-class UPVSkeletonVisualizerComponent : public UPVBatchComponent
+class UPVSkeletonVisualizerComponent : public UPVLineBatchComponent
 {
 	GENERATED_BODY()
 
