Fix crash if deleting a Procedural Vegetation Asset while the editor is open

Remove AddEditingObject from FPVEditor::Initialize as it calls "NotifyAssetsOpened" which is already called from within FAssetEditorToolkit::InitAssetEditor. This leads to double register of the asset in UAssetEditorSubsystem, which then leads to double delete if the asset is removed while the editor is open

diff --git a/Engine/Plugins/Experimental/ProceduralVegetationEditor/Source/ProceduralVegetationEditor/Private/PVEditor.cpp b/Engine/Plugins/Experimental/ProceduralVegetationEditor/Source/ProceduralVegetationEditor/Private/PVEditor.cpp
--- a/Engine/Plugins/Experimental/ProceduralVegetationEditor/Source/ProceduralVegetationEditor/Private/PVEditor.cpp
+++ b/Engine/Plugins/Experimental/ProceduralVegetationEditor/Source/ProceduralVegetationEditor/Private/PVEditor.cpp
@@ -48,8 +48,6 @@ void FPVEditor::Initialize(const EToolkitMode::Type InMode, const TSharedPtr<cla
 	
 	FPCGEditor::Initialize(InMode, InToolkitHost, ProceduralVegetationGraph);
 
-	AddEditingObject(InProceduralVegetation);
-
 	check(GEditor);
 	
 	UPCGEditorSubsystem* Subsystem = GEditor->GetEditorSubsystem<UPCGEditorSubsystem>();
