Revert change made in cl 45623518 as it caused regression in PVEditor behavior

diff --git a/Engine/Plugins/Experimental/ProceduralVegetationEditor/Source/ProceduralVegetationEditor/Private/PVEditor.cpp b/Engine/Plugins/Experimental/ProceduralVegetationEditor/Source/ProceduralVegetationEditor/Private/PVEditor.cpp
--- a/Engine/Plugins/Experimental/ProceduralVegetationEditor/Source/ProceduralVegetationEditor/Private/PVEditor.cpp
+++ b/Engine/Plugins/Experimental/ProceduralVegetationEditor/Source/ProceduralVegetationEditor/Private/PVEditor.cpp
@@ -48,6 +48,8 @@ void FPVEditor::Initialize(const EToolkitMode::Type InMode, const TSharedPtr<cla
 	
 	FPCGEditor::Initialize(InMode, InToolkitHost, ProceduralVegetationGraph);
 
+	AddEditingObject(InProceduralVegetation);
+
 	check(GEditor);
 	
 	UPCGEditorSubsystem* Subsystem = GEditor->GetEditorSubsystem<UPCGEditorSubsystem>();
