[PVE] Fix issue where aggressive simplification causes branches to disconnect from parent

diff --git a/Engine/Plugins/Experimental/ProceduralVegetationEditor/Source/ProceduralVegetation/Private/Facades/PVBranchFacade.cpp b/Engine/Plugins/Experimental/ProceduralVegetationEditor/Source/ProceduralVegetation/Private/Facades/PVBranchFacade.cpp
--- a/Engine/Plugins/Experimental/ProceduralVegetationEditor/Source/ProceduralVegetation/Private/Facades/PVBranchFacade.cpp
+++ b/Engine/Plugins/Experimental/ProceduralVegetationEditor/Source/ProceduralVegetation/Private/Facades/PVBranchFacade.cpp
@@ -138,6 +138,21 @@ namespace PV::Facades
 		return ParentIndices;
 	}
 
+	int32 FBranchFacade::GetParentBranchIndex(const int32 BranchIndex) const
+	{
+		if (IsValid() && Parents.Num() > 0)
+		{
+			TArray<int32> ParentsBranchNumbers = Parents[BranchIndex];
+			if (ParentsBranchNumbers.Num() > 0)
+			{
+				const int32 ParentBranchNumber = ParentsBranchNumbers.Last();
+				return BranchNumbers.Get().Find(ParentBranchNumber);
+			}
+		}
+
+		return INDEX_NONE;
+	}
+
 	int32 FBranchFacade::GetGenerationNumber(const int32 Index) const
 	{
 		// The generation attributes contain incorrect information at the moment 
diff --git a/Engine/Plugins/Experimental/ProceduralVegetationEditor/Source/ProceduralVegetation/Private/Implementations/PVMesher.cpp b/Engine/Plugins/Experimental/ProceduralVegetationEditor/Source/ProceduralVegetation/Private/Implementations/PVMesher.cpp
--- a/Engine/Plugins/Experimental/ProceduralVegetationEditor/Source/ProceduralVegetation/Private/Implementations/PVMesher.cpp
+++ b/Engine/Plugins/Experimental/ProceduralVegetationEditor/Source/ProceduralVegetation/Private/Implementations/PVMesher.cpp
@@ -30,7 +30,7 @@ void FPVMesher::GenerateGeometryCollection(const FManagedArrayCollection& InSkel
 		UE_LOG(LogProceduralVegetation, Log, TEXT("No branch data available for generating mesh."));
 		return;
 	}
-	
+
 	SetNjordPixelID(OutGeometryCollection);
 
 	MesherSettings.MaterialSettings.ApplyMaterialSettings(OutGeometryCollection);
@@ -215,11 +215,37 @@ TSet<int32> FPVMesher::CollectHardPoints(const PV::Facades::FBranchFacade& Branc
 	TSet<int32> HardPoints;
 	TSet<int32> BranchSourceBudNumbers;
 	BranchSourceBudNumbers.Reserve(BranchFacade.GetElementCount());
+	TSet<int32> BranchSourcePointIndices;
+	BranchSourcePointIndices.Reserve(BranchFacade.GetElementCount());
 
 	for (int32 BranchIndex = 0; BranchIndex < BranchFacade.GetElementCount(); ++BranchIndex)
 	{
 		const int32 SourceBudNumber = BranchFacade.GetBranchSourceBudNumber(BranchIndex);
 		BranchSourceBudNumbers.Add(SourceBudNumber);
+
+		const int32 ParentBranchIndex = BranchFacade.GetParentBranchIndex(BranchIndex);
+
+		if (const TArray<int32>& BranchPoints = BranchFacade.GetPoints(BranchIndex);
+			ParentBranchIndex != INDEX_NONE
+			&& BranchPoints.Num() > 0)
+		{
+			const FVector3f FirstPointPosition = PointFacade.GetPosition(BranchPoints[0]);
+			const TArray<int32>& ParentBranchPoints = BranchFacade.GetPoints(ParentBranchIndex);
+			const int32* Closest = Algo::MinElement(ParentBranchPoints,
+				[FirstPointPosition, PointFacade](const int32& PointIndexA, const int32& PointIndexB)
+					{
+						return FVector3f::Dist(PointFacade.GetPosition(PointIndexA), FirstPointPosition) <
+							FVector3f::Dist(PointFacade.GetPosition(PointIndexB), FirstPointPosition);
+					});
+			const int32 ClosestPointIndex = Closest
+				? *Closest
+				: INDEX_NONE;
+
+			if (ClosestPointIndex != INDEX_NONE)
+			{
+				BranchSourcePointIndices.Add(ClosestPointIndex);
+			}
+		}
 	}
 
 	for (int32 BranchIndex = 0; BranchIndex < BranchFacade.GetElementCount(); ++BranchIndex)
@@ -235,7 +261,9 @@ TSet<int32> FPVMesher::CollectHardPoints(const PV::Facades::FBranchFacade& Branc
 			const float NjordPixelIndex = PointFacade.GetNjordPixelIndex(PointIndex);
 
 			bool bIsSeedPoint = PlantFacade.IsTrunkIndex(BranchIndex) && PointIndex == FirstPointIndex;
-			bool bIsBranchJoint = (NjordPixelIndex == FMath::FloorToFloat(NjordPixelIndex)) && BranchSourceBudNumbers.Contains(PointBudNumber);
+			bool bIsBranchJoint = ((NjordPixelIndex == FMath::FloorToFloat(NjordPixelIndex))
+					&& BranchSourceBudNumbers.Contains(PointBudNumber))
+				|| BranchSourcePointIndices.Contains(PointIndex);
 
 			if (PointIndex == FirstPointIndex ||
 				PointIndex == LastPointIndex ||
@@ -622,9 +650,11 @@ void FPVMesher::GenerateBranchMeshData(const bool bPrimitiveIsTrunk, const int32
 		FVector2f NewRange = PointFacade.GetURange(PointIndex);
 
 		float Interval = NewRange.GetMax() - NewRange.GetMin();
-		float URatio = FMath::IsNearlyZero(Interval) ? 1.0f : 1.0f / Interval;
+		float URatio = FMath::IsNearlyZero(Interval)
+			? 1.0f
+			: 1.0f / Interval;
 		TextureCoordV /= URatio;
-		
+
 		const int32 PointDivisions = TargetMeshDivisions[PointIndex];
 		const FVector3f& PointPosition = PointPositions[PointIndex];
 		const float PointScale = PointScales[PointIndex];
@@ -696,7 +726,9 @@ void FPVMesher::GenerateBranchMeshData(const bool bPrimitiveIsTrunk, const int32
 		}
 
 		const FLocalDynamicMeshData::FVertex FirstVertex = OutLocalMeshData.Vertices[CurrentIndices[0]];
-		CurrentIndices.Add(OutLocalMeshData.Vertices.Emplace(FirstVertex.Position, FirstVertex.Normal, FVector2f(FMath::GetMappedRangeValueUnclamped(OldRange,NewRange, FMath::Wrap(TextureCoordUOffset, OldRange.X, OldRange.Y)), TextureCoordV),
+		CurrentIndices.Add(OutLocalMeshData.Vertices.Emplace(FirstVertex.Position, FirstVertex.Normal,
+			FVector2f(FMath::GetMappedRangeValueUnclamped(OldRange, NewRange, FMath::Wrap(TextureCoordUOffset, OldRange.X, OldRange.Y)),
+				TextureCoordV),
 			NjordPixelId));
 
 		PreviousPosition = PointPosition;
diff --git a/Engine/Plugins/Experimental/ProceduralVegetationEditor/Source/ProceduralVegetation/Public/Facades/PVBranchFacade.h b/Engine/Plugins/Experimental/ProceduralVegetationEditor/Source/ProceduralVegetation/Public/Facades/PVBranchFacade.h
--- a/Engine/Plugins/Experimental/ProceduralVegetationEditor/Source/ProceduralVegetation/Public/Facades/PVBranchFacade.h
+++ b/Engine/Plugins/Experimental/ProceduralVegetationEditor/Source/ProceduralVegetation/Public/Facades/PVBranchFacade.h
@@ -42,6 +42,8 @@ namespace PV::Facades
 		int32 GetBranchSourceBudNumber(const int32 Index) const;
 
 		TArray<int32> GetParentBranchIndices(int BranchIndex) const;
+		
+		int32 GetParentBranchIndex(const int32 BranchIndex) const;
 
 		int32 GetGenerationNumber(const int32 Index) const;
 
