[PVE] Resolve potential bugs due to out of bound access

diff --git a/Engine/Plugins/Experimental/ProceduralVegetationEditor/Source/ProceduralVegetation/Private/Facades/PVBranchFacade.cpp b/Engine/Plugins/Experimental/ProceduralVegetationEditor/Source/ProceduralVegetation/Private/Facades/PVBranchFacade.cpp
--- a/Engine/Plugins/Experimental/ProceduralVegetationEditor/Source/ProceduralVegetation/Private/Facades/PVBranchFacade.cpp
+++ b/Engine/Plugins/Experimental/ProceduralVegetationEditor/Source/ProceduralVegetation/Private/Facades/PVBranchFacade.cpp
@@ -33,12 +33,12 @@ namespace PV::Facades
 		, BranchSimulationGroupIndex(InCollection, AttributeNames::BranchSimulationGroupIndex, GroupNames::BranchGroup)
 		, TrunkMaterialPathAttribute(InCollection, AttributeNames::TrunkMaterialPath, GroupNames::DetailsGroup)
 	{
-		
 	}
 
 	bool FBranchFacade::IsValid() const
 	{
-		return Parents.IsValid() && Children.IsValid() && BranchPoints.IsValid() && BranchNumbers.IsValid() && BranchSourceBudNumber.IsValid() && BranchHierarchyNumber.IsValid();
+		return Parents.IsValid() && Children.IsValid() && BranchPoints.IsValid() && BranchNumbers.IsValid() && BranchSourceBudNumber.IsValid() &&
+			BranchHierarchyNumber.IsValid();
 	}
 
 	int32 FBranchFacade::GetElementCount() const
@@ -140,10 +140,10 @@ namespace PV::Facades
 
 	int32 FBranchFacade::GetParentBranchIndex(const int32 BranchIndex) const
 	{
-		if (IsValid() && Parents.Num() > 0)
+		if (IsValid() && Parents.IsValidIndex(BranchIndex))
 		{
-			TArray<int32> ParentsBranchNumbers = Parents[BranchIndex];
-			if (ParentsBranchNumbers.Num() > 0)
+			if (TArray<int32> ParentsBranchNumbers = Parents[BranchIndex];
+				ParentsBranchNumbers.Num() > 0)
 			{
 				const int32 ParentBranchNumber = ParentsBranchNumbers.Last();
 				return BranchNumbers.Get().Find(ParentBranchNumber);
@@ -223,7 +223,7 @@ namespace PV::Facades
 			BranchSimulationGroupIndex.ModifyAt(Index, InSimulationGroupIndex);
 		}
 	}
-	
+
 	void FBranchFacade::SetTrunkMaterialPath(const FString& InPath)
 	{
 		check(!IsConst());
@@ -263,7 +263,7 @@ namespace PV::Facades
 			BranchSourceBudNumber.ModifyAt(ToIndex, BranchSourceBudNumber[FromIndex]);
 			BranchFoliageIDs.ModifyAt(ToIndex, BranchFoliageIDs[FromIndex]);
 			BranchHierarchyNumber.ModifyAt(ToIndex, BranchHierarchyNumber[FromIndex]);
-			if(BranchUVMaterial.IsValid())
+			if (BranchUVMaterial.IsValid())
 			{
 				BranchUVMaterial.ModifyAt(ToIndex, BranchUVMaterial[FromIndex]);
 			}
@@ -290,17 +290,17 @@ namespace PV::Facades
 		}
 
 		OutSortedIndices.Sort([&](int32 Index1, int32 Index2)
-		{
-			int32 Index1Hierarchy = GetBranchHierarchyNumber(Index1);
-			int32 Index2Hierarchy = GetBranchHierarchyNumber(Index2);
-
-			if (Index1Hierarchy == Index2Hierarchy)
 			{
-				return Index1 < Index2;
-			}
-				
-			return Index1Hierarchy < Index2Hierarchy;
-		});
+				int32 Index1Hierarchy = GetBranchHierarchyNumber(Index1);
+				int32 Index2Hierarchy = GetBranchHierarchyNumber(Index2);
+
+				if (Index1Hierarchy == Index2Hierarchy)
+				{
+					return Index1 < Index2;
+				}
+
+				return Index1Hierarchy < Index2Hierarchy;
+			});
 	}
 
 	int32 FBranchFacade::GetBranchIndexFromPointIndex(const int32 PointIndex) const
diff --git a/Engine/Plugins/Experimental/ProceduralVegetationEditor/Source/ProceduralVegetation/Private/Implementations/PVGravity.cpp b/Engine/Plugins/Experimental/ProceduralVegetationEditor/Source/ProceduralVegetation/Private/Implementations/PVGravity.cpp
--- a/Engine/Plugins/Experimental/ProceduralVegetationEditor/Source/ProceduralVegetation/Private/Implementations/PVGravity.cpp
+++ b/Engine/Plugins/Experimental/ProceduralVegetationEditor/Source/ProceduralVegetation/Private/Implementations/PVGravity.cpp
@@ -202,9 +202,15 @@ void FPVGravity::ApplyGravity(
 
 		for (const int32 ChildIndex : ChildrenIndices)
 		{
-			const int32 FirstPointIndex = BranchFacade.GetPoints(ChildIndex)[0];
-			const float ChildLengthFromRoot = PointFacade.GetLengthFromRoot(FirstPointIndex);
-			if (ChildLengthFromRoot <= LengthFromRoot && ChildLengthFromRoot > PreviousLengthFromRoot)
+			const TArray<int32>& BranchPoints = BranchFacade.GetPoints(ChildIndex);
+			if (BranchPoints.Num() == 0)
+			{
+				continue;
+			}
+
+			const int32 FirstPointIndex = BranchPoints[0];
+			if (const float ChildLengthFromRoot = PointFacade.GetLengthFromRoot(FirstPointIndex);
+				ChildLengthFromRoot <= LengthFromRoot && ChildLengthFromRoot > PreviousLengthFromRoot)
 			[[unlikely]]
 			{
 				ApplyGravity(
diff --git a/Engine/Plugins/Experimental/ProceduralVegetationEditor/Source/ProceduralVegetation/Private/Implementations/PVMesher.cpp b/Engine/Plugins/Experimental/ProceduralVegetationEditor/Source/ProceduralVegetation/Private/Implementations/PVMesher.cpp
--- a/Engine/Plugins/Experimental/ProceduralVegetationEditor/Source/ProceduralVegetation/Private/Implementations/PVMesher.cpp
+++ b/Engine/Plugins/Experimental/ProceduralVegetationEditor/Source/ProceduralVegetation/Private/Implementations/PVMesher.cpp
@@ -251,6 +251,11 @@ TSet<int32> FPVMesher::CollectHardPoints(const PV::Facades::FBranchFacade& Branc
 	for (int32 BranchIndex = 0; BranchIndex < BranchFacade.GetElementCount(); ++BranchIndex)
 	{
 		const TArray<int32>& BranchPoints = BranchFacade.GetPoints(BranchIndex);
+		if (BranchPoints.Num() == 0)
+		{
+			continue;
+		}
+
 		const int32 LastPointIndex = BranchPoints.Last();
 		const int32 FirstPointIndex = BranchPoints[0];
 
diff --git a/Engine/Plugins/Experimental/ProceduralVegetationEditor/Source/ProceduralVegetation/Private/Implementations/PVRemoveBranches.cpp b/Engine/Plugins/Experimental/ProceduralVegetationEditor/Source/ProceduralVegetation/Private/Implementations/PVRemoveBranches.cpp
--- a/Engine/Plugins/Experimental/ProceduralVegetationEditor/Source/ProceduralVegetation/Private/Implementations/PVRemoveBranches.cpp
+++ b/Engine/Plugins/Experimental/ProceduralVegetationEditor/Source/ProceduralVegetation/Private/Implementations/PVRemoveBranches.cpp
@@ -28,14 +28,14 @@ namespace PV::Facades
 			check(TrunkIndex != INDEX_NONE);
 
 			if (TrunkIndex == INDEX_NONE)
-				return;
+				continue;
 
 			const TArray<int32>& TrunkPoints = InBranchFacade.GetPoints(TrunkIndex);
 
 			check(TrunkPoints.Num() > 0);
 
 			if (TrunkPoints.Num() == 0)
-				return;
+				continue;
 
 			const int32 RootPointIndex = TrunkPoints[0];
 
diff --git a/Engine/Plugins/Experimental/ProceduralVegetationEditor/Source/ProceduralVegetation/Private/Implementations/PVScale.cpp b/Engine/Plugins/Experimental/ProceduralVegetationEditor/Source/ProceduralVegetation/Private/Implementations/PVScale.cpp
--- a/Engine/Plugins/Experimental/ProceduralVegetationEditor/Source/ProceduralVegetation/Private/Implementations/PVScale.cpp
+++ b/Engine/Plugins/Experimental/ProceduralVegetationEditor/Source/ProceduralVegetation/Private/Implementations/PVScale.cpp
@@ -30,14 +30,14 @@ void FPVScale::ApplyScale(const float InManualScale, FManagedArrayCollection& Ou
 		check(TrunkIndex != INDEX_NONE);
 
 		if (TrunkIndex == INDEX_NONE)
-			return;
+			continue;
 
 		const TArray<int32>& TrunkPoints = BranchFacade.GetPoints(TrunkIndex);
 
 		check(TrunkPoints.Num() > 0);
 
 		if (TrunkPoints.Num() == 0)
-			return;
+			continue;
 
 		const int32 RootPointIndex = TrunkPoints[0];
 		const FVector3f& PivotPos = PointFacade.GetPosition(RootPointIndex);
