[PVE] Fix issue where ramp controls are not from base to top

diff --git a/Engine/Plugins/Experimental/ProceduralVegetationEditor/Source/ProceduralVegetation/Private/Implementations/PVMesher.cpp b/Engine/Plugins/Experimental/ProceduralVegetationEditor/Source/ProceduralVegetation/Private/Implementations/PVMesher.cpp
--- a/Engine/Plugins/Experimental/ProceduralVegetationEditor/Source/ProceduralVegetation/Private/Implementations/PVMesher.cpp
+++ b/Engine/Plugins/Experimental/ProceduralVegetationEditor/Source/ProceduralVegetation/Private/Implementations/PVMesher.cpp
@@ -302,22 +302,22 @@ TSet<int32> FPVMesher::ComputePointGradients(const PV::Facades::FPointFacade& Po
 
 		// Compute scale removal gradient
 		const float PointHullGradient =
-			MesherSettings.PointHullRetentionGradient.GetRichCurveConst()->Eval(HullGradient) * MesherSettings.PointHullRetention;
+			MesherSettings.PointHullRetentionGradient.GetRichCurveConst()->Eval(1.0f - HullGradient) * MesherSettings.PointHullRetention;
 
 		const float PointMainTrunkGradient =
-			MesherSettings.PointMainTrunkRetentionGradient.GetRichCurveConst()->Eval(MainTrunkGradient) * MesherSettings.PointMainTrunkRetention;
+			MesherSettings.PointMainTrunkRetentionGradient.GetRichCurveConst()->Eval(1.0f - MainTrunkGradient) * MesherSettings.PointMainTrunkRetention;
 
 		const float PointGroundGradient =
-			MesherSettings.PointGroundRetentionGradient.GetRichCurveConst()->Eval(GroundGradient) * MesherSettings.PointGroundRetention;
+			MesherSettings.PointGroundRetentionGradient.GetRichCurveConst()->Eval(1.0f - GroundGradient) * MesherSettings.PointGroundRetention;
 
 		const float ScaleRemovalGradient = (ScaleGradient + PointHullGradient + PointMainTrunkGradient + PointGroundGradient) * MaxPointScale /
 			100.0f;
 
 		// Compute mesh divisions gradient
-		const float DivisionsScaleGradient = MesherSettings.DivisionScaleRetentionGradient.GetRichCurveConst()->Eval(ScaleGradient);
-		const float DivisionsHullGradient = MesherSettings.DivisionHullRetentionGradient.GetRichCurveConst()->Eval(HullGradient);
-		const float DivisionsMainTrunkGradient = MesherSettings.DivisionMainTrunkRetentionGradient.GetRichCurveConst()->Eval(MainTrunkGradient);
-		const float DivisionsGroundGradient = MesherSettings.DivisionGroundRetentionGradient.GetRichCurveConst()->Eval(GroundGradient);
+		const float DivisionsScaleGradient = MesherSettings.DivisionScaleRetentionGradient.GetRichCurveConst()->Eval(1.0f - ScaleGradient);
+		const float DivisionsHullGradient = MesherSettings.DivisionHullRetentionGradient.GetRichCurveConst()->Eval(1.0f - HullGradient);
+		const float DivisionsMainTrunkGradient = MesherSettings.DivisionMainTrunkRetentionGradient.GetRichCurveConst()->Eval(1.0f - MainTrunkGradient);
+		const float DivisionsGroundGradient = MesherSettings.DivisionGroundRetentionGradient.GetRichCurveConst()->Eval(1.0f - GroundGradient);
 
 		const float InterpolatedHullGradient = FMath::GetMappedRangeValueClamped(
 			FVector2f(0.0f, 1.0f),
@@ -337,10 +337,10 @@ TSet<int32> FPVMesher::ComputePointGradients(const PV::Facades::FPointFacade& Po
 		OutMeshDivisionsGradients.Add(MeshDivisionsGradient);
 
 		// Compute delta modifier gradient
-		const float SegmentScaleGradient = MesherSettings.SegmentScaleRetentionGradient.GetRichCurveConst()->Eval(ScaleGradient);
-		const float SegmentHullGradient = MesherSettings.SegmentHullRetentionGradient.GetRichCurveConst()->Eval(HullGradient);
-		const float SegmentMainTrunkGradient = MesherSettings.SegmentMainTrunkRetentionGradient.GetRichCurveConst()->Eval(MainTrunkGradient);
-		const float SegmentGroundGradient = MesherSettings.SegmentGroundRetentionGradient.GetRichCurveConst()->Eval(GroundGradient);
+		const float SegmentScaleGradient = MesherSettings.SegmentScaleRetentionGradient.GetRichCurveConst()->Eval(1.0f - ScaleGradient);
+		const float SegmentHullGradient = MesherSettings.SegmentHullRetentionGradient.GetRichCurveConst()->Eval(1.0f - HullGradient);
+		const float SegmentMainTrunkGradient = MesherSettings.SegmentMainTrunkRetentionGradient.GetRichCurveConst()->Eval(1.0f - MainTrunkGradient);
+		const float SegmentGroundGradient = MesherSettings.SegmentGroundRetentionGradient.GetRichCurveConst()->Eval(1.0f - GroundGradient);
 
 		const float DeltaModifierGradient =
 			(SegmentScaleGradient * MesherSettings.SegmentScaleRetention) +
