[PVE] Fixed issue: Output folder path provided in the output node is not auto-created

diff --git a/Engine/Plugins/Experimental/ProceduralVegetationEditor/Source/ProceduralVegetation/Private/Helpers/PVExportHelper.cpp b/Engine/Plugins/Experimental/ProceduralVegetationEditor/Source/ProceduralVegetation/Private/Helpers/PVExportHelper.cpp
--- a/Engine/Plugins/Experimental/ProceduralVegetationEditor/Source/ProceduralVegetation/Private/Helpers/PVExportHelper.cpp
+++ b/Engine/Plugins/Experimental/ProceduralVegetationEditor/Source/ProceduralVegetation/Private/Helpers/PVExportHelper.cpp
@@ -981,6 +981,8 @@ namespace PV::Export
 			);
 
 			MeshPackage->SetDirtyFlag(true);
+
+			IAssetRegistry::Get()->AssetCreated(ExportMesh);
 		}
 		else if (Settings.ExportMeshType == EPVExportMeshType::SkeletalMesh)
 		{
