Add pin types to all Procedural Vegetation nodes and re-name all nodes to adhere to new naming conventions

diff --git a/Engine/Plugins/Experimental/ProceduralVegetationEditor/Config/BaseProceduralVegetationEditor.ini b/Engine/Plugins/Experimental/ProceduralVegetationEditor/Config/BaseProceduralVegetationEditor.ini
--- a/Engine/Plugins/Experimental/ProceduralVegetationEditor/Config/BaseProceduralVegetationEditor.ini
+++ b/Engine/Plugins/Experimental/ProceduralVegetationEditor/Config/BaseProceduralVegetationEditor.ini
@@ -3,37 +3,53 @@
 +PackageRedirects=(OldName="/Script/MegaPlants", NewName="/Script/ProceduralVegetation")
 +PackageRedirects=(OldName="/Script/MegaPlantsEditor", NewName="/Script/ProceduralVegetationEditor")
 
+; Module re-name from MegaPlant->ProceduralVegetation
 +ClassRedirects=(OldName="MegaPlant",NewName="/Script/ProceduralVegetation.ProceduralVegetation")
 +ClassRedirects=(OldName="MegaPlantGraph",NewName="/Script/ProceduralVegetation.ProceduralVegetationGraph")
 +ClassRedirects=(OldName="MegaPlantLink",NewName="/Script/ProceduralVegetation.ProceduralVegetationLink")
 
-+ClassRedirects=(OldName="PCGMegaPlantRemoveBranchesSettings",NewName="/Script/ProceduralVegetation.PVPCGRemoveBranchesSettings")
-+ClassRedirects=(OldName="PCGMegaPlantScaleSettings",NewName="/Script/ProceduralVegetation.PVPCGScaleSettings")
-+ClassRedirects=(OldName="PCGMegaPlantRecipeLoadSettings",NewName="/Script/ProceduralVegetation.PVPCGMegaSeedLoadSettings")
-+PropertyRedirects=(OldName="PCGMegaPlantRecipeLoadSettings.RecipeDataAsset",NewName="PVPCGMegaSeedLoadSettings.MegaSeedDataAsset")
-+ClassRedirects=(OldName="PCGMegaPlantOutputSettings",NewName="/Script/ProceduralVegetation.PVPCGOutputSettings")
-+ClassRedirects=(OldName="PCGMegaPlantMesherSettings",NewName="/Script/ProceduralVegetation.PVPCGMesherSettings")
-+ClassRedirects=(OldName="PCGMegaPlantImporterSettings",NewName="/Script/ProceduralVegetation.PVPCGImporterSettings")
-+ClassRedirects=(OldName="PCGMegaPlantGravitySettings",NewName="/Script/ProceduralVegetation.PVPCGGravitySettings")
-+ClassRedirects=(OldName="PCGMegaPlantFoliageDistributorSettings",NewName="/Script/ProceduralVegetation.PVPCGFoliageDistributorSettings")
-+ClassRedirects=(OldName="PCGMegaPlantFoliageSettings",NewName="/Script/ProceduralVegetation.PVPCGFoliageSettings")
-+ClassRedirects=(OldName="PCGMegaPlantCarveSettings",NewName="/Script/ProceduralVegetation.PVPCGCarveSettings")
-+ClassRedirects=(OldName="PCGMegaPlantBaseSettings",NewName="/Script/ProceduralVegetation.PVPCGBaseSettings")
-+ClassRedirects=(OldName="PCGMegaPlantBoneReductionSettings",NewName="/Script/ProceduralVegetation.PVPCGBoneReductionSettings")
-
 +ClassRedirects=(OldName="MegaPlantRecipeDataAsset",NewName="/Script/ProceduralVegetation.MegaSeedDataAsset")
 
-+StructRedirects=(OldName="MegaPlantMesherSettings",NewName="/Script/ProceduralVegetation.PVMesherSettings")
++ClassRedirects=(OldName="PCGMegaPlantRemoveBranchesSettings",NewName="/Script/ProceduralVegetation.PVRemoveBranchesSettings")
++ClassRedirects=(OldName="PCGMegaPlantScaleSettings",NewName="/Script/ProceduralVegetation.PVScaleSettings")
++ClassRedirects=(OldName="PCGMegaPlantRecipeLoadSettings",NewName="/Script/ProceduralVegetation.PVPresetLoaderSettings")
++ClassRedirects=(OldName="PCGMegaPlantOutputSettings",NewName="/Script/ProceduralVegetation.PVOutputSettings")
++ClassRedirects=(OldName="PCGMegaPlantMesherSettings",NewName="/Script/ProceduralVegetation.PVMeshBuilderSettings")
++ClassRedirects=(OldName="PCGMegaPlantImporterSettings",NewName="/Script/ProceduralVegetation.PVImporterSettings")
++ClassRedirects=(OldName="PCGMegaPlantGravitySettings",NewName="/Script/ProceduralVegetation.PVGravitySettings")
++ClassRedirects=(OldName="PCGMegaPlantFoliageDistributorSettings",NewName="/Script/ProceduralVegetation.PVFoliageDistributorSettings")
++ClassRedirects=(OldName="PCGMegaPlantFoliageSettings",NewName="/Script/ProceduralVegetation.PVFoliagePaletteSettings")
++ClassRedirects=(OldName="PCGMegaPlantCarveSettings",NewName="/Script/ProceduralVegetation.PVCarveSettings")
++ClassRedirects=(OldName="PCGMegaPlantBaseSettings",NewName="/Script/ProceduralVegetation.PVBaseSettings")
++ClassRedirects=(OldName="PCGMegaPlantBoneReductionSettings",NewName="/Script/ProceduralVegetation.PVBoneReductionSettings")
+
++StructRedirects=(OldName="MegaPlantMesherSettings",NewName="/Script/ProceduralVegetation.PVMeshBuilderParams")
 +StructRedirects=(OldName="MegaPlantMaterialSettings",NewName="/Script/ProceduralVegetation.PVMaterialSettings")
-+StructRedirects=(OldName="MegaPlantGravitySettings",NewName="/Script/ProceduralVegetation.PVGravitySettings")
-+StructRedirects=(OldName="MegaPlantCarveSettings",NewName="/Script/ProceduralVegetation.PVCarveSettings")
-+StructRedirects=(OldName="MegaPlantExportSettings",NewName="/Script/ProceduralVegetation.PVExportSettings")
++StructRedirects=(OldName="MegaPlantGravitySettings",NewName="/Script/ProceduralVegetation.PVGravityParams")
++StructRedirects=(OldName="MegaPlantCarveSettings",NewName="/Script/ProceduralVegetation.PVCarveParams")
++StructRedirects=(OldName="MegaPlantExportSettings",NewName="/Script/ProceduralVegetation.PVExportParams")
 
 +EnumRedirects=(OldName="MegaPlantExportMeshType",NewName="/Script/ProceduralVegetation.PVExportMeshType")
 
-+ClassRedirects=(OldName="MegaPlantEditorSchema",NewName="/Script/ProceduralVegetationEditor.PVEditorSchema")
-+ClassRedirects=(OldName="MegaPlantBoneComponent",NewName="/Script/ProceduralVegetationEditor.PVBoneComponent")
-+StructRedirects=(OldName="MegaPlantBoneInfo",NewName="/Script/ProceduralVegetationEditor.PVBoneInfo")
-+ClassRedirects=(OldName="MegaPlantLineBatchComponent",NewName="/Script/ProceduralVegetationEditor.PVLineBatchComponent")
-+StructRedirects=(OldName="MegaPlantLineInfo",NewName="/Script/ProceduralVegetationEditor.PVLineInfo")
-+ClassRedirects=(OldName="MegaPlantSkeletonVisualizerComponent",NewName="/Script/ProceduralVegetationEditor.PVSkeletonVisualizerComponent")
++PropertyRedirects=(OldName="PCGMegaPlantRecipeLoadSettings.RecipeDataAsset",NewName="PVPresetLoaderSettings.PresetDataAsset")
+
+; Class/file name refactor
++ClassRedirects=(OldName="/Script/ProceduralVegetation.PVPCGRemoveBranchesSettings",NewName="/Script/ProceduralVegetation.PVRemoveBranchesSettings")
++ClassRedirects=(OldName="/Script/ProceduralVegetation.PVPCGScaleSettings",NewName="/Script/ProceduralVegetation.PVScaleSettings")
++ClassRedirects=(OldName="/Script/ProceduralVegetation.PVPCGMegaSeedLoadSettings",NewName="/Script/ProceduralVegetation.PVPresetLoaderSettings")
++ClassRedirects=(OldName="/Script/ProceduralVegetation.PVPCGOutputSettings",NewName="/Script/ProceduralVegetation.PVOutputSettings")
++ClassRedirects=(OldName="/Script/ProceduralVegetation.PVPCGMesherSettings",NewName="/Script/ProceduralVegetation.PVMeshBuilderSettings")
++ClassRedirects=(OldName="/Script/ProceduralVegetation.PVPCGImporterSettings",NewName="/Script/ProceduralVegetation.PVImporterSettings")
++ClassRedirects=(OldName="/Script/ProceduralVegetation.PVPCGGravitySettings",NewName="/Script/ProceduralVegetation.PVGravitySettings")
++ClassRedirects=(OldName="/Script/ProceduralVegetation.PVPCGFoliageDistributorSettings",NewName="/Script/ProceduralVegetation.PVFoliageDistributorSettings")
++ClassRedirects=(OldName="/Script/ProceduralVegetation.PVPCGFoliageSettings",NewName="/Script/ProceduralVegetation.PVFoliagePaletteSettings")
++ClassRedirects=(OldName="/Script/ProceduralVegetation.PVPCGCarveSettings",NewName="/Script/ProceduralVegetation.PVCarveSettings")
++ClassRedirects=(OldName="/Script/ProceduralVegetation.PVPCGBaseSettings",NewName="/Script/ProceduralVegetation.PVBaseSettings")
++ClassRedirects=(OldName="/Script/ProceduralVegetation.PVPCGBoneReductionSettings",NewName="/Script/ProceduralVegetation.PVBoneReductionSettings")
+
++StructRedirects=(OldName="/Script/ProceduralVegetation.PVMesherSettings",NewName="/Script/ProceduralVegetation.PVMeshBuilderParams")
++StructRedirects=(OldName="/Script/ProceduralVegetation.PVGravitySettings",NewName="/Script/ProceduralVegetation.PVGravityParams")
++StructRedirects=(OldName="/Script/ProceduralVegetation.PVCarveSettings",NewName="/Script/ProceduralVegetation.PVCarveParams")
++StructRedirects=(OldName="/Script/ProceduralVegetation.PVExportSettings",NewName="/Script/ProceduralVegetation.PVExportParams")
+
++PropertyRedirects=(OldName="PVPCGMegaSeedLoadSettings.RecipeDataAsset",NewName="PVPresetLoaderSettings.PresetDataAsset")
diff --git a/Engine/Plugins/Experimental/ProceduralVegetationEditor/Source/ProceduralVegetation/Private/DataTypes/PVData.cpp b/Engine/Plugins/Experimental/ProceduralVegetationEditor/Source/ProceduralVegetation/Private/DataTypes/PVData.cpp
new file mode 100644
--- /dev/null
+++ b/Engine/Plugins/Experimental/ProceduralVegetationEditor/Source/ProceduralVegetation/Private/DataTypes/PVData.cpp
@@ -0,0 +1,96 @@
+// Copyright Epic Games, Inc. All Rights Reserved.
+
+#include "DataTypes/PVData.h"
+#include "PCGContext.h"
+#include "Data/PCGPointArrayData.h"
+#include "Data/PCGPointData.h"
+#include "Helpers/PCGHelpers.h"
+#include "Serialization/ArchiveCrc32.h"
+#include "UObject/Package.h"
+
+UPVData::UPVData(const FObjectInitializer& ObjectInitializer)
+	: UPCGSpatialData(ObjectInitializer)
+{
+	
+}
+
+void UPVData::Initialize(FManagedArrayCollection&& InCollection, bool bCanTakeOwnership)
+{
+	Collection = InCollection;
+}
+
+void UPVData::AddToCrc(FArchiveCrc32& Ar, bool bFullDataCrc) const
+{
+	Super::AddToCrc(Ar, bFullDataCrc);
+
+	if (bFullDataCrc)
+	{
+		// Implementation note: metadata not supported at this point.
+
+		FString ClassName = StaticClass()->GetPathName();
+
+		Ar << ClassName;
+		Ar << Collection;
+#if WITH_EDITOR
+		Ar << RenderType;
+
+#if PV_SHOW_DEBUG_VISUALIZATION 
+		Ar << DebugSettings;
+#endif
+		
+#endif
+	}
+	else
+	{
+		AddUIDToCrc(Ar);
+	}
+}
+
+FBox UPVData::GetBounds() const
+{
+	return CachedBounds;
+}
+
+bool UPVData::SamplePoint(const FTransform& InTransform, const FBox& InBounds, FPCGPoint& OutPoint, UPCGMetadata* OutMetadata) const
+{
+	return false;
+}
+
+const UPCGBasePointData* UPVData::ToBasePointData(FPCGContext* Context, const FBox& InBounds, TSubclassOf<UPCGBasePointData> PointDataClass) const
+{
+	TRACE_CPUPROFILER_EVENT_SCOPE(UPCGPrimitiveData::CreatePointData);
+
+	UPCGBasePointData* Data = FPCGContext::NewObject_AnyThread<UPCGBasePointData>(Context, GetTransientPackage(), PointDataClass);
+	return Data;
+}
+
+UPCGSpatialData* UPVData::CopyInternal(FPCGContext* Context) const
+{
+	UPVData* NewProceduralVegetationData = FPCGContext::NewObject_AnyThread<UPVData>(Context);
+
+	FManagedArrayCollection NewCollection;
+	Collection.CopyTo(&NewCollection);
+	
+	NewProceduralVegetationData->Initialize(MoveTemp(NewCollection), /*bCanTakeOwnership=*/false);
+
+#if WITH_EDITOR
+	NewProceduralVegetationData->RenderType = RenderType;
+
+#if PV_SHOW_DEBUG_VISUALIZATION 
+	NewProceduralVegetationData->DebugSettings = DebugSettings;
+#endif
+	
+#endif
+	
+	return NewProceduralVegetationData;
+}
+
+const UPCGPointArrayData* UPVData::ToPointArrayData(FPCGContext* Context, const FBox& InBounds) const
+{
+	return CastChecked<UPCGPointArrayData>(ToBasePointData(Context, InBounds, UPCGPointArrayData::StaticClass()));
+}
+const UPCGPointData* UPVData::ToPointData(FPCGContext* Context, const FBox& InBounds) const
+{
+	return CastChecked<UPCGPointData>(ToBasePointData(Context, InBounds, UPCGPointData::StaticClass()));
+}
+
diff --git a/Engine/Plugins/Experimental/ProceduralVegetationEditor/Source/ProceduralVegetation/Private/DataTypes/PVFoliageMeshData.cpp b/Engine/Plugins/Experimental/ProceduralVegetationEditor/Source/ProceduralVegetation/Private/DataTypes/PVFoliageMeshData.cpp
new file mode 100644
--- /dev/null
+++ b/Engine/Plugins/Experimental/ProceduralVegetationEditor/Source/ProceduralVegetation/Private/DataTypes/PVFoliageMeshData.cpp
@@ -0,0 +1,5 @@
+// Copyright Epic Games, Inc. All Rights Reserved.
+
+#include "DataTypes/PVFoliageMeshData.h"
+
+PCG_DEFINE_TYPE_INFO(FPVDataTypeInfoFoliageMesh, UPVFoliageMeshData);
\ No newline at end of file
diff --git a/Engine/Plugins/Experimental/ProceduralVegetationEditor/Source/ProceduralVegetation/Private/DataTypes/PVGrowthData.cpp b/Engine/Plugins/Experimental/ProceduralVegetationEditor/Source/ProceduralVegetation/Private/DataTypes/PVGrowthData.cpp
new file mode 100644
--- /dev/null
+++ b/Engine/Plugins/Experimental/ProceduralVegetationEditor/Source/ProceduralVegetation/Private/DataTypes/PVGrowthData.cpp
@@ -0,0 +1,5 @@
+// Copyright Epic Games, Inc. All Rights Reserved.
+
+#include "DataTypes/PVGrowthData.h"
+
+PCG_DEFINE_TYPE_INFO(FPVDataTypeInfoGrowth, UPVGrowthData);
\ No newline at end of file
diff --git a/Engine/Plugins/Experimental/ProceduralVegetationEditor/Source/ProceduralVegetation/Private/DataTypes/PVMeshData.cpp b/Engine/Plugins/Experimental/ProceduralVegetationEditor/Source/ProceduralVegetation/Private/DataTypes/PVMeshData.cpp
new file mode 100644
--- /dev/null
+++ b/Engine/Plugins/Experimental/ProceduralVegetationEditor/Source/ProceduralVegetation/Private/DataTypes/PVMeshData.cpp
@@ -0,0 +1,5 @@
+// Copyright Epic Games, Inc. All Rights Reserved.
+
+#include "DataTypes/PVMeshData.h"
+
+PCG_DEFINE_TYPE_INFO(FPVDataTypeInfoMesh, UPVMeshData);
\ No newline at end of file
diff --git a/Engine/Plugins/Experimental/ProceduralVegetationEditor/Source/ProceduralVegetation/Private/Helpers/PVExportHelper.cpp b/Engine/Plugins/Experimental/ProceduralVegetationEditor/Source/ProceduralVegetation/Private/Helpers/PVExportHelper.cpp
--- a/Engine/Plugins/Experimental/ProceduralVegetationEditor/Source/ProceduralVegetation/Private/Helpers/PVExportHelper.cpp
+++ b/Engine/Plugins/Experimental/ProceduralVegetationEditor/Source/ProceduralVegetation/Private/Helpers/PVExportHelper.cpp
@@ -955,16 +955,16 @@ namespace PV::Export
 		const TObjectPtr<UProceduralVegetation> InProceduralVegetation,
 		const FManagedArrayCollection& Collection,
 		const FManagedArrayCollection& FoliageCollection,
-		const FPVExportSettings& Settings
+		const FPVExportParams& ExportParams
 	)
 	{
-		// Set up the Package path and name from settings
-		const FString PackageName = Settings.ContentBrowserFolder.Path;
-		const FString MeshName = Settings.MeshName.ToString();
+		// Set up the Package path and name from ExportParams
+		const FString PackageName = ExportParams.ContentBrowserFolder.Path;
+		const FString MeshName = ExportParams.MeshName.ToString();
 
 		const FString FullPath = FPaths::Combine(PackageName, MeshName);
 
-		if (Settings.ExportMeshType == EPVExportMeshType::StaticMesh)
+		if (ExportParams.ExportMeshType == EPVExportMeshType::StaticMesh)
 		{
 			UPackage* MeshPackage = CreatePackage(*FullPath);
 
@@ -976,15 +976,15 @@ namespace PV::Export
 			PV::Export::Internal::ExportCollectionToStaticMesh(
 				ExportMesh,
 				Collection,
-				Settings.bCreateNaniteAssemblies,
-				Settings.NaniteShapePreservation
+				ExportParams.bCreateNaniteAssemblies,
+				ExportParams.NaniteShapePreservation
 			);
 
 			MeshPackage->SetDirtyFlag(true);
 
 			IAssetRegistry::Get()->AssetCreated(ExportMesh);
 		}
-		else if (Settings.ExportMeshType == EPVExportMeshType::SkeletalMesh)
+		else if (ExportParams.ExportMeshType == EPVExportMeshType::SkeletalMesh)
 		{
 			UPackage* MeshPackage = CreatePackage(*FullPath);
 			UPackage* SkeletonPackage = CreatePackage(*(FullPath + "_Skeleton"));
@@ -1001,8 +1001,8 @@ namespace PV::Export
 			PV::Export::Internal::ExportCollectionToSkeletalMesh(
 				ExportMesh,
 				Collection,
-				Settings.NaniteShapePreservation,
-				Settings.bCreateNaniteAssemblies
+				ExportParams.NaniteShapePreservation,
+				ExportParams.bCreateNaniteAssemblies
 			);
 
 			MeshPackage->SetDirtyFlag(true);
diff --git a/Engine/Plugins/Experimental/ProceduralVegetationEditor/Source/ProceduralVegetation/Private/Implementations/PVCarve.h b/Engine/Plugins/Experimental/ProceduralVegetationEditor/Source/ProceduralVegetation/Private/Implementations/PVCarve.h
--- a/Engine/Plugins/Experimental/ProceduralVegetationEditor/Source/ProceduralVegetation/Private/Implementations/PVCarve.h
+++ b/Engine/Plugins/Experimental/ProceduralVegetationEditor/Source/ProceduralVegetation/Private/Implementations/PVCarve.h
@@ -20,7 +20,7 @@ enum class ECarveBasis : uint8
 };
 
 USTRUCT()
-struct FPVCarveSettings
+struct FPVCarveParams
 {
 	GENERATED_BODY()
 
diff --git a/Engine/Plugins/Experimental/ProceduralVegetationEditor/Source/ProceduralVegetation/Private/Implementations/PVGravity.cpp b/Engine/Plugins/Experimental/ProceduralVegetationEditor/Source/ProceduralVegetation/Private/Implementations/PVGravity.cpp
--- a/Engine/Plugins/Experimental/ProceduralVegetationEditor/Source/ProceduralVegetation/Private/Implementations/PVGravity.cpp
+++ b/Engine/Plugins/Experimental/ProceduralVegetationEditor/Source/ProceduralVegetation/Private/Implementations/PVGravity.cpp
@@ -7,7 +7,7 @@
 #include "Facades/PVPlantFacade.h"
 #include "Facades/PVPointFacade.h"
 
-void FPVGravity::GeneratePhototropicData(const FPVGravitySettings& InGravitySettings, const PV::Facades::FBranchFacade& InBranchFacade,
+void FPVGravity::GeneratePhototropicData(const FPVGravityParams& InGravityParams, const PV::Facades::FBranchFacade& InBranchFacade,
                                          const PV::Facades::FPointFacade& InPointFacade, TArray<FVector3f>& OutPhototropicDirections)
 {
 	FVector3f PrevVec = FVector3f::UpVector;
@@ -31,7 +31,7 @@ void FPVGravity::GeneratePhototropicData(const FPVGravitySettings& InGravitySett
 				const FVector3f& LightOptimal = BudDirections[BudDirectionsLightOptimalIndex];
 				const FVector3f& LightSubOptimal = BudDirections[BudDirectionsLightSubOptimal];
 
-				FVector3f PhototropicDirection = FMath::Lerp(LightOptimal, LightSubOptimal, InGravitySettings.PhototropicBias);
+				FVector3f PhototropicDirection = FMath::Lerp(LightOptimal, LightSubOptimal, InGravityParams.PhototropicBias);
 				PhototropicDirection.Normalize();
 
 				if (PhototropicDirection.Length() < 0.001)
@@ -52,7 +52,7 @@ void FPVGravity::GeneratePhototropicData(const FPVGravitySettings& InGravitySett
 	}
 }
 
-void FPVGravity::ApplyGravity(const FPVGravitySettings& InGravitySettings, FManagedArrayCollection& OutCollection)
+void FPVGravity::ApplyGravity(const FPVGravityParams& InGravityParams, FManagedArrayCollection& OutCollection)
 {
 	const PV::Facades::FBranchFacade BranchFacade(OutCollection);
 	const PV::Facades::FPointFacade PointFacade(OutCollection);
@@ -65,14 +65,14 @@ void FPVGravity::ApplyGravity(const FPVGravitySettings& InGravitySettings, FMana
 
 	TArray<FVector3f> OutPhototropicDirection;
 
-	if (InGravitySettings.Mode == EGravityMode::Phototropic)
+	if (InGravityParams.Mode == EGravityMode::Phototropic)
 	{
-		GeneratePhototropicData(InGravitySettings, BranchFacade, PointFacade, OutPhototropicDirection);
+		GeneratePhototropicData(InGravityParams, BranchFacade, PointFacade, OutPhototropicDirection);
 	}
 
 	for (const int32 TrunkIndex : PlantFacade.GetTrunkIndices())
 	{
-		ApplyGravity(TrunkIndex, InGravitySettings, OutPhototropicDirection, OutCollection);
+		ApplyGravity(TrunkIndex, InGravityParams, OutPhototropicDirection, OutCollection);
 	}
 }
 
@@ -82,7 +82,7 @@ void FPVGravity::ApplyGravity(const FPVGravitySettings& InGravitySettings, FMana
 // involved they will end up giving erroneous results.
 void FPVGravity::ApplyGravity(
 	const int32 BranchIndex,
-	const FPVGravitySettings& GravitySettings,
+	const FPVGravityParams& GravitySettings,
 	const TArray<FVector3f>& PhototropicDirections,
 	FManagedArrayCollection& OutCollection,
 	FQuat4f TotalDownForce,
@@ -91,22 +91,22 @@ void FPVGravity::ApplyGravity(
 )
 {
 	static const auto CalculateDownForce = [](const FVector3f InRelativeDirection, const int InPointIndex, const float InPScaleGradient,
-	                                          const FPVGravitySettings& InGravitySettings,
+	                                          const FPVGravityParams& InGravityParams,
 	                                          const TArray<FVector3f>& InPhototropicDirections) -> FQuat4f
 		{
 			const FVector3f PhototropicDirection = InPhototropicDirections.Num() > InPointIndex
 				? InPhototropicDirections[InPointIndex]
 				: FVector3f::UpVector;
-			const FVector3f GravityDirection = InGravitySettings.Mode == EGravityMode::Phototropic
+			const FVector3f GravityDirection = InGravityParams.Mode == EGravityMode::Phototropic
 				? PhototropicDirection
-				: InGravitySettings.Direction;
+				: InGravityParams.Direction;
 
-			float YDot = 1.0f - FMath::Abs(FVector3f::DotProduct(InRelativeDirection, InGravitySettings.Direction));
+			float YDot = 1.0f - FMath::Abs(FVector3f::DotProduct(InRelativeDirection, InGravityParams.Direction));
 			YDot = FMath::GetMappedRangeValueClamped(FVector2f(0, 1),
-				FVector2f(FMath::Max(0.1f, 1 - InGravitySettings.AngleCorrection), 1),
+				FVector2f(FMath::Max(0.1f, 1 - InGravityParams.AngleCorrection), 1),
 				YDot);
 
-			const float GravityImpact = InGravitySettings.Gravity * (1 - InPScaleGradient) * YDot;
+			const float GravityImpact = InGravityParams.Gravity * (1 - InPScaleGradient) * YDot;
 			const FVector3f GravityForce = FMath::Lerp(InRelativeDirection, GravityDirection, GravityImpact).GetUnsafeNormal();
 			return FQuat4f::FindBetween(InRelativeDirection, GravityForce);
 		};
diff --git a/Engine/Plugins/Experimental/ProceduralVegetationEditor/Source/ProceduralVegetation/Private/Implementations/PVGravity.h b/Engine/Plugins/Experimental/ProceduralVegetationEditor/Source/ProceduralVegetation/Private/Implementations/PVGravity.h
--- a/Engine/Plugins/Experimental/ProceduralVegetationEditor/Source/ProceduralVegetation/Private/Implementations/PVGravity.h
+++ b/Engine/Plugins/Experimental/ProceduralVegetationEditor/Source/ProceduralVegetation/Private/Implementations/PVGravity.h
@@ -20,7 +20,7 @@ enum class EGravityMode : uint8
 };
 
 USTRUCT()
-struct FPVGravitySettings
+struct FPVGravityParams
 {
 	GENERATED_BODY()
 	
@@ -42,12 +42,12 @@ struct FPVGravitySettings
 
 struct FPVGravity
 {
-	static void ApplyGravity(const FPVGravitySettings& InGravitySettings, FManagedArrayCollection& OutCollection);
+	static void ApplyGravity(const FPVGravityParams& InGravityParams, FManagedArrayCollection& OutCollection);
 	
 private :
 
-	static void GeneratePhototropicData(const FPVGravitySettings& InGravitySettings, const PV::Facades::FBranchFacade& InBranchFacade, const PV::Facades::FPointFacade& InPointFacade, TArray<FVector3f>& OutPhototropicDirections);
+	static void GeneratePhototropicData(const FPVGravityParams& InGravityParams, const PV::Facades::FBranchFacade& InBranchFacade, const PV::Facades::FPointFacade& InPointFacade, TArray<FVector3f>& OutPhototropicDirections);
 
-	static void ApplyGravity(const int32 BranchIndex, const FPVGravitySettings& GravitySettings, const TArray<FVector3f>& PhototropicDirections, FManagedArrayCollection& OutCollection,
+	static void ApplyGravity(const int32 BranchIndex, const FPVGravityParams& GravitySettings, const TArray<FVector3f>& PhototropicDirections, FManagedArrayCollection& OutCollection,
 		FQuat4f TotalDownForce = FQuat4f::Identity, FVector3f PreviousPosition = FVector3f::ZeroVector, FVector3f PivotPosition = FVector3f::ZeroVector);
 };
diff --git a/Engine/Plugins/Experimental/ProceduralVegetationEditor/Source/ProceduralVegetation/Private/Implementations/PVMeshBuilder.cpp b/Engine/Plugins/Experimental/ProceduralVegetationEditor/Source/ProceduralVegetation/Private/Implementations/PVMeshBuilder.cpp
new file mode 100644
--- /dev/null
+++ b/Engine/Plugins/Experimental/ProceduralVegetationEditor/Source/ProceduralVegetation/Private/Implementations/PVMeshBuilder.cpp
@@ -0,0 +1,964 @@
+// Copyright Epic Games, Inc. All Rights Reserved.
+
+#include "PVMeshBuilder.h"
+
+#include <Facades/PVProfileFacade.h>
+
+#include "ProceduralVegetationModule.h"
+#include "Algo/MaxElement.h"
+#include "Facades/PVBoneFacade.h"
+#include "Facades/PVBranchFacade.h"
+#include "Facades/PVBudVectorsFacade.h"
+#include "Facades/PVFoliageFacade.h"
+#include "Facades/PVPlantFacade.h"
+#include "Facades/PVPointFacade.h"
+
+void FPVMeshBuilder::GenerateGeometryCollection(const FManagedArrayCollection& InSkeletonCollection, const FPVMeshBuilderParams& MeshBuilderParams,
+                                           FGeometryCollection& OutGeometryCollection)
+{
+	InSkeletonCollection.CopyTo(&OutGeometryCollection);
+
+	TArray<FLocalDynamicMeshData> MeshesDataCollection;
+
+	PV::Facades::FBoneFacade::DefineSchema(OutGeometryCollection);
+	PV::Facades::FBoneFacade BoneFacade = PV::Facades::FBoneFacade(OutGeometryCollection);
+
+	PV::Facades::FBranchFacade BranchFacade = PV::Facades::FBranchFacade(OutGeometryCollection);
+
+	if (BranchFacade.GetElementCount() == 0)
+	{
+		UE_LOG(LogProceduralVegetation, Log, TEXT("No branch data available for generating mesh."));
+		return;
+	}
+
+	SetNjordPixelID(OutGeometryCollection);
+
+	MeshBuilderParams.MaterialSettings.ApplyMaterialSettings(OutGeometryCollection);
+	GetLocalDynamicMeshData(OutGeometryCollection, MeshBuilderParams, MeshesDataCollection);
+
+	uint32 VertexCount = 0;
+	uint32 TriangleCount = 0;
+	for (auto [Vertices, Triangles, MaterialID] : MeshesDataCollection)
+	{
+		VertexCount += Vertices.Num();
+		TriangleCount += Triangles.Num();
+	}
+
+	TManagedArray<int32>& PointIDs = BoneFacade.ModifyVertexPointIds();
+
+	OutGeometryCollection.AddElements(1, OutGeometryCollection.TransformGroup);
+	OutGeometryCollection.AddElements(1, OutGeometryCollection.GeometryGroup);
+	OutGeometryCollection.AddElements(TriangleCount, OutGeometryCollection.FacesGroup);
+	OutGeometryCollection.AddElements(VertexCount, OutGeometryCollection.VerticesGroup);
+	OutGeometryCollection.SetNumUVLayers(1);
+
+	TManagedArray<FVector3f>& Vertices = OutGeometryCollection.Vertex;
+	TManagedArray<FVector3f>& Normals = OutGeometryCollection.Normal;
+	TManagedArray<FVector2f>& UVs = *OutGeometryCollection.FindUVLayer(0);
+	TManagedArray<FIntVector>& Indices = OutGeometryCollection.Indices;
+	TManagedArray<int32>& MaterialIndex = OutGeometryCollection.MaterialIndex;
+	TManagedArray<int32>& MaterialIDAttribute = OutGeometryCollection.MaterialID;
+
+	struct FSection
+	{
+		int32 MaterialID;
+		int32 FirstIndex;
+		int32 NumTriangles;
+		int32 MinVertexIndex;
+		int32 MaxVertexIndex;
+	};
+
+	int32 VertexIndex = 0;
+	int32 FaceIndex = 0;
+	int32 PrevMaterialID = -1;
+	TArray<FSection> Sections;
+
+	//Sorting by material id so sections are created according to material ids
+	MeshesDataCollection.Sort([](const FLocalDynamicMeshData& A, const FLocalDynamicMeshData& B)
+		{
+			return A.MaterialID > B.MaterialID;
+		});
+
+	for (auto MeshIndex = MeshesDataCollection.Num() - 1; MeshIndex >= 0; --MeshIndex)
+	{
+		int MaterialID = MeshesDataCollection[MeshIndex].MaterialID;
+		if (PrevMaterialID != MaterialID)
+		{
+			PrevMaterialID = MaterialID;
+			Sections.Add({MaterialID, FaceIndex, 0, VertexIndex, 0});
+		}
+
+		const FLocalDynamicMeshData& MeshesData = MeshesDataCollection[MeshIndex];
+		for (const FIntVector4& Triangle : MeshesData.Triangles)
+		{
+			const FIntVector4 OffsetTriangle = Triangle + FIntVector4(VertexIndex);
+			Indices[FaceIndex] = FIntVector(OffsetTriangle);
+			MaterialIndex[FaceIndex] = OffsetTriangle.W;
+			MaterialIDAttribute[FaceIndex] = MaterialID;
+			++FaceIndex;
+		}
+		for (const FLocalDynamicMeshData::FVertex& Vertex : MeshesData.Vertices)
+		{
+			Vertices[VertexIndex] = static_cast<FVector3f>(Vertex.Position);
+			Normals[VertexIndex] = static_cast<FVector3f>(Vertex.Normal);
+			UVs[VertexIndex] = Vertex.UV;
+			PointIDs[VertexIndex] = Vertex.NjordPixelId;
+			++VertexIndex;
+		}
+
+		int32 CurrentSectionIndex = Sections.Num() - 1;
+		if (Sections.IsValidIndex(CurrentSectionIndex))
+		{
+			Sections[CurrentSectionIndex].MaxVertexIndex = VertexIndex;
+			Sections[CurrentSectionIndex].NumTriangles = FaceIndex - Sections[CurrentSectionIndex].FirstIndex;
+		}
+	}
+
+	OutGeometryCollection.SimulationType[0] = FGeometryCollection::ESimulationTypes::FST_Rigid;
+	OutGeometryCollection.TransformToGeometryIndex[0] = 0;
+	OutGeometryCollection.TransformIndex[0] = 0;
+	OutGeometryCollection.BoneName[0] = "TreeTest";
+	OutGeometryCollection.BoneColor[0] = FLinearColor::White;
+	OutGeometryCollection.FaceStart[0] = 0;
+	OutGeometryCollection.FaceCount[0] = TriangleCount;
+	OutGeometryCollection.VertexStart[0] = 0;
+	OutGeometryCollection.VertexCount[0] = VertexCount;
+
+	OutGeometryCollection.Visible.Fill(true);
+	OutGeometryCollection.BoneMap.Fill(0);
+	OutGeometryCollection.TangentU.Fill(FVector3f::ForwardVector);
+	OutGeometryCollection.TangentV.Fill(FVector3f::RightVector);
+	OutGeometryCollection.Color.Fill(FLinearColor::White);
+
+	auto& MaterialSetups = MeshBuilderParams.MaterialSettings.MaterialSetups;
+	OutGeometryCollection.AddElements(Sections.Num(), FGeometryCollection::MaterialGroup);
+	auto& MaterialArray = OutGeometryCollection.AddAttribute<FString>("MaterialPath", FGeometryCollection::MaterialGroup);
+
+	int32 SectionIndex = 0;
+	for (auto Section : Sections)
+	{
+		TObjectPtr<UMaterialInterface> SectionMaterial;
+		if (MaterialSetups.IsValidIndex(Section.MaterialID))
+		{
+			SectionMaterial = MaterialSetups[Section.MaterialID].Material;
+		}
+		MaterialArray[SectionIndex] = SectionMaterial
+			? SectionMaterial->GetPathName()
+			: "";
+
+		TManagedArray<FGeometryCollectionSection>& GeometrySections = OutGeometryCollection.Sections;
+
+		GeometrySections[SectionIndex].MaterialID = Section.MaterialID;
+		GeometrySections[SectionIndex].FirstIndex = Section.FirstIndex;
+		GeometrySections[SectionIndex].NumTriangles = Section.NumTriangles;
+		GeometrySections[SectionIndex].MinVertexIndex = Section.MinVertexIndex;
+		GeometrySections[SectionIndex].MaxVertexIndex = Section.MaxVertexIndex;
+
+		SectionIndex++;
+	}
+
+	OutGeometryCollection.UpdateBoundingBox();
+}
+
+void FPVMeshBuilder::SetNjordPixelID(FGeometryCollection& OutGeometryCollection)
+{
+	PV::Facades::FPointFacade PointFacade(OutGeometryCollection);
+	TManagedArray<float>& NjordPixelIds = PointFacade.ModifyNjordPixelIDs();
+
+	TMap<int, TArray<int>> BudNumberPointIndicesMap;
+	for (int PointIndex = 0; PointIndex < PointFacade.GetElementCount(); PointIndex++)
+	{
+		int BudNumber = PointFacade.GetBudNumber(PointIndex);
+		BudNumberPointIndicesMap.FindOrAdd(BudNumber).Add(PointIndex);
+	}
+
+	for (const auto& KeyValuePair : BudNumberPointIndicesMap)
+	{
+		int ValueIndex = 0;
+		for (auto Value : KeyValuePair.Value)
+		{
+			//Modify NjordPixelID
+			float NjordPixelID = KeyValuePair.Key + (float)ValueIndex / KeyValuePair.Value.Num();
+			NjordPixelIds[Value] = NjordPixelID;
+			ValueIndex++;
+		}
+	}
+}
+
+void FPVMeshBuilder::GenerateDynamicMesh(FManagedArrayCollection& Collection, const FPVMeshBuilderParams& MeshBuilderParams,
+                                    TObjectPtr<UDynamicMesh>& OutMesh)
+{
+	TArray<FLocalDynamicMeshData> MeshesDataCollection;
+	GetLocalDynamicMeshData(Collection, MeshBuilderParams, MeshesDataCollection);
+
+	uint32 VertexCount = 0;
+	uint32 TriangleCount = 0;
+	for (auto [Vertices, Triangles , MaterialID] : MeshesDataCollection)
+	{
+		VertexCount += Vertices.Num();
+		TriangleCount += Triangles.Num();
+	}
+
+	OutMesh->InitializeMesh();
+	OutMesh->EditMesh([&](UE::Geometry::FDynamicMesh3& InternalMesh)
+			{
+				FPVMeshGenerator Generator(VertexCount, TriangleCount, &MeshesDataCollection);
+				Generator.Generate();
+				InternalMesh.Copy(&Generator);
+			},
+		EDynamicMeshChangeType::MeshChange);
+}
+
+TSet<int32> FPVMeshBuilder::CollectHardPoints(const PV::Facades::FBranchFacade& BranchFacade, const PV::Facades::FPointFacade& PointFacade,
+                                         const PV::Facades::FPlantFacade& PlantFacade)
+{
+	TSet<int32> HardPoints;
+	TSet<int32> BranchSourceBudNumbers;
+	BranchSourceBudNumbers.Reserve(BranchFacade.GetElementCount());
+	TSet<int32> BranchSourcePointIndices;
+	BranchSourcePointIndices.Reserve(BranchFacade.GetElementCount());
+
+	for (int32 BranchIndex = 0; BranchIndex < BranchFacade.GetElementCount(); ++BranchIndex)
+	{
+		const int32 SourceBudNumber = BranchFacade.GetBranchSourceBudNumber(BranchIndex);
+		BranchSourceBudNumbers.Add(SourceBudNumber);
+
+		const int32 ParentBranchIndex = BranchFacade.GetParentBranchIndex(BranchIndex);
+
+		if (const TArray<int32>& BranchPoints = BranchFacade.GetPoints(BranchIndex);
+			ParentBranchIndex != INDEX_NONE
+			&& BranchPoints.Num() > 0)
+		{
+			const FVector3f FirstPointPosition = PointFacade.GetPosition(BranchPoints[0]);
+			const TArray<int32>& ParentBranchPoints = BranchFacade.GetPoints(ParentBranchIndex);
+			const int32* Closest = Algo::MinElement(ParentBranchPoints,
+				[FirstPointPosition, PointFacade](const int32& PointIndexA, const int32& PointIndexB)
+					{
+						return FVector3f::Dist(PointFacade.GetPosition(PointIndexA), FirstPointPosition) <
+							FVector3f::Dist(PointFacade.GetPosition(PointIndexB), FirstPointPosition);
+					});
+			const int32 ClosestPointIndex = Closest
+				? *Closest
+				: INDEX_NONE;
+
+			if (ClosestPointIndex != INDEX_NONE)
+			{
+				BranchSourcePointIndices.Add(ClosestPointIndex);
+			}
+		}
+	}
+
+	for (int32 BranchIndex = 0; BranchIndex < BranchFacade.GetElementCount(); ++BranchIndex)
+	{
+		const TArray<int32>& BranchPoints = BranchFacade.GetPoints(BranchIndex);
+		if (BranchPoints.Num() == 0)
+		{
+			continue;
+		}
+
+		const int32 LastPointIndex = BranchPoints.Last();
+		const int32 FirstPointIndex = BranchPoints[0];
+
+		for (int32 i = 0; i < BranchPoints.Num(); ++i)
+		{
+			const int32 PointIndex = BranchPoints[i];
+			const int32 PointBudNumber = PointFacade.GetBudNumber(PointIndex);
+			const float NjordPixelIndex = PointFacade.GetNjordPixelIndex(PointIndex);
+
+			bool bIsSeedPoint = PlantFacade.IsTrunkIndex(BranchIndex) && PointIndex == FirstPointIndex;
+			bool bIsBranchJoint = ((NjordPixelIndex == FMath::FloorToFloat(NjordPixelIndex))
+					&& BranchSourceBudNumbers.Contains(PointBudNumber))
+				|| BranchSourcePointIndices.Contains(PointIndex);
+
+			if (PointIndex == FirstPointIndex ||
+				PointIndex == LastPointIndex ||
+				bIsSeedPoint ||
+				bIsBranchJoint
+			)
+			{
+				HardPoints.Add(PointIndex);
+			}
+		}
+	}
+
+	return HardPoints;
+}
+
+TSet<int32> FPVMeshBuilder::ComputePointGradients(const PV::Facades::FPointFacade& PointFacade,
+                                             const FPVMeshBuilderParams& MeshBuilderParams, const TSet<int32>& HardPoints,
+                                             const float MaxPointScale, TArray<float>& OutMeshDivisionsGradients,
+                                             TArray<float>& OutDeltaModifiers)
+{
+	TSet<int32> PointsToRemove;
+	OutMeshDivisionsGradients.Reserve(PointFacade.GetElementCount());
+	OutDeltaModifiers.Reserve(PointFacade.GetElementCount());
+
+	for (int32 PointIndex = 0; PointIndex < PointFacade.GetElementCount(); ++PointIndex)
+	{
+		const float HullGradient = PointFacade.GetHullGradient(PointIndex);
+		const float MainTrunkGradient = PointFacade.GetMainTrunkGradient(PointIndex);
+		const float GroundGradient = PointFacade.GetGroundGradient(PointIndex);
+		const float ScaleGradient = PointFacade.GetPointScaleGradient(PointIndex);
+
+		// Compute scale removal gradient
+		const float PointHullGradient =
+			MeshBuilderParams.PointHullRetentionGradient.GetRichCurveConst()->Eval(1.0f - HullGradient) * MeshBuilderParams.PointHullRetention;
+
+		const float PointMainTrunkGradient =
+			MeshBuilderParams.PointMainTrunkRetentionGradient.GetRichCurveConst()->Eval(1.0f - MainTrunkGradient) * MeshBuilderParams.PointMainTrunkRetention;
+
+		const float PointGroundGradient =
+			MeshBuilderParams.PointGroundRetentionGradient.GetRichCurveConst()->Eval(1.0f - GroundGradient) * MeshBuilderParams.PointGroundRetention;
+
+		const float ScaleRemovalGradient = (ScaleGradient + PointHullGradient + PointMainTrunkGradient + PointGroundGradient) * MaxPointScale /
+			100.0f;
+
+		// Compute mesh divisions gradient
+		const float DivisionsScaleGradient = MeshBuilderParams.DivisionScaleRetentionGradient.GetRichCurveConst()->Eval(1.0f - ScaleGradient);
+		const float DivisionsHullGradient = MeshBuilderParams.DivisionHullRetentionGradient.GetRichCurveConst()->Eval(1.0f - HullGradient);
+		const float DivisionsMainTrunkGradient = MeshBuilderParams.DivisionMainTrunkRetentionGradient.GetRichCurveConst()->Eval(1.0f - MainTrunkGradient);
+		const float DivisionsGroundGradient = MeshBuilderParams.DivisionGroundRetentionGradient.GetRichCurveConst()->Eval(1.0f - GroundGradient);
+
+		const float InterpolatedHullGradient = FMath::GetMappedRangeValueClamped(
+			FVector2f(0.0f, 1.0f),
+			FVector2f(DivisionsHullGradient, 1.0f),
+			1.0f - MeshBuilderParams.DivisionHullRetention);
+		const float InterpolatedMainTrunkGradient = FMath::GetMappedRangeValueClamped(
+			FVector2f(0.0f, 1.0f),
+			FVector2f(DivisionsMainTrunkGradient, 1.0f),
+			1.0f - MeshBuilderParams.DivisionMainTrunkRetention);
+		const float InterpolatedGroundGradient = FMath::GetMappedRangeValueClamped(
+			FVector2f(0.0f, 1.0f),
+			FVector2f(DivisionsGroundGradient, 1.0f),
+			1.0f - MeshBuilderParams.DivisionGroundRetention);
+
+		const float MeshDivisionsGradient = DivisionsScaleGradient * InterpolatedHullGradient * InterpolatedMainTrunkGradient *
+			InterpolatedGroundGradient;
+		OutMeshDivisionsGradients.Add(MeshDivisionsGradient);
+
+		// Compute delta modifier gradient
+		const float SegmentScaleGradient = MeshBuilderParams.SegmentScaleRetentionGradient.GetRichCurveConst()->Eval(1.0f - ScaleGradient);
+		const float SegmentHullGradient = MeshBuilderParams.SegmentHullRetentionGradient.GetRichCurveConst()->Eval(1.0f - HullGradient);
+		const float SegmentMainTrunkGradient = MeshBuilderParams.SegmentMainTrunkRetentionGradient.GetRichCurveConst()->Eval(1.0f - MainTrunkGradient);
+		const float SegmentGroundGradient = MeshBuilderParams.SegmentGroundRetentionGradient.GetRichCurveConst()->Eval(1.0f - GroundGradient);
+
+		const float DeltaModifierGradient =
+			(SegmentScaleGradient * MeshBuilderParams.SegmentScaleRetention) +
+			(SegmentHullGradient * MeshBuilderParams.SegmentHullRetention) +
+			(SegmentMainTrunkGradient * MeshBuilderParams.SegmentMainTrunkRetention) +
+			(SegmentGroundGradient * MeshBuilderParams.SegmentGroundRetention);
+
+		// Should the point be removed per scale gradient
+		if (!HardPoints.Contains(PointIndex) && MeshBuilderParams.PointRemoval > 0.0f && ScaleRemovalGradient < MeshBuilderParams.PointRemoval)
+		{
+			PointsToRemove.Add(PointIndex);
+		}
+
+		// Compute delta modifier
+		float DeltaModifier = 1.0f;
+		if (!HardPoints.Contains(PointIndex))
+		{
+			DeltaModifier = MeshBuilderParams.SegmentRetentionImpact * (MaxPointScale / 100.0f) * DeltaModifierGradient;
+		}
+		OutDeltaModifiers.Add(DeltaModifier);
+	}
+
+	return PointsToRemove;
+}
+
+void FPVMeshBuilder::PerformPathSimplification(const PV::Facades::FBranchFacade& BranchFacade,
+                                          const PV::Facades::FPointFacade& PointFacade, const FPVMeshBuilderParams& MeshBuilderParams,
+                                          const TSet<int32>& HardPoints, const TArray<float>& DeltaModifiers, TSet<int32>& InOutPointsToRemove)
+{
+	const int32 Iterations = MeshBuilderParams.Accuracy;
+	// Epsilon is the maximum allowed curvature error
+	const float Epsilon = MeshBuilderParams.SegmentReduction / static_cast<float>(Iterations);
+
+	for (int32 BranchIndex = 0; BranchIndex < BranchFacade.GetElementCount(); ++BranchIndex)
+	{
+		TArray<int32> BranchPoints = BranchFacade.GetPoints(BranchIndex).FilterByPredicate(
+			[&InOutPointsToRemove](const int32 PointIndex)
+				{
+					return !InOutPointsToRemove.Contains(PointIndex);
+				});
+
+		for (int32 j = 0; j < Iterations; ++j)
+		{
+			const float IterationEpsilon = Epsilon * (j + 1);
+			for (int32 i = 1; i < BranchPoints.Num() - 1; ++i)
+			{
+				const int32 Pt0Index = BranchPoints[i - 1],
+				            Pt1Index = BranchPoints[i],
+				            Pt2Index = BranchPoints[i + 1];
+				const FVector Pt0Position = FVector(PointFacade.GetPosition(Pt0Index)),
+				              Pt1Position = FVector(PointFacade.GetPosition(Pt1Index)),
+				              Pt2Position = FVector(PointFacade.GetPosition(Pt2Index));
+				const float DeltaModifier = DeltaModifiers[Pt1Index];
+				const float Delta = FMath::PointDistToLine(Pt1Position, Pt2Position - Pt0Position, Pt0Position) + DeltaModifier;
+
+				if (Delta < IterationEpsilon && !HardPoints.Contains(Pt1Index))
+				{
+					InOutPointsToRemove.Add(Pt1Index);
+					BranchPoints.Remove(Pt1Index);
+					i--;
+				}
+			}
+		}
+
+		// Ensure long chains of removed points do not exceed a maximum threshold
+		BranchPoints = BranchFacade.GetPoints(BranchIndex);
+		FVector3f PreviousPointPosition = PointFacade.GetPosition(BranchPoints[0]);
+		float TravelDistance = 0.0f;
+
+		for (int32 i = 0; i < BranchPoints.Num(); ++i)
+		{
+			const int32 CurrentPointIndex = BranchPoints[i];
+			FVector3f CurrentPointPosition = PointFacade.GetPosition(CurrentPointIndex);
+			const float Distance = FVector3f::Distance(PreviousPointPosition, CurrentPointPosition) / 100.0f;
+
+			if (InOutPointsToRemove.Contains(CurrentPointIndex))
+			{
+				TravelDistance += Distance;
+			}
+			else
+			{
+				TravelDistance = 0.0f;
+			}
+
+			if (TravelDistance > MeshBuilderParams.LongestSegmentLength)
+			{
+				InOutPointsToRemove.Remove(CurrentPointIndex);
+				TravelDistance = 0.0f;
+			}
+
+			PreviousPointPosition = CurrentPointPosition;
+		}
+
+		// Ensure two kept points are never closer than a minimum threshold
+		BranchPoints = BranchFacade.GetPoints(BranchIndex).FilterByPredicate(
+			[&InOutPointsToRemove](const int32 PointIndex)
+				{
+					return !InOutPointsToRemove.Contains(PointIndex);
+				});
+		PreviousPointPosition = PointFacade.GetPosition(BranchPoints[0]);
+
+		for (int32 i = 0; i < BranchPoints.Num(); ++i)
+		{
+			const int32 CurrentPointIndex = BranchPoints[i];
+			FVector3f CurrentPointPosition = PointFacade.GetPosition(CurrentPointIndex);
+			const float Distance = FVector3f::Distance(PreviousPointPosition, CurrentPointPosition) / 100.0f;
+
+			if (Distance < MeshBuilderParams.ShortestSegmentLength && !HardPoints.Contains(CurrentPointIndex))
+			{
+				InOutPointsToRemove.Add(CurrentPointIndex);
+			}
+			else
+			{
+				PreviousPointPosition = CurrentPointPosition;
+			}
+		}
+	}
+}
+
+TArray<int32> FPVMeshBuilder::ComputeMeshDivisions(const PV::Facades::FBranchFacade& BranchFacade,
+                                              const FPVMeshBuilderParams& MeshBuilderParams, const TArray<float>& MeshDivisionsGradients,
+                                              const int32 PointCount)
+{
+	TArray<int32> TargetMeshDivisions;
+	static constexpr int32 MinimumDivisions = 3;
+	TargetMeshDivisions.Init(MinimumDivisions, PointCount);
+
+	for (int32 BranchIndex = 0; BranchIndex < BranchFacade.GetElementCount(); ++BranchIndex)
+	{
+		const TArray<int32>& BranchPoints = BranchFacade.GetPoints(BranchIndex);
+		const float MinDivisions = static_cast<float>(MeshBuilderParams.MinMeshDivisions);
+		const float MaxDivisions = static_cast<float>(MeshBuilderParams.MaxMeshDivisions);
+		const FVector2f NormalRange(0.0f, 1.0f);
+		const FVector2f DivisionsRange(MinDivisions, MaxDivisions);
+
+		for (int32 i = 0; i < BranchPoints.Num(); ++i)
+		{
+			const int32 CurrentPointIndex = BranchPoints[i];
+			const float MeshDivisionsGradient = MeshDivisionsGradients[CurrentPointIndex];
+
+			const int32 TargetMeshDivisionsForCurrentPoint = FMath::RoundToInt32(
+				FMath::GetMappedRangeValueClamped(
+					NormalRange,
+					DivisionsRange,
+					MeshDivisionsGradient));
+
+			TargetMeshDivisions[CurrentPointIndex] = TargetMeshDivisionsForCurrentPoint;
+		}
+	}
+
+	return TargetMeshDivisions;
+}
+
+void FPVMeshBuilder::TriangulateRings(const TArray<int32>& PreviousIndices, const TArray<int32>& CurrentIndices, int32& InOutPolyGroupIndex,
+                                 FLocalDynamicMeshData& OutMeshData)
+{
+	if (!PreviousIndices.IsEmpty())
+	[[likely]]
+	{
+		const float IndexDelta1 = PreviousIndices.Num() < CurrentIndices.Num()
+			? static_cast<float>(PreviousIndices.Num()) / CurrentIndices.Num()
+			: 1.0f;
+		const float IndexDelta2 = CurrentIndices.Num() < PreviousIndices.Num()
+			? static_cast<float>(CurrentIndices.Num()) / PreviousIndices.Num()
+			: 1.0f;
+
+		for (
+			float Index1 = 0.0f, Index2 = 0.0f;
+			FMath::RoundToInt32(Index1 + IndexDelta1) <= PreviousIndices.Num() && FMath::RoundToInt32(Index2 + IndexDelta2) <=
+			CurrentIndices.Num();
+			Index1 += IndexDelta1, Index2 += IndexDelta2
+		)
+		{
+			int32 VertexIndices[4] = {
+				PreviousIndices[FMath::RoundToInt32(Index1) % PreviousIndices.Num()],
+				CurrentIndices[FMath::RoundToInt32(Index2) % CurrentIndices.Num()],
+				PreviousIndices[FMath::RoundToInt32(Index1 + IndexDelta1) % PreviousIndices.Num()],
+				CurrentIndices[FMath::RoundToInt32(Index2 + IndexDelta2) % CurrentIndices.Num()],
+			};
+
+			if (VertexIndices[0] != VertexIndices[2])
+			{
+				OutMeshData.Triangles.Emplace(VertexIndices[0], VertexIndices[1], VertexIndices[2], InOutPolyGroupIndex);
+			}
+			if (VertexIndices[1] != VertexIndices[3])
+			{
+				OutMeshData.Triangles.Emplace(VertexIndices[1], VertexIndices[3], VertexIndices[2], InOutPolyGroupIndex);
+			}
+
+			++InOutPolyGroupIndex;
+		}
+	}
+}
+
+float FPVMeshBuilder::GetProfileMultiplier(const TArray<float>& InProfilePoints, const float ProfileUV_U)
+{
+	if (InProfilePoints.IsEmpty())
+	{
+		return 1.0f;
+	}
+
+	const int32 MaxPointsIndex = InProfilePoints.Num() - 1;
+	const int32 ProfileIndex0 = FMath::Min(FMath::FloorToInt(ProfileUV_U * 100), MaxPointsIndex);
+	const int32 ProfileIndex1 = FMath::Min(FMath::CeilToInt(ProfileUV_U * 100), MaxPointsIndex);
+
+	const float Scale0 = InProfilePoints[ProfileIndex0];
+	const float Scale1 = InProfilePoints[ProfileIndex1];
+
+	const float BlendValue = FMath::GetMappedRangeValueClamped(
+		FVector2f(static_cast<float>(ProfileIndex0), static_cast<float>(ProfileIndex1)),
+		FVector2f(0.0f, 1.0f),
+		ProfileUV_U * 100);
+
+	const float ProfileMultiplier = FMath::Lerp(Scale0, Scale1, BlendValue);
+
+	return ProfileMultiplier;
+}
+
+TMap<int32, TArray<int32>> FPVMeshBuilder::GetPointsIndicesToFoliageIndicesMap(const FManagedArrayCollection& Collection)
+{
+	TMap<int32, TArray<int32>> Map;
+	const PV::Facades::FPointFacade PointFacade(Collection);
+	const PV::Facades::FBranchFacade BranchFacade(Collection);
+	const PV::Facades::FFoliageFacade FoliageFacade(Collection);
+
+	for (int32 BranchIndex = 0; BranchIndex < BranchFacade.GetElementCount(); ++BranchIndex)
+	{
+		const TArray<int32>& FoliageEntryIds = FoliageFacade.GetFoliageEntryIdsForBranch(BranchIndex);
+		const TArray<int32>& BranchPoints = BranchFacade.GetPoints(BranchIndex);
+
+		if (BranchPoints.Num() == 0)
+		{
+			continue;
+		}
+
+		for (int32 FoliageEntryId : FoliageEntryIds)
+		{
+			// Find the closest point for current entry's pivot position
+			const FVector3f PivotPosition = FoliageFacade.GetPivotPoint(FoliageEntryId);
+			const int32* Closest = Algo::MinElement(BranchPoints, [PivotPosition, PointFacade](const int32& PointIndexA, const int32& PointIndexB)
+				{
+					return FVector3f::Dist(PointFacade.GetPosition(PointIndexA), PivotPosition) <
+						FVector3f::Dist(PointFacade.GetPosition(PointIndexB), PivotPosition);
+				});
+			int32 ClosestPointIndex = Closest
+				? *Closest
+				: BranchPoints[0];
+
+			Map.FindOrAdd(ClosestPointIndex).Add(FoliageEntryId);
+		}
+	}
+
+	return Map;
+}
+
+void FPVMeshBuilder::GenerateBranchMeshData(const bool bPrimitiveIsTrunk, const int32 GenerationNumber, const TArray<int32>& PrimitivePoints,
+                                       const FDisplacementData& DisplacementData, const FPVMeshBuilderParams& MeshBuilderParams,
+                                       const TArray<int32>& TargetMeshDivisions, const FManagedArrayCollection& Collection,
+                                       FLocalDynamicMeshData& OutLocalMeshData)
+{
+	const PV::Facades::FPointFacade PointFacade(Collection);
+	const PV::Facades::FBudVectorsFacade BudVectorFacade(Collection);
+	const PV::Facades::FPlantProfileFacade PlantProfileFacade(Collection);
+
+	bool bShouldApplyProfile =
+		PlantProfileFacade.NumProfileEntries() > 0
+		&& MeshBuilderParams.PlantProfileOptions.Num() > 0
+		&& (MeshBuilderParams.bApplyProfileToBranches || bPrimitiveIsTrunk);
+
+	TArray<float> ProfilePoints;
+	if (bShouldApplyProfile)
+	{
+		if (const int32 ProfileId = MeshBuilderParams.PlantProfileOptions.Find(MeshBuilderParams.SelectedPlantProfile) - 1;
+			ProfileId > -1)
+		{
+			ProfilePoints = PlantProfileFacade.GetProfilePoints(ProfileId);
+		}
+		else
+		{
+			bShouldApplyProfile = false;
+		}
+	}
+
+	const TManagedArray<FVector3f>& PointPositions = PointFacade.GetPositions();
+	const TManagedArray<TArray<FVector3f>>& PointBudDirections = BudVectorFacade.GetBudDirections();
+	const TManagedArray<float>& PointScales = PointFacade.GetPointScales();
+
+	const float MaxPointScale = *Algo::MaxElement(PointScales);
+	const float MaxScaleRatio = 1.0f / (MaxPointScale * UE_TWO_PI);
+	int32 PolyGroupIndex = 0;
+	TArray<int32> CurrentIndices;
+	TArray<int32> PreviousIndices;
+
+	FVector3f PreviousPosition = PointPositions[PrimitivePoints[0]];
+	FVector3f PreviousUpVector = FVector3f::ForwardVector;
+
+	bool bShouldApplyDisplacement =
+		DisplacementData.Values.Num() > 0
+		&& DisplacementData.TextureWidth > 0
+		&& DisplacementData.TextureHeight > 0
+		&& GenerationNumber <= MeshBuilderParams.DisplacementGenerationUpperLimit;
+
+	const int32 MaxIndex = PrimitivePoints.Num() - 1;
+	for (int32 Index = 0; Index <= MaxIndex; ++Index)
+	{
+		const int32& PointIndex = PrimitivePoints[Index];
+		const float NjordPixelId = PointIndex;
+
+		float TextureCoordV = PointFacade.GetTextureCoordV(PointIndex);
+		float TextureCoordUOffset = PointFacade.GetTextureCoordUOffset(PointIndex);
+		FVector2f OldRange = FVector2f(0, 1.0f);
+		FVector2f NewRange = PointFacade.GetURange(PointIndex);
+
+		float Interval = NewRange.GetMax() - NewRange.GetMin();
+		float URatio = FMath::IsNearlyZero(Interval)
+			? 1.0f
+			: 1.0f / Interval;
+		TextureCoordV /= URatio;
+
+		const int32 PointDivisions = TargetMeshDivisions[PointIndex];
+		const FVector3f& PointPosition = PointPositions[PointIndex];
+		const float PointScale = PointScales[PointIndex];
+		FVector3f NormalVector = PointBudDirections[PointIndex][0].GetUnsafeNormal();
+		const float PlantGradient = PointFacade.GetPlantGradient(PointIndex);
+		const float ProfileFallOff = MeshBuilderParams.PlantProfileFallOff.GetRichCurveConst()->Eval(PlantGradient);
+
+		const FVector3f PointUpOriginal = PointBudDirections[PointIndex][5].GetUnsafeNormal();
+		FVector3f PointUp = FVector3f::VectorPlaneProject(PointUpOriginal, NormalVector).GetSafeNormal();
+
+		if (Index > 0 && FVector3f::DotProduct(PointUp, PreviousUpVector) < 0.0f)
+		{
+			PointUp = -PointUp;
+		}
+		PreviousUpVector = PointUp;
+
+		if (PointScale == 0.0f)
+		[[unlikely]]
+		{
+			CurrentIndices.Add(OutLocalMeshData.Vertices.Emplace(PointPosition, NormalVector, FVector2f(0.0f, TextureCoordV), NjordPixelId));
+		}
+		else
+		[[likely]]
+		{
+			const float RotationDelta = UE_TWO_PI / PointDivisions;
+			for (int32 Division = 0; Division < PointDivisions; ++Division)
+			{
+				float UV_U = (1 - (static_cast<float>(Division) / PointDivisions)) + TextureCoordUOffset;
+				UV_U = FMath::GetMappedRangeValueUnclamped(OldRange, NewRange, FMath::Wrap(UV_U, OldRange.X, OldRange.Y));
+
+				float ProfileMultiplier = 1.0f;
+				if (bShouldApplyProfile)
+				{
+					ProfileMultiplier = GetProfileMultiplier(ProfilePoints, UV_U);
+					ProfileMultiplier = FMath::Lerp(1.0f, ProfileMultiplier, ProfileFallOff);
+					ProfileMultiplier = ProfileMultiplier * MeshBuilderParams.PlantProfileScale;
+				}
+
+				const FVector3f Direction = FQuat4f(NormalVector, RotationDelta * Division).RotateVector(PointUp);
+				FVector3f RadialPoint = PointPosition + (Direction * PointScale * ProfileMultiplier);
+
+				if (bShouldApplyDisplacement)
+				{
+					const int32 X = FMath::Clamp(
+						FMath::FloorToInt(
+							UV_U * MeshBuilderParams.DisplacementUVScale.X * (DisplacementData.TextureWidth - 1)
+						) % DisplacementData.TextureWidth,
+						0, DisplacementData.TextureWidth - 1
+					);
+					const int32 Y = FMath::Clamp(
+						FMath::FloorToInt(
+							TextureCoordV * MeshBuilderParams.DisplacementUVScale.Y * (DisplacementData.TextureHeight - 1)
+						) % DisplacementData.TextureHeight,
+						0, DisplacementData.TextureHeight - 1
+					);
+
+					const int32 Idx = Y * DisplacementData.TextureWidth + X;
+					float DisplacementMultiplier =
+						MeshBuilderParams.DisplacementStrength * (DisplacementData.Values[Idx] - MeshBuilderParams.DisplacementBias);
+					DisplacementMultiplier = DisplacementMultiplier * PointScale / MaxPointScale;
+
+					RadialPoint = RadialPoint + (Direction * DisplacementMultiplier);
+				}
+
+				const FVector3f Normal = FVector3f((RadialPoint - PointPosition).GetUnsafeNormal());
+
+				CurrentIndices.Add(OutLocalMeshData.Vertices.Emplace(RadialPoint, Normal, FVector2f(UV_U, TextureCoordV), NjordPixelId));
+			}
+		}
+
+		const FLocalDynamicMeshData::FVertex FirstVertex = OutLocalMeshData.Vertices[CurrentIndices[0]];
+		CurrentIndices.Add(OutLocalMeshData.Vertices.Emplace(FirstVertex.Position, FirstVertex.Normal,
+			FVector2f(FMath::GetMappedRangeValueUnclamped(OldRange, NewRange, FMath::Wrap(TextureCoordUOffset, OldRange.X, OldRange.Y)),
+				TextureCoordV),
+			NjordPixelId));
+
+		PreviousPosition = PointPosition;
+
+		TriangulateRings(PreviousIndices, CurrentIndices, PolyGroupIndex, OutLocalMeshData);
+
+		PreviousIndices = MoveTemp(CurrentIndices);
+		CurrentIndices.Empty();
+	}
+}
+
+void FPVMeshBuilder::UpdateFoliagePivotPoints(const TSet<int32>& PointsToRemove, FManagedArrayCollection& OutCollection)
+{
+	const PV::Facades::FBranchFacade BranchFacade(OutCollection);
+	const PV::Facades::FPointFacade PointFacade(OutCollection);
+	PV::Facades::FFoliageFacade FoliageFacade(OutCollection);
+
+	TMap<int32, TArray<int32>> PointsIndicesToFoliageIndicesMap = GetPointsIndicesToFoliageIndicesMap(OutCollection);
+
+	for (int32 BranchIndex = 0; BranchIndex < BranchFacade.GetElementCount(); ++BranchIndex)
+	{
+		const TArray<int32>& BranchPoints = BranchFacade.GetPoints(BranchIndex);
+
+		if (BranchPoints.Num() == 0)
+		{
+			continue;
+		}
+
+		TArray<int32> SortedPointIndices = BranchPoints;
+		SortedPointIndices.Sort([PointFacade](const int32& A, const int32& B)
+			{
+				return PointFacade.GetLengthFromRoot(A) < PointFacade.GetLengthFromRoot(B);
+			});
+
+		for (int32 i = 0; i < SortedPointIndices.Num(); ++i)
+		{
+			const int32 CurrentPointIndex = SortedPointIndices[i];
+			if (PointsToRemove.Contains(CurrentPointIndex) &&
+				i > 0 &&
+				i < SortedPointIndices.Num() - 1 &&
+				PointsIndicesToFoliageIndicesMap.Contains(CurrentPointIndex))
+			{
+				int32 j = i - 1;
+				while (j > 0 && PointsToRemove.Contains(SortedPointIndices[j]))
+				{
+					j--;
+				}
+				const int32 PreviousPointIndex = SortedPointIndices[j];
+
+				int32 k = i + 1;
+				while (k < (SortedPointIndices.Num() - 1) && PointsToRemove.Contains(SortedPointIndices[k]))
+				{
+					k++;
+				}
+				const int32 NextPointIndex = SortedPointIndices[k];
+
+				const float PreviousPointLengthFromRoot = PointFacade.GetLengthFromRoot(PreviousPointIndex);
+				const float CurrentPointLengthFromRoot = PointFacade.GetLengthFromRoot(CurrentPointIndex);
+				const float NextPointLengthFromRoot = PointFacade.GetLengthFromRoot(NextPointIndex);
+				const float BlendValue = FMath::GetMappedRangeValueClamped(
+					FVector2f(PreviousPointLengthFromRoot, NextPointLengthFromRoot),
+					FVector2f(0.0f, 1.0f),
+					CurrentPointLengthFromRoot
+				);
+
+				const FVector3f& CurrentPointPosition = PointFacade.GetPosition(CurrentPointIndex);
+				const FVector3f& PreviousPointPosition = PointFacade.GetPosition(PreviousPointIndex);
+				const FVector3f& NextPointPosition = PointFacade.GetPosition(NextPointIndex);
+
+				const FVector3f InterpolatedPointPosition = FMath::Lerp(PreviousPointPosition, NextPointPosition, BlendValue);
+
+				for (const int32 FoliageIndex : PointsIndicesToFoliageIndicesMap[CurrentPointIndex])
+				{
+					const FVector3f& FoliagePosition = FoliageFacade.GetPivotPoint(FoliageIndex);
+					const FVector3f OffsetVector = FoliagePosition - CurrentPointPosition;
+					const FVector3f UpdatedFoliagePosition = InterpolatedPointPosition + OffsetVector;
+
+					FoliageFacade.SetPivotPoint(FoliageIndex, UpdatedFoliagePosition);
+				}
+			}
+		}
+	}
+}
+
+void FPVMeshBuilder::GetLocalDynamicMeshData(FManagedArrayCollection& Collection, const FPVMeshBuilderParams& MeshBuilderParams,
+                                        TArray<FLocalDynamicMeshData>& MeshesData)
+{
+	const PV::Facades::FBranchFacade BranchFacade(Collection);
+	const PV::Facades::FPointFacade PointFacade(Collection);
+	const PV::Facades::FPlantFacade PlantFacade(Collection);
+
+	if (!BranchFacade.IsValid())
+	{
+		UE_LOG(LogProceduralVegetation, Warning, TEXT("Fail to evaluate, InCollection dose not match the schema."));
+		return;
+	}
+
+	const TManagedArray<float>& PointScales = PointFacade.GetPointScales();
+	const float MaxPointScale = *Algo::MaxElement(PointScales);
+
+	const TSet<int32> HardPoints = CollectHardPoints(BranchFacade, PointFacade, PlantFacade);
+
+	TArray<float> MeshDivisionsGradientsForPoints;
+	TArray<float> DeltaModifiersForPoints;
+	TSet<int32> PointsToRemove = ComputePointGradients(PointFacade, MeshBuilderParams, HardPoints, MaxPointScale, MeshDivisionsGradientsForPoints,
+		DeltaModifiersForPoints);
+
+	PerformPathSimplification(BranchFacade, PointFacade, MeshBuilderParams, HardPoints, DeltaModifiersForPoints, PointsToRemove);
+
+	// Point removal may update curvate of the branch.
+	// Foliage instance pivot/attachment points need to be updated to compensate for that.
+	if (PointsToRemove.Num() > 0)
+	{
+		UpdateFoliagePivotPoints(PointsToRemove, Collection);
+	}
+
+	const TArray<int32> TargetMeshDivisions = ComputeMeshDivisions(BranchFacade, MeshBuilderParams, MeshDivisionsGradientsForPoints,
+		PointFacade.GetElementCount());
+
+	MeshesData.SetNum(BranchFacade.GetElementCount());
+
+	int32 DisplacementWidth = 0;
+	int32 DisplacementHeight = 0;
+	if (MeshBuilderParams.DisplacementTexture
+		&& MeshBuilderParams.DisplacementValues.Num() > 0
+		&& MeshBuilderParams.DisplacementStrength > 0.0f)
+	{
+		DisplacementWidth = MeshBuilderParams.DisplacementTexture->GetSizeX();
+		DisplacementHeight = MeshBuilderParams.DisplacementTexture->GetSizeY();
+
+		if (MeshBuilderParams.DisplacementTexture->SRGB)
+		{
+			UE_LOG(LogProceduralVegetation, Warning,
+				TEXT("Displacement texture selected has sRGB enabled! This will corrupt displacement data. Please disable it."));
+		}
+	}
+
+	const FDisplacementData DisplacementData(MeshBuilderParams.DisplacementValues, DisplacementWidth, DisplacementHeight);
+
+	ParallelFor(
+		BranchFacade.GetElementCount(),
+		[&](const int32 PrimitiveIndex)
+			{
+				const TArray<int32>& PrimitivePoints = BranchFacade.GetPoints(PrimitiveIndex).FilterByPredicate(
+					[&PointsToRemove](const int32 PointIndex)
+						{
+							return !PointsToRemove.Contains(PointIndex);
+						});
+
+				const bool PrimitiveIsTrunk = PlantFacade.IsTrunkIndex(PrimitiveIndex);
+				const int32 GenerationNumber = BranchFacade.GetGenerationNumber(PrimitiveIndex);
+
+				MeshesData[PrimitiveIndex].MaterialID = BranchFacade.GetBranchUVMaterial(PrimitiveIndex);
+
+				GenerateBranchMeshData(PrimitiveIsTrunk, GenerationNumber, PrimitivePoints, DisplacementData, MeshBuilderParams, TargetMeshDivisions,
+					Collection, MeshesData[PrimitiveIndex]);
+			});
+}
+
+void FPVMeshBuilder::ExtractDisplacementData(const TObjectPtr<UTexture2D>& Texture, TArray<float>& OutValues)
+{
+	OutValues.Empty();
+
+#if WITH_EDITOR
+	if (!Texture)
+	{
+		UE_LOG(LogProceduralVegetation, Warning, TEXT("ExtractDisplacementData: Texture is null."));
+		return;
+	}
+
+	FTextureSource& Source = Texture->Source;
+	if (!Source.IsValid())
+	{
+		UE_LOG(LogProceduralVegetation, Warning, TEXT("ExtractDisplacementData: No source data found on texture %s."), *Texture->GetName());
+		return;
+	}
+
+	const int32 Width = Source.GetSizeX();
+	const int32 Height = Source.GetSizeY();
+	const int32 NumPix = Width * Height;
+	OutValues.SetNumZeroed(NumPix);
+
+	const uint8* RawData = Source.LockMip(0);
+
+	switch (Source.GetFormat())
+	{
+	case TSF_RGBA32F:
+		{
+			const float* FloatData = reinterpret_cast<const float*>(RawData);
+			for (int32 i = 0; i < NumPix; ++i)
+			{
+				OutValues[i] = FloatData[i * 4 + 0];
+			}
+		}
+		break;
+
+	case TSF_RGBA16F:
+		{
+			const FFloat16* HalfData = reinterpret_cast<const FFloat16*>(RawData);
+			for (int32 i = 0; i < NumPix; ++i)
+			{
+				OutValues[i] = HalfData[i * 4 + 0].GetFloat();
+			}
+		}
+		break;
+	case TSF_BGRA8:
+		{
+			const FColor* ColorData = reinterpret_cast<const FColor*>(RawData);
+			constexpr float Inv255 = 1.0f / 255.0f;
+			for (int32 i = 0; i < NumPix; ++i)
+			{
+				OutValues[i] = ColorData[i].R * Inv255;
+			}
+		}
+		break;
+
+	default:
+		UE_LOG(LogProceduralVegetation, Warning, TEXT("ExtractDisplacementData: Unsupported source format %d on texture %s."),
+			(int32)Source.GetFormat(), *Texture->GetName());
+		break;
+	}
+
+
+	Source.UnlockMip(0);
+#else
+	UE_LOG(LogProceduralVegetation, Warning, TEXT("ExtractDisplacementData only works in editor builds (WITH_EDITOR)."));
+#endif
+}
diff --git a/Engine/Plugins/Experimental/ProceduralVegetationEditor/Source/ProceduralVegetation/Private/Implementations/PVMeshBuilder.h b/Engine/Plugins/Experimental/ProceduralVegetationEditor/Source/ProceduralVegetation/Private/Implementations/PVMeshBuilder.h
new file mode 100644
--- /dev/null
+++ b/Engine/Plugins/Experimental/ProceduralVegetationEditor/Source/ProceduralVegetation/Private/Implementations/PVMeshBuilder.h
@@ -0,0 +1,332 @@
+// Copyright Epic Games, Inc. All Rights Reserved.
+
+#pragma once
+#include "CoreMinimal.h"
+#include "UDynamicMesh.h"
+#include "Facades/PVBranchFacade.h"
+#include "Facades/PVPointFacade.h"
+#include "Generators/MeshShapeGenerator.h"
+#include "GeometryCollection/GeometryCollection.h"
+#include "GeometryCollection/ManagedArrayCollection.h"
+#include "Materials/MaterialInterface.h"
+#include "Materials/MaterialInstance.h"
+#include "PVMaterialSettings.h"
+#include "Curves/CurveFloat.h"
+#include "Engine/Texture2D.h"
+#include "Facades/PVPlantFacade.h"
+
+#include "PVMeshBuilder.generated.h"
+
+USTRUCT()
+struct FPVMeshBuilderParams
+{
+	GENERATED_BODY()
+
+	UPROPERTY(EditAnywhere, Category="Materials", meta=(PinHiddenByDefault, InlineEditConditionToggle))
+	bool bOverrideMaterial = false;
+
+	UPROPERTY(EditAnywhere, Category="Materials", meta = (ShowOnlyInnerProperties))
+	FPVMaterialSettings MaterialSettings;
+
+	UPROPERTY(EditAnywhere, Category="Mesh Settings", DisplayName="Point Removal",
+		meta=(ClampMin=0.0f, ClampMax=0.1f, UIMin=0.0f, UIMax=0.1f))
+	float PointRemoval = 0.0f;
+
+	UPROPERTY(EditAnywhere, Category="Mesh Settings", DisplayName="Segment Reduction",
+		meta=(ClampMin=0.0f, ClampMax=0.1f, UIMin=0.0f, UIMax=0.1f))
+	float SegmentReduction = 0.0f;
+
+	UPROPERTY(EditAnywhere, Category="Mesh Settings", DisplayName="Segment Retention Impact",
+		meta=(ClampMin=0.0f, ClampMax=2.0f, UIMin=0.0f, UIMax=2.0f))
+	float SegmentRetentionImpact = 1.0f;
+
+	UPROPERTY(EditAnywhere, Category="Mesh Settings", DisplayName="Longest Segment Allowed",
+		meta=(ClampMin=0.0f, ClampMax=1.0f, UIMin=0.0f, UIMax=1.0f))
+	float LongestSegmentLength = 0.1f;
+
+	UPROPERTY(EditAnywhere, Category="Mesh Settings", DisplayName="Shortest Segment Allowed",
+		meta=(ClampMin=0.0f, ClampMax=1.0f, UIMin=0.0f, UIMax=1.0f))
+	float ShortestSegmentLength = 0.0f;
+
+	UPROPERTY(EditAnywhere, Category="Mesh Settings", DisplayName="Reduction Accuracy",
+		meta=(ClampMin=3, ClampMax=15, UIMin=1, UIMax=15))
+	int32 Accuracy = 5;
+
+	UPROPERTY(EditAnywhere, Category="Mesh Settings", DisplayName="Min Divisions",
+		meta=(ClampMin=3, ClampMax=1024, UIMin=3, UIMax=36))
+	int32 MinMeshDivisions = 6;
+
+	UPROPERTY(EditAnywhere, Category="Mesh Settings", DisplayName="Max Divisions",
+		meta=(ClampMin=3, ClampMax=1024, UIMin=3, UIMax=36))
+	int32 MaxMeshDivisions = 12;
+
+	UPROPERTY(EditAnywhere, Category="Point Removal Settings", DisplayName="Hull Retention",
+		meta=(ClampMin=0.0f, ClampMax=1.0f, UIMin=0.0f, UIMax=1.0f))
+	float PointHullRetention = 0.0f;
+
+	UPROPERTY(EditAnywhere, Category="Point Removal Settings", DisplayName="Main Trunk Retention",
+		meta=(ClampMin=0.0f, ClampMax=1.0f, UIMin=0.0f, UIMax=1.0f))
+	float PointMainTrunkRetention = 0.0f;
+
+	UPROPERTY(EditAnywhere, Category="Point Removal Settings", DisplayName="Ground Retention",
+		meta=(ClampMin=0.0f, ClampMax=1.0f, UIMin=0.0f, UIMax=1.0f))
+	float PointGroundRetention = 0.0f;
+
+	UPROPERTY(EditAnywhere, Category="Point Removal Settings", DisplayName="Hull Retention Gradient", meta=(ClampMin=0.0f, ClampMax=1.0f))
+	FRuntimeFloatCurve PointHullRetentionGradient;
+
+	UPROPERTY(EditAnywhere, Category="Point Removal Settings", DisplayName="Main Trunk Retention Gradient", meta=(ClampMin=0.0f, ClampMax=1.0f))
+	FRuntimeFloatCurve PointMainTrunkRetentionGradient;
+
+	UPROPERTY(EditAnywhere, Category="Point Removal Settings", DisplayName="Ground Retention Gradient", meta=(ClampMin=0.0f, ClampMax=1.0f))
+	FRuntimeFloatCurve PointGroundRetentionGradient;
+
+	UPROPERTY(EditAnywhere, Category="Segment Removal Settings", DisplayName="Hull Retention",
+		meta=(ClampMin=0.0f, ClampMax=1.0f, UIMin=0.0f, UIMax=1.0f))
+	float SegmentHullRetention = 0.0f;
+
+	UPROPERTY(EditAnywhere, Category="Segment Removal Settings", DisplayName="Main Trunk Retention",
+		meta=(ClampMin=0.0f, ClampMax=1.0f, UIMin=0.0f, UIMax=1.0f))
+	float SegmentMainTrunkRetention = 0.0f;
+
+	UPROPERTY(EditAnywhere, Category="Segment Removal Settings", DisplayName="Ground Retention",
+		meta=(ClampMin=0.0f, ClampMax=1.0f, UIMin=0.0f, UIMax=1.0f))
+	float SegmentGroundRetention = 0.0f;
+
+	UPROPERTY(EditAnywhere, Category="Segment Removal Settings", DisplayName="Scale Retention",
+		meta=(ClampMin=0.0f, ClampMax=1.0f, UIMin=0.0f, UIMax=1.0f))
+	float SegmentScaleRetention = 0.0f;
+
+	UPROPERTY(EditAnywhere, Category="Segment Removal Settings", DisplayName="Hull Retention Gradient", meta=(ClampMin=0.0f, ClampMax=1.0f))
+	FRuntimeFloatCurve SegmentHullRetentionGradient;
+
+	UPROPERTY(EditAnywhere, Category="Segment Removal Settings", DisplayName="Main Trunk Retention Gradient", meta=(ClampMin=0.0f, ClampMax=1.0f))
+	FRuntimeFloatCurve SegmentMainTrunkRetentionGradient;
+
+	UPROPERTY(EditAnywhere, Category="Segment Removal Settings", DisplayName="Ground Retention Gradient", meta=(ClampMin=0.0f, ClampMax=1.0f))
+	FRuntimeFloatCurve SegmentGroundRetentionGradient;
+
+	UPROPERTY(EditAnywhere, Category="Segment Removal Settings", DisplayName="Scale Retention Gradient", meta=(ClampMin=0.0f, ClampMax=1.0f))
+	FRuntimeFloatCurve SegmentScaleRetentionGradient;
+
+	UPROPERTY(EditAnywhere, Category="Division Removal Settings", DisplayName="Hull Retention",
+		meta=(ClampMin=0.0f, ClampMax=1.0f, UIMin=0.0f, UIMax=1.0f))
+	float DivisionHullRetention = 0.0f;
+
+	UPROPERTY(EditAnywhere, Category="Division Removal Settings", DisplayName="Main Trunk Retention",
+		meta=(ClampMin=0.0f, ClampMax=1.0f, UIMin=0.0f, UIMax=1.0f))
+	float DivisionMainTrunkRetention = 0.0f;
+
+	UPROPERTY(EditAnywhere, Category="Division Removal Settings", DisplayName="Ground Retention",
+		meta=(ClampMin=0.0f, ClampMax=1.0f, UIMin=0.0f, UIMax=1.0f))
+	float DivisionGroundRetention = 0.0f;
+
+	UPROPERTY(EditAnywhere, Category="Division Removal Settings", DisplayName="Hull Retention Gradient", meta=(ClampMin=0.0f, ClampMax=1.0f))
+	FRuntimeFloatCurve DivisionHullRetentionGradient;
+
+	UPROPERTY(EditAnywhere, Category="Division Removal Settings", DisplayName="Main Trunk Retention Gradient", meta=(ClampMin=0.0f, ClampMax=1.0f))
+	FRuntimeFloatCurve DivisionMainTrunkRetentionGradient;
+
+	UPROPERTY(EditAnywhere, Category="Division Removal Settings", DisplayName="Ground Retention Gradient", meta=(ClampMin=0.0f, ClampMax=1.0f))
+	FRuntimeFloatCurve DivisionGroundRetentionGradient;
+
+	UPROPERTY(EditAnywhere, Category="Division Removal Settings", DisplayName="Scale Retention Gradient", meta=(ClampMin=0.0f, ClampMax=1.0f))
+	FRuntimeFloatCurve DivisionScaleRetentionGradient;
+
+	UPROPERTY(EditAnywhere, Category="Plant Profile Settings", DisplayName="Profile",
+		meta=(EditCondition="bIsPlantProfileDropdownEnabled", EditConditionHides, HideEditConditionToggle, GetOptions="GetPlantProfileOptions"))
+	FString SelectedPlantProfile = "None";
+
+	UPROPERTY(Transient)
+	TArray<FString> PlantProfileOptions = {"None"};
+
+	UPROPERTY(Transient)
+	bool bIsPlantProfileDropdownEnabled = false;
+
+	UPROPERTY(EditAnywhere, Category="Plant Profile Settings", DisplayName="Profile Falloff",
+		meta=(ClampMin=0.0f, ClampMax=1.0f, EditCondition="bIsPlantProfileDropdownEnabled", EditConditionHides, HideEditConditionToggle))
+	FRuntimeFloatCurve PlantProfileFallOff;
+
+	UPROPERTY(EditAnywhere, Category="Plant Profile Settings", DisplayName="Profile Scale",
+		meta=(ClampMin=0.0f, ClampMax=3.0f, UIMin=0.0f, UIMax=3.0f, EditCondition="bIsPlantProfileDropdownEnabled", EditConditionHides,
+			HideEditConditionToggle))
+	float PlantProfileScale = 1.0f;
+
+	UPROPERTY(EditAnywhere, Category="Plant Profile Settings",
+		meta=(EditCondition="bIsPlantProfileDropdownEnabled", EditConditionHides, HideEditConditionToggle))
+	bool bApplyProfileToBranches = false;
+
+	UPROPERTY(EditAnywhere, Category = "Displacement Settings")
+	TObjectPtr<UTexture2D> DisplacementTexture = nullptr;
+
+	UPROPERTY(Transient, NonTransactional)
+	TArray<float> DisplacementValues;
+
+	UPROPERTY(EditAnywhere, Category="Displacement Settings", DisplayName="Displacement Scale",
+		meta=(ClampMin=0.0f, ClampMax=10.0f, UIMin=0.0f, UIMax=1.0f))
+	float DisplacementStrength = 0.5f;
+
+	UPROPERTY(EditAnywhere, Category="Displacement Settings", DisplayName="Displacement Bias",
+		meta=(ClampMin=0.0f, ClampMax=10.0f, UIMin=0.0f, UIMax=10.0f))
+	float DisplacementBias = 0.0f;
+
+	UPROPERTY(EditAnywhere, Category="Displacement Settings", DisplayName="Displacement UV Scale")
+	FIntVector2 DisplacementUVScale = FIntVector2(1.0f, 1.0f);
+
+	UPROPERTY(EditAnywhere, Category="Displacement Settings", DisplayName="Apply Displacement up to generation",
+		meta=(ClampMin=1, ClampMax=10, UIMin=1, UIMax=10))
+	int32 DisplacementGenerationUpperLimit = 1;
+
+	FPVMeshBuilderParams()
+	{
+		InitializeLinearCurve(PointHullRetentionGradient);
+		InitializeLinearCurve(PointMainTrunkRetentionGradient);
+		InitializeLinearCurve(PointGroundRetentionGradient);
+		InitializeLinearCurve(SegmentHullRetentionGradient);
+		InitializeLinearCurve(SegmentMainTrunkRetentionGradient);
+		InitializeLinearCurve(SegmentGroundRetentionGradient);
+		InitializeLinearCurve(SegmentScaleRetentionGradient);
+		InitializeLinearCurve(DivisionHullRetentionGradient);
+		InitializeLinearCurve(DivisionMainTrunkRetentionGradient);
+		InitializeLinearCurve(DivisionGroundRetentionGradient);
+		InitializeLinearCurve(DivisionScaleRetentionGradient);
+		InitializeLinearCurve(PlantProfileFallOff);
+	}
+
+private:
+	static void InitializeLinearCurve(FRuntimeFloatCurve& Curve)
+	{
+		if (FRichCurve* RichCurve = Curve.GetRichCurve())
+		{
+			RichCurve->Reset();
+			const FKeyHandle Handle0 = RichCurve->AddKey(0.0f, 0.0f);
+			const FKeyHandle Handle1 = RichCurve->AddKey(1.0f, 1.0f);
+
+			RichCurve->SetKeyInterpMode(Handle0, RCIM_Linear);
+			RichCurve->SetKeyInterpMode(Handle1, RCIM_Linear);
+		}
+	}
+};
+
+struct FLocalDynamicMeshData
+{
+	struct FVertex
+	{
+		FVector3f Position;
+		FVector3f Normal;
+		FVector2f UV;
+		float NjordPixelId;
+	};
+
+	TArray<FVertex> Vertices;
+	TArray<FIntVector4> Triangles;
+	int MaterialID = 0;
+};
+
+struct FDisplacementData
+{
+	const TArrayView<const float> Values;
+	const int32 TextureWidth;
+	const int32 TextureHeight;
+
+	FDisplacementData(const TArrayView<const float>& InValues, const int32 InWidth, const int32 InHeight)
+		: Values(InValues)
+		, TextureWidth(InWidth)
+		, TextureHeight(InHeight)
+	{
+	}
+};
+
+class FPVMeshGenerator : public UE::Geometry::FMeshShapeGenerator
+{
+	int32 VerticesCount = 0;
+	int32 TriangleCount = 0;
+
+	TArray<FLocalDynamicMeshData>* MeshesDatas = nullptr;
+
+public:
+	FPVMeshGenerator(const int32 InVerticesCount, const int32 InTriangleCount, TArray<FLocalDynamicMeshData>* InMeshesDatas)
+		: VerticesCount(InVerticesCount)
+		, TriangleCount(InTriangleCount)
+		, MeshesDatas(InMeshesDatas)
+	{
+	}
+
+	virtual FMeshShapeGenerator& Generate() override
+	{
+		SetBufferSizes(VerticesCount, TriangleCount, VerticesCount, TriangleCount);
+
+		int32 VertexIndex = 0;
+		int32 FaceIndex = 0;
+		for (auto MeshIndex = MeshesDatas->Num() - 1; MeshIndex >= 0; --MeshIndex)
+		{
+			const FLocalDynamicMeshData& MeshesData = (*MeshesDatas)[MeshIndex];
+			for (const FIntVector4& Triangle : MeshesData.Triangles)
+			{
+				const FIntVector4 OffsetTriangle = Triangle + FIntVector4(VertexIndex);
+				const UE::Geometry::FIndex3i Indices = FIntVector(OffsetTriangle);
+				Triangles[FaceIndex] = Indices;
+				TriangleUVs[FaceIndex] = Indices;
+				TriangleNormals[FaceIndex] = Indices;
+				TrianglePolygonIDs[FaceIndex] = OffsetTriangle.W;
+				++FaceIndex;
+			}
+			for (const FLocalDynamicMeshData::FVertex& Vertex : MeshesData.Vertices)
+			{
+				Vertices[VertexIndex] = static_cast<FVector>(Vertex.Position);
+				Normals[VertexIndex] = Vertex.Normal;
+				NormalParentVertex[VertexIndex] = VertexIndex;
+				UVs[VertexIndex] = Vertex.UV;
+				UVParentVertex[VertexIndex] = VertexIndex;
+				++VertexIndex;
+			}
+		}
+
+		return *this;
+	}
+};
+
+struct FPVMeshBuilder
+{
+	static void GenerateGeometryCollection(const FManagedArrayCollection& InSkeletonCollection, const FPVMeshBuilderParams& MeshBuilderParams,
+	                                       FGeometryCollection& OutGeometryCollection);
+
+	static void SetNjordPixelID(FGeometryCollection& OutGeometryCollection);
+
+	static void GenerateDynamicMesh(FManagedArrayCollection& Collection, const FPVMeshBuilderParams& MeshBuilderParams,
+	                                TObjectPtr<UDynamicMesh>& OutMesh);
+
+	static void ExtractDisplacementData(const TObjectPtr<UTexture2D>& Texture, TArray<float>& OutValues);
+
+private:
+	static TSet<int32> CollectHardPoints(const PV::Facades::FBranchFacade& BranchFacade, const PV::Facades::FPointFacade& PointFacade,
+	                                     const PV::Facades::FPlantFacade& PlantFacade);
+
+	static TSet<int32> ComputePointGradients(const PV::Facades::FPointFacade& PointFacade, const FPVMeshBuilderParams& MeshBuilderParams,
+	                                         const TSet<int32>& HardPoints, const float MaxPointScale, TArray<float>& OutMeshDivisionsGradients,
+	                                         TArray<float>& OutDeltaModifiers);
+
+	static void PerformPathSimplification(const PV::Facades::FBranchFacade& BranchFacade,
+	                                      const PV::Facades::FPointFacade& PointFacade, const FPVMeshBuilderParams& MeshBuilderParams,
+	                                      const TSet<int32>& HardPoints, const TArray<float>& DeltaModifiers, TSet<int32>& InOutPointsToRemove);
+
+	static TArray<int32> ComputeMeshDivisions(const PV::Facades::FBranchFacade& BranchFacade, const FPVMeshBuilderParams& MeshBuilderParams,
+	                                          const TArray<float>& MeshDivisionsGradients, const int32 PointCount);
+
+	static void TriangulateRings(const TArray<int32>& PreviousIndices, const TArray<int32>& CurrentIndices, int32& InOutPolyGroupIndex,
+	                             FLocalDynamicMeshData& OutMeshData);
+
+	static float GetProfileMultiplier(const TArray<float>& InProfilePoints, const float ProfileUV_U);
+
+	static TMap<int32, TArray<int32>> GetPointsIndicesToFoliageIndicesMap(const FManagedArrayCollection& Collection);
+
+	static void GenerateBranchMeshData(const bool bPrimitiveIsTrunk, const int32 GenerationNumber, const TArray<int32>& PrimitivePoints,
+	                                   const FDisplacementData& DisplacementData, const FPVMeshBuilderParams& MeshBuilderParams,
+	                                   const TArray<int32>& TargetMeshDivisions, const FManagedArrayCollection& Collection,
+	                                   FLocalDynamicMeshData& OutLocalMeshData);
+
+	static void UpdateFoliagePivotPoints(const TSet<int32>& PointsToRemove, FManagedArrayCollection& OutCollection);
+
+	static void GetLocalDynamicMeshData(FManagedArrayCollection& Collection, const FPVMeshBuilderParams& MeshBuilderParams,
+	                                    TArray<FLocalDynamicMeshData>& MeshesData);
+};
diff --git a/Engine/Plugins/Experimental/ProceduralVegetationEditor/Source/ProceduralVegetation/Private/Implementations/PVMesher.cpp b/Engine/Plugins/Experimental/ProceduralVegetationEditor/Source/ProceduralVegetation/Private/Implementations/PVMesher.cpp
deleted file mode 100644
diff --git a/Engine/Plugins/Experimental/ProceduralVegetationEditor/Source/ProceduralVegetation/Private/Implementations/PVMesher.h b/Engine/Plugins/Experimental/ProceduralVegetationEditor/Source/ProceduralVegetation/Private/Implementations/PVMesher.h
deleted file mode 100644
diff --git a/Engine/Plugins/Experimental/ProceduralVegetationEditor/Source/ProceduralVegetation/Private/Nodes/PVBaseSettings.cpp b/Engine/Plugins/Experimental/ProceduralVegetationEditor/Source/ProceduralVegetation/Private/Nodes/PVBaseSettings.cpp
new file mode 100644
--- /dev/null
+++ b/Engine/Plugins/Experimental/ProceduralVegetationEditor/Source/ProceduralVegetation/Private/Nodes/PVBaseSettings.cpp
@@ -0,0 +1,75 @@
+﻿// Copyright Epic Games, Inc. All Rights Reserved.
+
+#include "Nodes/PVBaseSettings.h"
+#include "PCGContext.h"
+#include "Helpers/PCGHelpers.h"
+#include "DataTypes/PVData.h"
+#include "PCGData.h"
+
+TArray<FPCGPinProperties> UPVBaseSettings::InputPinProperties() const
+{
+	TArray<FPCGPinProperties> Properties;
+
+	FPCGPinProperties& Pin = Properties.Emplace_GetRef(PCGPinConstants::DefaultInputLabel, GetInputPinTypeIdentifier());
+	Pin.SetRequiredPin();
+	Pin.SetAllowMultipleConnections(false);
+	Pin.bAllowMultipleData = false;
+
+	return Properties;
+}
+
+TArray<FPCGPinProperties> UPVBaseSettings::OutputPinProperties() const
+{
+	TArray<FPCGPinProperties> Properties;
+
+	FPCGPinProperties& Pin = Properties.Emplace_GetRef(PCGPinConstants::DefaultOutputLabel, GetOutputPinTypeIdentifier());
+	Pin.SetAllowMultipleConnections(true);
+	Pin.bAllowMultipleData = false;
+
+	return Properties;
+}
+
+FPCGDataTypeIdentifier UPVBaseSettings::GetInputPinTypeIdentifier() const
+{
+	return FPCGDataTypeIdentifier{ EPCGDataType::Any };
+}
+
+FPCGDataTypeIdentifier UPVBaseSettings::GetOutputPinTypeIdentifier() const
+{
+	return FPCGDataTypeIdentifier{ EPCGDataType::Any };
+}
+
+FPCGElementPtr UPVBaseSettings::CreateElement() const
+{
+	return MakeShared<FPVBaseElement>();
+}
+
+void FPVBaseElement::PostExecuteInternal(FPCGContext* InContext) const
+{
+	IPCGElement::PostExecuteInternal(InContext);
+
+#if WITH_EDITOR
+	const UPVBaseSettings* Settings = InContext->GetInputSettings<UPVBaseSettings>();
+	check(Settings);
+
+	for (FPCGTaggedData& OutputData : InContext->OutputData.TaggedData)
+	{
+		if (OutputData.Data && OutputData.Data.IsA<UPVData>())
+		{
+			if (const UPVData* ProceduralVegetationData = Cast<UPVData>(OutputData.Data))
+			{
+				ProceduralVegetationData->SetRenderType(Settings->GetRenderType());
+
+#if PV_SHOW_DEBUG_VISUALIZATION
+				ProceduralVegetationData->SetDebugSettings(Settings->GetDebugVisualizationSettings());
+#endif
+			}
+		}
+	}
+#endif
+}
+
+bool FPVBaseElement::ExecuteInternal(FPCGContext* Context) const
+{
+	return true;
+}
\ No newline at end of file
diff --git a/Engine/Plugins/Experimental/ProceduralVegetationEditor/Source/ProceduralVegetation/Private/Nodes/PVBoneReductionSettings.cpp b/Engine/Plugins/Experimental/ProceduralVegetationEditor/Source/ProceduralVegetation/Private/Nodes/PVBoneReductionSettings.cpp
new file mode 100644
--- /dev/null
+++ b/Engine/Plugins/Experimental/ProceduralVegetationEditor/Source/ProceduralVegetation/Private/Nodes/PVBoneReductionSettings.cpp
@@ -0,0 +1,69 @@
+﻿// Copyright Epic Games, Inc. All Rights Reserved.
+
+#include "PVBoneReductionSettings.h"
+#include "DataTypes/PVData.h"
+#include "DataTypes/PVMeshData.h"
+#include "Data/PCGBasePointData.h"
+#include "Facades/PVBoneFacade.h"
+#include "GeometryCollection/ManagedArrayCollection.h"
+#include "Helpers/PCGSettingsHelpers.h"
+
+FPCGDataTypeIdentifier UPVBoneReductionSettings::GetInputPinTypeIdentifier() const
+{
+	return FPCGDataTypeIdentifier{ FPVDataTypeInfoMesh::AsId() };
+}
+
+FPCGDataTypeIdentifier UPVBoneReductionSettings::GetOutputPinTypeIdentifier() const
+{
+	return FPCGDataTypeIdentifier{ FPVDataTypeInfoMesh::AsId() };
+}
+
+FPCGElementPtr UPVBoneReductionSettings::CreateElement() const
+{
+	return MakeShared<FPVBoneReductionElement>();
+}
+
+bool FPVBoneReductionElement::ExecuteInternal(FPCGContext* InContext) const
+{
+	TRACE_CPUPROFILER_EVENT_SCOPE(FPVBoneReductionElement::ExecuteInternal);
+
+	check(InContext);
+
+	const UPVBoneReductionSettings* InputSettings = InContext->GetInputSettings<UPVBoneReductionSettings>();
+	check(InputSettings);
+
+	for (const FPCGTaggedData& Input : InContext->InputData.GetInputsByPin(PCGPinConstants::DefaultInputLabel))
+	{
+		if (const UPVMeshData* InputData = Cast<UPVMeshData>(Input.Data))
+		{
+			FManagedArrayCollection SourceCollection = InputData->GetCollection();
+			PV::Facades::FBoneFacade BoneFacadeConst = PV::Facades::FBoneFacade(SourceCollection);
+			if (!BoneFacadeConst.IsValid())
+			{
+				PCGLog::LogErrorOnGraph(FText::FromString(TEXT("Can not reduce bones. Mesh data not available.")), InContext);
+				return true;
+			}
+
+			FManagedArrayCollection OutCollection;
+			SourceCollection.CopyTo(&OutCollection);
+
+			PV::Facades::FBoneFacade BoneFacade = PV::Facades::FBoneFacade(OutCollection);
+			UPVMeshData* OutManagedArrayCollectionData = FPCGContext::NewObject_AnyThread<UPVMeshData>(InContext);
+
+			TArray<PV::Facades::FBoneNode> BoneNodes;
+			float ReductionStrength = InputSettings->Strength;
+			BoneFacade.CreateBoneData(BoneNodes, ReductionStrength);
+			BoneFacade.SetBoneDataToCollection(BoneNodes);
+			
+			OutManagedArrayCollectionData->Initialize(MoveTemp(OutCollection));
+			InContext->OutputData.TaggedData.Emplace(OutManagedArrayCollectionData);
+		}
+		else
+		{
+			PCGLog::InputOutput::LogInvalidInputDataError(InContext);
+			return true;
+		}
+	}
+
+	return true;
+}
diff --git a/Engine/Plugins/Experimental/ProceduralVegetationEditor/Source/ProceduralVegetation/Private/Nodes/PVBoneReductionSettings.h b/Engine/Plugins/Experimental/ProceduralVegetationEditor/Source/ProceduralVegetation/Private/Nodes/PVBoneReductionSettings.h
new file mode 100644
--- /dev/null
+++ b/Engine/Plugins/Experimental/ProceduralVegetationEditor/Source/ProceduralVegetation/Private/Nodes/PVBoneReductionSettings.h
@@ -0,0 +1,40 @@
+﻿// Copyright Epic Games, Inc. All Rights Reserved.
+#pragma once
+#include "CoreMinimal.h"
+#include "Nodes/PVBaseSettings.h"
+#include "DataTypes/PVData.h"
+#include "PVBoneReductionSettings.generated.h"
+
+UCLASS(BlueprintType, ClassGroup = (Procedural))
+class UPVBoneReductionSettings : public UPVBaseSettings
+{
+	GENERATED_BODY()
+
+	friend class FPVBoneReductionElement;
+
+public:
+#if WITH_EDITOR
+	virtual FName GetDefaultNodeName() const override { return FName(TEXT("ProceduralVegetationBoneReduction")); }
+	virtual FText GetDefaultNodeTitle() const override { return NSLOCTEXT("PVBoneReductionSettings", "NodeTitle", "Bone Reduction"); }
+	virtual FText GetNodeTooltipText() const override { return NSLOCTEXT("PVBoneReductionSettings", "NodeTooltip", "Controls the number of bones created for skeletal mesh"); }
+	virtual FLinearColor GetNodeTitleColor() const override { return FColor::Cyan; }
+
+	virtual TArray<EPVRenderType> GetRenderType() const override { return TArray{EPVRenderType::Bones, EPVRenderType::Mesh}; }
+#endif
+
+protected:
+	virtual FPCGDataTypeIdentifier GetInputPinTypeIdentifier() const override;
+	virtual FPCGDataTypeIdentifier GetOutputPinTypeIdentifier() const override;
+
+	virtual FPCGElementPtr CreateElement() const override;
+
+private:
+	UPROPERTY(EditAnywhere, Category = "Bone Reduction", meta = (ShowOnlyInnerProperties, ClampMin ="0", ClampMax = "1", UIMin = "0", UIMax = "1"))
+	float Strength;
+};
+
+class FPVBoneReductionElement : public FPVBaseElement
+{
+protected:
+	virtual bool ExecuteInternal(FPCGContext* InContext) const override;
+};
diff --git a/Engine/Plugins/Experimental/ProceduralVegetationEditor/Source/ProceduralVegetation/Private/Nodes/PVCarveSettings.cpp b/Engine/Plugins/Experimental/ProceduralVegetationEditor/Source/ProceduralVegetation/Private/Nodes/PVCarveSettings.cpp
new file mode 100644
--- /dev/null
+++ b/Engine/Plugins/Experimental/ProceduralVegetationEditor/Source/ProceduralVegetation/Private/Nodes/PVCarveSettings.cpp
@@ -0,0 +1,59 @@
+﻿// Copyright Epic Games, Inc. All Rights Reserved.
+
+#include "PVCarveSettings.h"
+#include "DataTypes/PVData.h"
+#include "DataTypes/PVGrowthData.h"
+#include "Data/PCGBasePointData.h"
+#include "GeometryCollection/ManagedArrayCollection.h"
+#include "Helpers/PCGSettingsHelpers.h"
+
+FPCGDataTypeIdentifier UPVCarveSettings::GetInputPinTypeIdentifier() const
+{
+	return FPCGDataTypeIdentifier{ FPVDataTypeInfoGrowth::AsId() };
+}
+
+FPCGDataTypeIdentifier UPVCarveSettings::GetOutputPinTypeIdentifier() const
+{
+	return FPCGDataTypeIdentifier{ FPVDataTypeInfoGrowth::AsId() };
+}
+
+FPCGElementPtr UPVCarveSettings::CreateElement() const
+{
+	return MakeShared<FPVCarveElement>();
+}
+
+bool FPVCarveElement::ExecuteInternal(FPCGContext* InContext) const
+{
+	TRACE_CPUPROFILER_EVENT_SCOPE(FPVCarveElement::ExecuteInternal);
+
+	check(InContext);
+
+	const UPVCarveSettings* InputSettings = InContext->GetInputSettings<UPVCarveSettings>();
+	check(InputSettings);
+
+	for (const FPCGTaggedData& Input : InContext->InputData.GetInputsByPin(PCGPinConstants::DefaultInputLabel))
+	{
+		if (const UPVGrowthData* InputData = Cast<UPVGrowthData>(Input.Data))
+		{
+			FManagedArrayCollection SourceCollection = InputData->GetCollection();
+
+			FManagedArrayCollection OutCollection;
+			SourceCollection.CopyTo(&OutCollection);
+
+			UPVGrowthData* OutManagedArrayCollectionData = FPCGContext::NewObject_AnyThread<UPVGrowthData>(InContext);
+
+			auto [CarveBasis, Carve] = InputSettings->CarveSettings;
+			FPVCarve::ApplyCarve(OutCollection, SourceCollection, CarveBasis, Carve);
+
+			OutManagedArrayCollectionData->Initialize(MoveTemp(OutCollection));
+			InContext->OutputData.TaggedData.Emplace(OutManagedArrayCollectionData);
+		}
+		else
+		{
+			PCGLog::InputOutput::LogInvalidInputDataError(InContext);
+			return true;
+		}
+	}
+
+	return true;
+}
diff --git a/Engine/Plugins/Experimental/ProceduralVegetationEditor/Source/ProceduralVegetation/Private/Nodes/PVCarveSettings.h b/Engine/Plugins/Experimental/ProceduralVegetationEditor/Source/ProceduralVegetation/Private/Nodes/PVCarveSettings.h
new file mode 100644
--- /dev/null
+++ b/Engine/Plugins/Experimental/ProceduralVegetationEditor/Source/ProceduralVegetation/Private/Nodes/PVCarveSettings.h
@@ -0,0 +1,42 @@
+﻿// Copyright Epic Games, Inc. All Rights Reserved.
+#pragma once
+
+#include "CoreMinimal.h"
+#include "Nodes/PVBaseSettings.h"
+#include "Implementations/PVCarve.h"
+#include "DataTypes/PVData.h"
+#include "PVCarveSettings.generated.h"
+
+UCLASS(BlueprintType, ClassGroup = (Procedural))
+class UPVCarveSettings : public UPVBaseSettings
+{
+	GENERATED_BODY()
+
+	friend class FPVCarveElement;
+
+public:
+#if WITH_EDITOR
+	virtual FName GetDefaultNodeName() const override { return FName(TEXT("ProceduralVegetationCarve")); }
+	virtual FText GetDefaultNodeTitle() const override { return NSLOCTEXT("PVCarveSettings", "NodeTitle", "Carve"); }
+	virtual FText GetNodeTooltipText() const override { return NSLOCTEXT("PVCarveSettings", "NodeTooltip", "Carve branches or trunk"); }
+	virtual FLinearColor GetNodeTitleColor() const override { return FLinearColor::Yellow; }
+
+	virtual TArray<EPVRenderType> GetRenderType() const override { return TArray{EPVRenderType::Skeleton}; }
+#endif
+
+protected:
+	virtual FPCGDataTypeIdentifier GetInputPinTypeIdentifier() const override;
+	virtual FPCGDataTypeIdentifier GetOutputPinTypeIdentifier() const override;
+
+	virtual FPCGElementPtr CreateElement() const override;
+
+private:
+	UPROPERTY(EditAnywhere, Category = "Carve", meta = (ShowOnlyInnerProperties))
+	FPVCarveParams CarveSettings;
+};
+
+class FPVCarveElement : public FPVBaseElement
+{
+protected:
+	virtual bool ExecuteInternal(FPCGContext* InContext) const override;
+};
diff --git a/Engine/Plugins/Experimental/ProceduralVegetationEditor/Source/ProceduralVegetation/Private/Nodes/PVFoliageDistributorSettings.cpp b/Engine/Plugins/Experimental/ProceduralVegetationEditor/Source/ProceduralVegetation/Private/Nodes/PVFoliageDistributorSettings.cpp
new file mode 100644
--- /dev/null
+++ b/Engine/Plugins/Experimental/ProceduralVegetationEditor/Source/ProceduralVegetation/Private/Nodes/PVFoliageDistributorSettings.cpp
@@ -0,0 +1,127 @@
+﻿// Copyright Epic Games, Inc. All Rights Reserved.
+
+#include "PVFoliageDistributorSettings.h"
+#include "DataTypes/PVData.h"
+#include "DataTypes/PVMeshData.h"
+#include "DataTypes/PVFoliageMeshData.h"
+#include "Data/PCGBasePointData.h"
+#include "GeometryCollection/ManagedArrayCollection.h"
+#include "Helpers/PCGHelpers.h"
+#include "Helpers/PCGSettingsHelpers.h"
+
+UPVFoliageDistributorSettings::UPVFoliageDistributorSettings()
+{
+	if (PCGHelpers::IsNewObjectAndNotDefault(this))
+	{
+		FRichCurve* RichCurve = InstanceSpacingRamp.GetRichCurve();
+		RichCurve->Reset();
+		const FKeyHandle Handle0 = RichCurve->AddKey(0.0f, 0.0f);
+		const FKeyHandle Handle1 = RichCurve->AddKey(1.0f, 1.0f);
+		RichCurve->SetKeyInterpMode(Handle0, RCIM_Linear);
+		RichCurve->SetKeyInterpMode(Handle1, RCIM_Linear);
+
+		RichCurve = ScaleRamp.GetRichCurve();
+		RichCurve->Reset();
+		const FKeyHandle Handle2 = RichCurve->AddKey(0.0f, 0.0f);
+		const FKeyHandle Handle3 = RichCurve->AddKey(1.0f, 1.0f);
+		RichCurve->SetKeyInterpMode(Handle2, RCIM_Linear);
+		RichCurve->SetKeyInterpMode(Handle3, RCIM_Linear);
+
+		RichCurve = AxilAngleRamp.GetRichCurve();
+		RichCurve->Reset();
+		const FKeyHandle Handle4 = RichCurve->AddKey(0.0f, 1.0f);
+		const FKeyHandle Handle5 = RichCurve->AddKey(1.0f, 0.5f);
+		RichCurve->SetKeyInterpMode(Handle4, RCIM_Linear);
+		RichCurve->SetKeyInterpMode(Handle5, RCIM_Linear);
+	}
+}
+
+FPCGDataTypeIdentifier UPVFoliageDistributorSettings::GetInputPinTypeIdentifier() const
+{
+	return FPCGDataTypeIdentifier{ FPVDataTypeInfoMesh::AsId() };
+}
+
+FPCGDataTypeIdentifier UPVFoliageDistributorSettings::GetOutputPinTypeIdentifier() const
+{
+	return FPCGDataTypeIdentifier{ FPVDataTypeInfoFoliageMesh::AsId() };
+}
+
+FPCGElementPtr UPVFoliageDistributorSettings::CreateElement() const
+{
+	return MakeShared<FPVFoliageDistributorElement>();
+}
+
+bool FPVFoliageDistributorElement::ExecuteInternal(FPCGContext* InContext) const
+{
+	TRACE_CPUPROFILER_EVENT_SCOPE(FPVFoliageDistributorElement::Execute);
+
+	check(InContext);
+
+	const UPVFoliageDistributorSettings* Settings = InContext->GetInputSettings<UPVFoliageDistributorSettings>();
+	check(Settings);
+
+	for (const FPCGTaggedData& Input : InContext->InputData.GetInputsByPin(PCGPinConstants::DefaultInputLabel))
+	{
+		if(const UPVMeshData* InputData = Cast<UPVMeshData>(Input.Data))
+		{
+			FManagedArrayCollection Collection = InputData->GetCollection();
+			
+			UPVFoliageMeshData* OutManagedArrayCollectionData = FPCGContext::NewObject_AnyThread<UPVFoliageMeshData>(InContext);
+
+			const FDistributionSettings DistributionSettings = {
+				.EthyleneThreshold = Settings->EthyleneThreshold,
+				.OverrideDistribution = Settings->OverrideDistribution,
+				.InstanceSpacing = Settings->InstanceSpacing,
+				.InstanceSpacingRamp = &Settings->InstanceSpacingRamp,
+				.InstanceSpacingRampEffect = Settings->InstanceSpacingRampEffect,
+				.MaxPerBranch = Settings->MaxPerBranch
+			};
+
+			const FScaleSettings ScaleSettings = {
+				.BaseScale = Settings->BaseScale,
+				.BranchScaleImpact = Settings->BranchScaleImpact,
+				.MinScale = Settings->MinScale,
+				.MaxScale = Settings->MaxScale,
+				.RandomScaleMin = Settings->RandomScaleMin,
+				.RandomScaleMax = Settings->RandomScaleMax,
+				.ScaleRamp = &Settings->ScaleRamp
+			};
+
+			const FVectorSettings VectorSettings = {
+				.OverrideAxilAngle = Settings->OverrideAxilAngle,
+				.AxilAngle = Settings->AxilAngle,
+				.AxilAngleRamp = &Settings->AxilAngleRamp,
+				.AxilAngleRampUpperValue = Settings->AxilAngleRampUpperValue,
+				.AxilAngleRampEffect = Settings->AxilAngleRampEffect
+			};
+
+			const FPhyllotaxySettings PhyllotaxySettings = {
+				.OverridePhyllotaxy = Settings->OverridePhyllotaxy,
+				.PhyllotaxyType = Settings->PhyllotaxyType,
+				.PhyllotaxyFormation = Settings->PhyllotaxyFormation,
+				.MinimumNodeBuds = Settings->MinimumNodeBuds,
+				.MaximumNodeBuds = Settings->MaximumNodeBuds,
+				.ResetPhyllotaxy = Settings->ResetPhyllotaxy,
+				.PhyllotaxyAdditionalAngle = Settings->PhyllotaxyAdditionalAngle,
+				.PhyllotaxyOffset = Settings->PhyllotaxyOffset
+			};
+
+			const FMiscSettings MiscSettings = {
+				.RandomSeed = Settings->RandomSeed
+			};
+
+			FPVFoliage::DistributeFoliage(Collection, Collection, DistributionSettings, ScaleSettings, VectorSettings,
+												 PhyllotaxySettings, MiscSettings);
+
+			OutManagedArrayCollectionData->Initialize(MoveTemp(Collection));
+			InContext->OutputData.TaggedData.Emplace(OutManagedArrayCollectionData);
+		}
+		else
+		{
+			PCGLog::InputOutput::LogInvalidInputDataError(InContext);
+			return true;
+		}
+	}
+	
+	return true;
+}
\ No newline at end of file
diff --git a/Engine/Plugins/Experimental/ProceduralVegetationEditor/Source/ProceduralVegetation/Private/Nodes/PVFoliageDistributorSettings.h b/Engine/Plugins/Experimental/ProceduralVegetationEditor/Source/ProceduralVegetation/Private/Nodes/PVFoliageDistributorSettings.h
new file mode 100644
--- /dev/null
+++ b/Engine/Plugins/Experimental/ProceduralVegetationEditor/Source/ProceduralVegetation/Private/Nodes/PVFoliageDistributorSettings.h
@@ -0,0 +1,125 @@
+﻿// Copyright Epic Games, Inc. All Rights Reserved.
+#pragma once
+
+#include "CoreMinimal.h"
+#include "Nodes/PVBaseSettings.h"
+#include "Implementations/PVFoliage.h"
+#include "DataTypes/PVData.h"
+#include "PVFoliageDistributorSettings.generated.h"
+
+
+UCLASS(BlueprintType, ClassGroup = (Procedural))
+class UPVFoliageDistributorSettings : public UPVBaseSettings
+{
+	GENERATED_BODY()
+
+	friend class FPVFoliageDistributorElement;
+	
+public:
+	UPVFoliageDistributorSettings();
+
+	#if WITH_EDITOR
+	virtual FName GetDefaultNodeName() const override { return FName(TEXT("ProceduralVegetationFoliageDistributor")); }
+	virtual FText GetDefaultNodeTitle() const override { return NSLOCTEXT("PVFoliageDistributorSettings", "NodeTitle", "Foliage Distributor"); }
+	virtual FText GetNodeTooltipText() const override { return NSLOCTEXT("PVFoliageDistributorSettings", "NodeTooltip", "Distributes the Foliage"); }
+	virtual FLinearColor GetNodeTitleColor() const override { return FLinearColor::Blue; }
+
+	virtual TArray<EPVRenderType> GetRenderType() const override { return TArray{ EPVRenderType::Mesh, EPVRenderType::Foliage }; }
+	#endif
+
+protected:
+
+	virtual FPCGDataTypeIdentifier GetInputPinTypeIdentifier() const override;
+	virtual FPCGDataTypeIdentifier GetOutputPinTypeIdentifier() const override;
+
+	virtual FPCGElementPtr CreateElement() const override;
+
+private:
+
+	UPROPERTY(EditAnywhere, Category="Distribution Settings", meta=(ClampMin=0.0f, ClampMax=1.0f, UIMin=0.0f, UIMax=1.0f))
+	float EthyleneThreshold = 0.85f;
+
+	UPROPERTY(EditAnywhere, Category="Distribution Settings")
+	bool OverrideDistribution = true;
+
+	UPROPERTY(EditAnywhere, Category="Distribution Settings", meta=(ClampMin=0.0f, ClampMax=1.0f, UIMin=0.0f, UIMax=1.0f))
+	float InstanceSpacing = 0.3f;
+
+	UPROPERTY(EditAnywhere, Category="Distribution Settings", meta=(ClampMin=0.0f, ClampMax=1.0f))
+	FRuntimeFloatCurve InstanceSpacingRamp;
+
+	UPROPERTY(EditAnywhere, Category="Distribution Settings", meta=(ClampMin=0.0f, ClampMax=1.0f, UIMin=0.0f, UIMax=1.0f))
+	float InstanceSpacingRampEffect = 0.0f;
+
+	UPROPERTY(EditAnywhere, Category="Distribution Settings", meta=(ClampMin=-1))
+	int32 MaxPerBranch = -1;
+
+	UPROPERTY(EditAnywhere, Category="Scale Settings", meta=(ClampMin=0.0f, ClampMax=2.0f, UIMin=0.0f, UIMax=2.0f))
+	float BaseScale = 1.0f;
+
+	UPROPERTY(EditAnywhere, Category="Scale Settings", meta=(ClampMin=0.0f, ClampMax=1.0f, UIMin=0.0f, UIMax=1.0f))
+	float BranchScaleImpact = 0.0f;
+
+	UPROPERTY(EditAnywhere, Category="Scale Settings", meta=(ClampMin=0.0f, ClampMax=1.0f, UIMin=0.0f, UIMax=1.0f))
+	float MinScale = 0.5f;
+
+	UPROPERTY(EditAnywhere, Category="Scale Settings", meta=(ClampMin=0.0f, ClampMax=1.0f, UIMin=0.0f, UIMax=1.0f))
+	float MaxScale = 1.0f;
+
+	UPROPERTY(EditAnywhere, Category="Scale Settings", meta=(ClampMin=0.0f, ClampMax=1.0f, UIMin=0.0f, UIMax=1.0f))
+	float RandomScaleMin = 1.0f;
+
+	UPROPERTY(EditAnywhere, Category="Scale Settings", meta=(ClampMin=1.0f, ClampMax=2.0f, UIMin=1.0f, UIMax=2.0f))
+	float RandomScaleMax = 1.0f;
+
+	UPROPERTY(EditAnywhere, Category="Scale Settings", meta=(ClampMin=0.0f, ClampMax=1.0f))
+	FRuntimeFloatCurve ScaleRamp;
+
+	UPROPERTY(EditAnywhere, Category="Vector Settings")
+	bool OverrideAxilAngle = true;
+
+	UPROPERTY(EditAnywhere, Category="Vector Settings", meta=(ClampMin=0.0f, ClampMax=90.0f, UIMin=0.0f, UIMax=90.0f))
+	float AxilAngle = 45.0f;
+
+	UPROPERTY(EditAnywhere, Category="Vector Settings", meta=(ClampMin=0.0f, ClampMax=1.0f))
+	FRuntimeFloatCurve AxilAngleRamp;
+
+	UPROPERTY(EditAnywhere, Category="Vector Settings", meta=(ClampMin=0.0f, ClampMax=90.0f, UIMin=0.0f, UIMax=90.0f))
+	float AxilAngleRampUpperValue = 45.0f;
+
+	UPROPERTY(EditAnywhere, Category="Vector Settings", meta=(ClampMin=0.0f, ClampMax=1.0f, UIMin=0.0f, UIMax=1.0f))
+	float AxilAngleRampEffect = 0.0f;
+
+	UPROPERTY(EditAnywhere, Category="Phyllotaxy Settings")
+	bool OverridePhyllotaxy = false;
+
+	UPROPERTY(EditAnywhere, Category="Phyllotaxy Settings")
+	EPhyllotaxyType PhyllotaxyType = EPhyllotaxyType::Alternate;
+
+	UPROPERTY(EditAnywhere, Category="Phyllotaxy Settings", meta=(EditCondition="PhyllotaxyType == EPhyllotaxyType::Spiral", EditConditionHides))
+	EPhyllotaxyFormation PhyllotaxyFormation = EPhyllotaxyFormation::Distichous;
+
+	UPROPERTY(EditAnywhere, Category="Phyllotaxy Settings", meta=(ClampMin=1, ClampMax=10, UIMin=1, UIMax=10))
+	int32 MinimumNodeBuds = 2;
+
+	UPROPERTY(EditAnywhere, Category="Phyllotaxy Settings", meta=(ClampMin=1, ClampMax=10, UIMin=1, UIMax=10))
+	int32 MaximumNodeBuds = 3;
+
+	UPROPERTY(EditAnywhere, Category="Phyllotaxy Settings")
+	bool ResetPhyllotaxy = false;
+
+	UPROPERTY(EditAnywhere, Category="Phyllotaxy Settings", meta=(ClampMin=0.0f, ClampMax=360.0f, UIMin=0.0f, UIMax=360.0f))
+	float PhyllotaxyAdditionalAngle = 0.0f;
+
+	UPROPERTY(EditAnywhere, Category="Phyllotaxy Settings", meta=(ClampMin=0.0f, ClampMax=360.0f, UIMin=0.0f, UIMax=360.0f))
+	float PhyllotaxyOffset = 0.0f;
+
+	UPROPERTY(EditAnywhere, Category="Misc Settings")
+	int32 RandomSeed = 123456;
+};
+
+class FPVFoliageDistributorElement : public FPVBaseElement
+{
+protected:
+	virtual bool ExecuteInternal(FPCGContext* InContext) const override;
+};
diff --git a/Engine/Plugins/Experimental/ProceduralVegetationEditor/Source/ProceduralVegetation/Private/Nodes/PVFoliagePaletteSettings.cpp b/Engine/Plugins/Experimental/ProceduralVegetationEditor/Source/ProceduralVegetation/Private/Nodes/PVFoliagePaletteSettings.cpp
new file mode 100644
--- /dev/null
+++ b/Engine/Plugins/Experimental/ProceduralVegetationEditor/Source/ProceduralVegetation/Private/Nodes/PVFoliagePaletteSettings.cpp
@@ -0,0 +1,119 @@
+﻿// Copyright Epic Games, Inc. All Rights Reserved.
+
+#include "PVFoliagePaletteSettings.h"
+
+#include "ProceduralVegetationModule.h"
+#include "DataTypes/PVData.h"
+#include "DataTypes/PVMeshData.h"
+#include "Data/PCGBasePointData.h"
+#include "Data/PCGStaticMeshResourceData.h"
+#include "Facades/PVFoliageFacade.h"
+#include "GeometryCollection/ManagedArrayCollection.h"
+#include "Helpers/PCGHelpers.h"
+#include "Helpers/PCGSettingsHelpers.h"
+
+#if WITH_EDITOR
+void UPVFoliagePaletteSettings::PostEditChangeProperty(FPropertyChangedEvent& PropertyChangedEvent)
+{
+	Super::PostEditChangeProperty(PropertyChangedEvent);
+
+	if (PropertyChangedEvent.Property
+		&& PropertyChangedEvent.Property->GetFName() == GET_MEMBER_NAME_CHECKED(UPVFoliagePaletteSettings, bOverrideFoliage))
+	{
+		DirtyCache();
+	}
+}
+#endif
+
+FPCGDataTypeIdentifier UPVFoliagePaletteSettings::GetInputPinTypeIdentifier() const
+{
+	return FPCGDataTypeIdentifier{ FPVDataTypeInfoMesh::AsId() };
+}
+
+FPCGDataTypeIdentifier UPVFoliagePaletteSettings::GetOutputPinTypeIdentifier() const
+{
+	return FPCGDataTypeIdentifier{ FPVDataTypeInfoMesh::AsId() };
+}
+
+FPCGElementPtr UPVFoliagePaletteSettings::CreateElement() const
+{
+	return MakeShared<FPVFoliageElement>();
+}
+
+bool FPVFoliageElement::ExecuteInternal(FPCGContext* InContext) const
+{
+	TRACE_CPUPROFILER_EVENT_SCOPE(FPVFoliageElement::Execute);
+
+	check(InContext);
+
+	const UPVFoliagePaletteSettings* Settings = InContext->GetInputSettings<UPVFoliagePaletteSettings>();
+	check(Settings);
+
+	UPVFoliagePaletteSettings* const MutableSettings = const_cast<UPVFoliagePaletteSettings*>(Settings);
+	check(MutableSettings);
+
+	const TArray<FPCGTaggedData>& Inputs = InContext->InputData.GetInputsByPin(PCGPinConstants::DefaultInputLabel);
+	if (!Inputs.IsEmpty())
+	{
+		if (const UPVMeshData* InputData = Cast<UPVMeshData>(Inputs[0].Data))
+		{
+			FManagedArrayCollection Collection = InputData->GetCollection();
+
+			UPVMeshData* OutManagedArrayCollectionData = FPCGContext::NewObject_AnyThread<UPVMeshData>(InContext);
+
+			PV::Facades::FFoliageFacade Facade(Collection);
+
+			if (!Settings->bOverrideFoliage)
+			{
+				MutableSettings->FoliageMeshes.Empty();
+				MutableSettings->bOverrideFoliage = false;
+				for (int32 FoliageIndex = 0; FoliageIndex < Facade.NumFoliageNames(); ++FoliageIndex)
+				{
+					MutableSettings->FoliageMeshes.Emplace(FSoftObjectPath(Facade.GetFoliageName(FoliageIndex)));
+				}
+#if WITH_EDITOR
+				MutableSettings->PostEditChange();
+#endif
+			}
+
+			const int32 IterationLimit = FMath::Min(Settings->FoliageMeshes.Num(), Facade.NumFoliageNames());
+			TArray<FString> FoliageNames;
+			for (int32 i = 0; i < IterationLimit; ++i)
+			{
+				FoliageNames.Emplace(Settings->FoliageMeshes[i].ToString());
+			}
+
+			if (Settings->FoliageMeshes.Num() > Facade.NumFoliageNames())
+			{
+				UE_LOG(LogProceduralVegetation, Warning,
+				       TEXT("Num of Foliage meshes specified is greater than the number of available entries, ignoring the last %d entries"),
+				       Settings->FoliageMeshes.Num() - Facade.NumFoliageNames());
+			}
+			else if (Settings->FoliageMeshes.Num() < Facade.NumFoliageNames())
+			{
+				if (Settings->FoliageMeshes.Num() > 0)
+				{
+					UE_LOG(LogProceduralVegetation, Warning,
+					       TEXT("Num of Foliage meshes specified is less than the number of available entries, repeating the last entry"));
+
+					for (int32 i = Settings->FoliageMeshes.Num(); i < Facade.NumFoliageNames(); ++i)
+					{
+						FoliageNames.Emplace(Settings->FoliageMeshes.Last().ToString());
+					}
+				}
+			}
+
+			Facade.SetFoliageNames(FoliageNames);
+
+			OutManagedArrayCollectionData->Initialize(MoveTemp(Collection));
+			InContext->OutputData.TaggedData.Emplace(OutManagedArrayCollectionData);
+		}
+		else
+		{
+			PCGLog::InputOutput::LogInvalidInputDataError(InContext);
+			return true;
+		}
+	}
+
+	return true;
+}
diff --git a/Engine/Plugins/Experimental/ProceduralVegetationEditor/Source/ProceduralVegetation/Private/Nodes/PVFoliagePaletteSettings.h b/Engine/Plugins/Experimental/ProceduralVegetationEditor/Source/ProceduralVegetation/Private/Nodes/PVFoliagePaletteSettings.h
new file mode 100644
--- /dev/null
+++ b/Engine/Plugins/Experimental/ProceduralVegetationEditor/Source/ProceduralVegetation/Private/Nodes/PVFoliagePaletteSettings.h
@@ -0,0 +1,47 @@
+﻿// Copyright Epic Games, Inc. All Rights Reserved.
+#pragma once
+
+#include "CoreMinimal.h"
+#include "Nodes/PVBaseSettings.h"
+#include "DataTypes/PVData.h"
+#include "PVFoliagePaletteSettings.generated.h"
+
+UCLASS(BlueprintType, ClassGroup = (Procedural))
+class UPVFoliagePaletteSettings : public UPVBaseSettings
+{
+	GENERATED_BODY()
+
+public:
+#if WITH_EDITOR
+	virtual FName GetDefaultNodeName() const override { return FName(TEXT("ProceduralVegetationFoliage")); }
+	virtual FText GetDefaultNodeTitle() const override { return NSLOCTEXT("PVFoliagePaletteSettings", "NodeTitle", "Foliage Palette"); }
+	virtual FText GetNodeTooltipText() const override { return NSLOCTEXT("PVFoliagePaletteSettings", "NodeTooltip", "Foliage setup for the tree"); }
+	virtual FLinearColor GetNodeTitleColor() const override { return FLinearColor::Green; }
+
+	virtual TArray<EPVRenderType> GetRenderType() const override { return TArray{EPVRenderType::Mesh, EPVRenderType::Foliage}; }
+
+	virtual void PostEditChangeProperty(FPropertyChangedEvent& PropertyChangedEvent) override;
+#endif
+
+protected:
+	virtual FPCGDataTypeIdentifier GetInputPinTypeIdentifier() const override;
+	virtual FPCGDataTypeIdentifier GetOutputPinTypeIdentifier() const override;
+	virtual FPCGElementPtr CreateElement() const override;
+
+public:
+	UPROPERTY(EditAnywhere, Category="Meshes", meta=(PinHiddenByDefault, InlineEditConditionToggle))
+	bool bOverrideFoliage = false;
+
+	UPROPERTY(EditAnywhere, Category = "Meshes", Meta = (EditCondition="bOverrideFoliage", AllowedClasses = "/Script/Engine.StaticMesh,/Script/Engine.SkeletalMesh"))
+	TArray<TSoftObjectPtr<UObject>> FoliageMeshes;
+
+};
+
+class FPVFoliageElement : public FPVBaseElement
+{
+protected:
+	virtual bool ExecuteInternal(FPCGContext* InContext) const override;
+
+public:
+	virtual bool CanExecuteOnlyOnMainThread(FPCGContext* Context) const override { return true; }
+};
diff --git a/Engine/Plugins/Experimental/ProceduralVegetationEditor/Source/ProceduralVegetation/Private/Nodes/PVGravitySettings.cpp b/Engine/Plugins/Experimental/ProceduralVegetationEditor/Source/ProceduralVegetation/Private/Nodes/PVGravitySettings.cpp
new file mode 100644
--- /dev/null
+++ b/Engine/Plugins/Experimental/ProceduralVegetationEditor/Source/ProceduralVegetation/Private/Nodes/PVGravitySettings.cpp
@@ -0,0 +1,59 @@
+﻿// Copyright Epic Games, Inc. All Rights Reserved.
+
+#include "PVGravitySettings.h"
+#include "DataTypes/PVData.h"
+#include "DataTypes/PVGrowthData.h"
+#include "Data/PCGBasePointData.h"
+#include "GeometryCollection/ManagedArrayCollection.h"
+#include "Helpers/PCGSettingsHelpers.h"
+#include "Implementations/PVGravity.h"
+
+FPCGDataTypeIdentifier UPVGravitySettings::GetInputPinTypeIdentifier() const
+{
+	return FPCGDataTypeIdentifier{ FPVDataTypeInfoGrowth::AsId() };
+}
+
+FPCGDataTypeIdentifier UPVGravitySettings::GetOutputPinTypeIdentifier() const
+{
+	return FPCGDataTypeIdentifier{ FPVDataTypeInfoGrowth::AsId() };
+}
+
+FPCGElementPtr UPVGravitySettings::CreateElement() const
+{
+	return MakeShared<FPVGravityElement>();
+}
+
+bool FPVGravityElement::ExecuteInternal(FPCGContext* InContext) const
+{
+	TRACE_CPUPROFILER_EVENT_SCOPE(FPVGravityElement::Execute);
+
+	check(InContext);
+
+	const UPVGravitySettings* Settings = InContext->GetInputSettings<UPVGravitySettings>();
+	check(Settings);
+	
+	for (const FPCGTaggedData& Input : InContext->InputData.GetInputsByPin(PCGPinConstants::DefaultInputLabel))
+	{
+		if(const UPVGrowthData* InputData = Cast<UPVGrowthData>(Input.Data))
+		{
+			FManagedArrayCollection OutCollection = InputData->GetCollection();
+		
+			UPVGrowthData* OutManagedArrayCollectionData = FPCGContext::NewObject_AnyThread<UPVGrowthData>(InContext);
+			
+			if (Settings->GravitySettings.Gravity != 0.0f)
+			{
+				FPVGravity::ApplyGravity(Settings->GravitySettings, OutCollection);
+			}
+
+			OutManagedArrayCollectionData->Initialize(MoveTemp(OutCollection));
+			InContext->OutputData.TaggedData.Emplace(OutManagedArrayCollectionData);
+		}
+		else
+		{
+			PCGLog::InputOutput::LogInvalidInputDataError(InContext);
+			return true;
+		}
+	}
+	
+	return true;
+}
diff --git a/Engine/Plugins/Experimental/ProceduralVegetationEditor/Source/ProceduralVegetation/Private/Nodes/PVGravitySettings.h b/Engine/Plugins/Experimental/ProceduralVegetationEditor/Source/ProceduralVegetation/Private/Nodes/PVGravitySettings.h
new file mode 100644
--- /dev/null
+++ b/Engine/Plugins/Experimental/ProceduralVegetationEditor/Source/ProceduralVegetation/Private/Nodes/PVGravitySettings.h
@@ -0,0 +1,39 @@
+﻿// Copyright Epic Games, Inc. All Rights Reserved.
+#pragma once
+
+#include "CoreMinimal.h"
+#include "Nodes/PVBaseSettings.h"
+#include "Implementations/PVGravity.h"
+#include "PVGravitySettings.generated.h"
+
+UCLASS(BlueprintType, ClassGroup = (Procedural))
+class UPVGravitySettings : public UPVBaseSettings
+{
+	GENERATED_BODY()
+
+public:
+#if WITH_EDITOR
+	virtual FName GetDefaultNodeName() const override { return FName(TEXT("ProceduralVegetationGravity")); }
+	virtual FText GetDefaultNodeTitle() const override { return NSLOCTEXT("PVGravitySettings", "NodeTitle", "Gravity"); }
+	virtual FText GetNodeTooltipText() const override { return NSLOCTEXT("PVGravitySettings", "NodeTooltip", "Apply gravity to branches and main trunk"); }
+	virtual FLinearColor GetNodeTitleColor() const override { return FLinearColor::Yellow; }
+	
+	virtual TArray<EPVRenderType> GetRenderType() const override { return TArray{ EPVRenderType::Skeleton }; }
+#endif
+
+protected:
+	virtual FPCGDataTypeIdentifier GetInputPinTypeIdentifier() const override;
+	virtual FPCGDataTypeIdentifier GetOutputPinTypeIdentifier() const override;
+
+	virtual FPCGElementPtr CreateElement() const override;
+
+public:
+	UPROPERTY(EditAnywhere, Category="Gravity", meta=(ShowOnlyInnerProperties))
+	FPVGravityParams GravitySettings;
+};
+
+class FPVGravityElement : public FPVBaseElement
+{
+protected:
+	virtual bool ExecuteInternal(FPCGContext* InContext) const override;
+};
diff --git a/Engine/Plugins/Experimental/ProceduralVegetationEditor/Source/ProceduralVegetation/Private/Nodes/PVImporterSettings.cpp b/Engine/Plugins/Experimental/ProceduralVegetationEditor/Source/ProceduralVegetation/Private/Nodes/PVImporterSettings.cpp
new file mode 100644
--- /dev/null
+++ b/Engine/Plugins/Experimental/ProceduralVegetationEditor/Source/ProceduralVegetation/Private/Nodes/PVImporterSettings.cpp
@@ -0,0 +1,54 @@
+﻿// Copyright Epic Games, Inc. All Rights Reserved.
+
+#include "PVImporterSettings.h"
+
+#include "ProceduralVegetationModule.h"
+#include "DataTypes/PVData.h"
+#include "DataTypes/PVGrowthData.h"
+#include "GeometryCollection/ManagedArrayCollection.h"
+#include "Helpers/PVJSONHelper.h"
+#include "Helpers/PCGHelpers.h"
+#include "Helpers/PCGSettingsHelpers.h"
+
+TArray<FPCGPinProperties> UPVImporterSettings::InputPinProperties() const
+{
+	return TArray<FPCGPinProperties>();
+}
+
+FPCGElementPtr UPVImporterSettings::CreateElement() const
+{
+	return MakeShared<FPVImporterElement>();
+}
+
+FPCGDataTypeIdentifier UPVImporterSettings::GetOutputPinTypeIdentifier() const
+{
+	return FPCGDataTypeIdentifier{ FPVDataTypeInfoGrowth::AsId() };
+}
+
+bool FPVImporterElement::ExecuteInternal(FPCGContext* InContext) const
+{
+	TRACE_CPUPROFILER_EVENT_SCOPE(FPVImporterElement::Execute);
+
+	check(InContext);
+
+	const UPVImporterSettings* Settings = InContext->GetInputSettings<UPVImporterSettings>();
+	check(Settings);
+	
+	const FString& SkeletonFilePath = Settings->SkeletonFile.FilePath;
+	FString ErrorMessage;
+
+	FManagedArrayCollection Collection;
+	bool bSuccess = PV::LoadMegaSeedJsonToCollection(Collection, SkeletonFilePath, ErrorMessage);
+	if (!bSuccess)
+	{
+		UE_LOG(LogProceduralVegetation, Error, TEXT("%s"), *ErrorMessage);
+		PCGLog::InputOutput::LogInvalidInputDataError(InContext);
+		return true;
+	}
+	UPVGrowthData* OutManagedArrayCollectionData = FPCGContext::NewObject_AnyThread<UPVGrowthData>(InContext);
+	OutManagedArrayCollectionData->Initialize(MoveTemp(Collection));
+		
+	InContext->OutputData.TaggedData.Emplace(OutManagedArrayCollectionData);
+	
+	return true;
+}
diff --git a/Engine/Plugins/Experimental/ProceduralVegetationEditor/Source/ProceduralVegetation/Private/Nodes/PVImporterSettings.h b/Engine/Plugins/Experimental/ProceduralVegetationEditor/Source/ProceduralVegetation/Private/Nodes/PVImporterSettings.h
new file mode 100644
--- /dev/null
+++ b/Engine/Plugins/Experimental/ProceduralVegetationEditor/Source/ProceduralVegetation/Private/Nodes/PVImporterSettings.h
@@ -0,0 +1,38 @@
+﻿// Copyright Epic Games, Inc. All Rights Reserved.
+#pragma once
+
+#include "CoreMinimal.h"
+#include "Nodes/PVBaseSettings.h"
+
+#include "PVImporterSettings.generated.h"
+
+UCLASS(BlueprintType, ClassGroup = (Procedural))
+class UPVImporterSettings : public UPVBaseSettings
+{
+	GENERATED_BODY()
+
+public:
+#if WITH_EDITOR
+	virtual FName GetDefaultNodeName() const override { return FName(TEXT("ProceduralVegetationImporter")); }
+	virtual FText GetDefaultNodeTitle() const override { return NSLOCTEXT("PVImporterSettings", "NodeTitle", "Procedural Vegetation Importer"); }
+	virtual FText GetNodeTooltipText() const override { return NSLOCTEXT("PVImporterSettings", "NodeTooltip", "Create Procedural Vegetation Skeleton"); }
+	virtual FLinearColor GetNodeTitleColor() const override { return FLinearColor(FColor::Black); }
+
+	virtual TArray<EPVRenderType> GetRenderType() const override { return TArray{ EPVRenderType::Skeleton }; }
+#endif
+
+protected:
+	virtual TArray<FPCGPinProperties> InputPinProperties() const override;
+	virtual FPCGDataTypeIdentifier GetOutputPinTypeIdentifier() const override;
+	virtual FPCGElementPtr CreateElement() const override;
+
+public:
+	UPROPERTY(EditAnywhere, Category="Skeleton")
+	FFilePath SkeletonFile;
+};
+
+class FPVImporterElement : public FPVBaseElement
+{
+protected:
+	virtual bool ExecuteInternal(FPCGContext* InContext) const override;
+};
diff --git a/Engine/Plugins/Experimental/ProceduralVegetationEditor/Source/ProceduralVegetation/Private/Nodes/PVMeshBuilderSettings.cpp b/Engine/Plugins/Experimental/ProceduralVegetationEditor/Source/ProceduralVegetation/Private/Nodes/PVMeshBuilderSettings.cpp
new file mode 100644
--- /dev/null
+++ b/Engine/Plugins/Experimental/ProceduralVegetationEditor/Source/ProceduralVegetation/Private/Nodes/PVMeshBuilderSettings.cpp
@@ -0,0 +1,147 @@
+﻿// Copyright Epic Games, Inc. All Rights Reserved.
+
+#include "PVMeshBuilderSettings.h"
+#include "DataTypes/PVData.h"
+#include "DataTypes/PVGrowthData.h"
+#include "DataTypes/PVMeshData.h"
+#include "Data/PCGBasePointData.h"
+#include "Facades/PVBranchFacade.h"
+#include "Facades/PVProfileFacade.h"
+#include "GeometryCollection/ManagedArrayCollection.h"
+#include "Helpers/PCGSettingsHelpers.h"
+#include "Implementations/PVMeshBuilder.h"
+#include "Materials/MaterialInterface.h"
+
+#if WITH_EDITOR
+void UPVMeshBuilderSettings::PostEditChangeProperty(FPropertyChangedEvent& PropertyChangedEvent)
+{
+	Super::PostEditChangeProperty(PropertyChangedEvent);
+
+	if (PropertyChangedEvent.Property
+		&& PropertyChangedEvent.Property->GetFName() == GET_MEMBER_NAME_CHECKED(FPVMeshBuilderParams, DisplacementTexture))
+	{
+		if (MesherSettings.DisplacementTexture)
+		{
+			FPVMeshBuilder::ExtractDisplacementData(MesherSettings.DisplacementTexture, MesherSettings.DisplacementValues);
+		}
+		else
+		{
+			MesherSettings.DisplacementValues.Empty();
+		}
+
+		Modify();
+	}
+}
+
+void UPVMeshBuilderSettings::PostLoad()
+{
+	Super::PostLoad();
+
+	if (MesherSettings.DisplacementTexture)
+	{
+		FPVMeshBuilder::ExtractDisplacementData(MesherSettings.DisplacementTexture, MesherSettings.DisplacementValues);
+	}
+	else
+	{
+		MesherSettings.DisplacementValues.Empty();
+	}
+}
+#endif
+
+FPCGDataTypeIdentifier UPVMeshBuilderSettings::GetInputPinTypeIdentifier() const
+{
+	return FPCGDataTypeIdentifier{ FPVDataTypeInfoGrowth::AsId() };
+}
+
+FPCGDataTypeIdentifier UPVMeshBuilderSettings::GetOutputPinTypeIdentifier() const
+{
+	return FPCGDataTypeIdentifier{ FPVDataTypeInfoMesh::AsId() };
+}
+
+FPCGElementPtr UPVMeshBuilderSettings::CreateElement() const
+{
+	return MakeShared<FPVMeshBuilderElement>();
+}
+
+bool FPVMeshBuilderElement::ExecuteInternal(FPCGContext* InContext) const
+{
+	TRACE_CPUPROFILER_EVENT_SCOPE(FPVMeshBuilderElement::Execute);
+
+	check(InContext);
+
+	const UPVMeshBuilderSettings* Settings = InContext->GetInputSettings<UPVMeshBuilderSettings>();
+	check(Settings);
+
+	UPVMeshBuilderSettings* MutableSettings = const_cast<UPVMeshBuilderSettings*>(Settings);
+	check(MutableSettings);
+
+	MutableSettings->MesherSettings.PlantProfileOptions.Empty();
+	MutableSettings->MesherSettings.PlantProfileOptions.Add(TEXT("None"));
+
+	const TArray<FPCGTaggedData>& Inputs = InContext->InputData.GetInputsByPin(PCGPinConstants::DefaultInputLabel);
+	if (!Inputs.IsEmpty())
+	{
+		if (const UPVGrowthData* InputData = Cast<UPVGrowthData>(Inputs[0].Data))
+		{
+			const FManagedArrayCollection& SkeletonCollection = InputData->GetCollection();
+			auto& MaterialSettings = MutableSettings->MesherSettings.MaterialSettings;
+					
+			if (!Settings->MesherSettings.bOverrideMaterial || MaterialSettings.MaterialSetups.Num() == 0)
+			{
+				const PV::Facades::FBranchFacade BranchFacade(SkeletonCollection);
+				const FString TrunkMaterial = BranchFacade.GetTrunkMaterialPath();
+				if (!TrunkMaterial.IsEmpty())
+				{
+					MaterialSettings.SetMaterial(TrunkMaterial, 0);
+				}
+				else
+				{
+					MutableSettings->MesherSettings.bOverrideMaterial = true;
+				}
+			}
+
+			if (const PV::Facades::FPlantProfileFacade PlantProfileFacade = PV::Facades::FPlantProfileFacade(SkeletonCollection);
+				PlantProfileFacade.NumProfileEntries() > 0)
+			{
+				for (int32 i = 0; i < PlantProfileFacade.NumProfileEntries(); i++)
+				{
+					MutableSettings->MesherSettings.PlantProfileOptions.Add(FString::FromInt(i));
+				}
+
+				MutableSettings->MesherSettings.bIsPlantProfileDropdownEnabled = true;
+			}
+			else
+			{
+				MutableSettings->MesherSettings.bIsPlantProfileDropdownEnabled = false;
+			}
+
+			UPVMeshData* OutManagedArrayCollectionData = FPCGContext::NewObject_AnyThread<UPVMeshData>(InContext);
+
+			FGeometryCollection OutGeometryCollection;
+			FPVMeshBuilder::GenerateGeometryCollection(SkeletonCollection, Settings->MesherSettings, OutGeometryCollection);
+
+			OutManagedArrayCollectionData->Initialize(MoveTemp(OutGeometryCollection));
+
+			FPCGTaggedData& CollectionOutput = InContext->OutputData.TaggedData.Emplace_GetRef();
+			CollectionOutput.Data = OutManagedArrayCollectionData;
+			CollectionOutput.Pin = PCGPinConstants::DefaultOutputLabel;
+
+			// TODO : Figure out if we want to keep the dynamic mesh inn the ProceduralVegetation data or as a pin
+			// UPCGDynamicMeshData* DynamicMeshData = FPCGContext::NewObject_AnyThread<UPCGDynamicMeshData>(InContext);
+			// TObjectPtr<UDynamicMesh> OutMesh = NewObject<UDynamicMesh>();
+			// FPVMeshBuilder::GenerateDynamicMesh(SkeletonCollection, OutMesh);
+			// DynamicMeshData->Initialize(OutMesh, true);
+			//
+			// FPCGTaggedData& MeshOutput = InContext->OutputData.TaggedData.Emplace_GetRef();
+			// MeshOutput.Data = DynamicMeshData;
+			// MeshOutput.Pin = TEXT("Mesh");
+		}
+		else
+		{
+			PCGLog::InputOutput::LogInvalidInputDataError(InContext);
+			return true;
+		}
+	}
+
+	return true;
+}
diff --git a/Engine/Plugins/Experimental/ProceduralVegetationEditor/Source/ProceduralVegetation/Private/Nodes/PVMeshBuilderSettings.h b/Engine/Plugins/Experimental/ProceduralVegetationEditor/Source/ProceduralVegetation/Private/Nodes/PVMeshBuilderSettings.h
new file mode 100644
--- /dev/null
+++ b/Engine/Plugins/Experimental/ProceduralVegetationEditor/Source/ProceduralVegetation/Private/Nodes/PVMeshBuilderSettings.h
@@ -0,0 +1,50 @@
+﻿// Copyright Epic Games, Inc. All Rights Reserved.
+#pragma once
+
+#include "CoreMinimal.h"
+#include "Implementations/PVMeshBuilder.h"
+#include "Nodes/PVBaseSettings.h"
+#include "PVMeshBuilderSettings.generated.h"
+
+class UMaterialInterface;
+
+UCLASS(BlueprintType, ClassGroup = (Procedural))
+class UPVMeshBuilderSettings : public UPVBaseSettings
+{
+	GENERATED_BODY()
+
+public:
+#if WITH_EDITOR
+	virtual FName GetDefaultNodeName() const override { return FName(TEXT("ProceduralVegetationMesher")); }
+	virtual FText GetDefaultNodeTitle() const override { return NSLOCTEXT("PVMeshBuilderSettings", "NodeTitle", "Mesh Builder"); }
+	virtual FText GetNodeTooltipText() const override { return NSLOCTEXT("PVMeshBuilderSettings", "NodeTooltip", "Create a mesh for the skeleton"); }
+	virtual FLinearColor GetNodeTitleColor() const override { return FLinearColor(FColor::Cyan); }
+	virtual TArray<EPVRenderType> GetRenderType() const override { return TArray{EPVRenderType::Mesh}; }
+
+	virtual void PostEditChangeProperty(FPropertyChangedEvent& PropertyChangedEvent) override;
+
+	virtual void PostLoad() override;
+#endif
+
+protected:
+	virtual FPCGElementPtr CreateElement() const override;
+
+	virtual FPCGDataTypeIdentifier GetInputPinTypeIdentifier() const override;
+	virtual FPCGDataTypeIdentifier GetOutputPinTypeIdentifier() const override;
+
+public:	
+	UPROPERTY(EditAnywhere, Category = "Mesh", meta = (ShowOnlyInnerProperties))
+	FPVMeshBuilderParams MesherSettings;
+
+	UFUNCTION()
+	TArray<FString> GetPlantProfileOptions() const
+	{
+		return MesherSettings.PlantProfileOptions;
+	}
+};
+
+class FPVMeshBuilderElement : public FPVBaseElement
+{
+protected:
+	virtual bool ExecuteInternal(FPCGContext* InContext) const override;
+};
diff --git a/Engine/Plugins/Experimental/ProceduralVegetationEditor/Source/ProceduralVegetation/Private/Nodes/PVOutputSettings.cpp b/Engine/Plugins/Experimental/ProceduralVegetationEditor/Source/ProceduralVegetation/Private/Nodes/PVOutputSettings.cpp
new file mode 100644
--- /dev/null
+++ b/Engine/Plugins/Experimental/ProceduralVegetationEditor/Source/ProceduralVegetation/Private/Nodes/PVOutputSettings.cpp
@@ -0,0 +1,105 @@
+﻿// Copyright Epic Games, Inc. All Rights Reserved.
+
+#include "Nodes/PVOutputSettings.h"
+#include "PCGContext.h"
+#include "DataTypes/PVMeshData.h"
+#include "DataTypes/PVFoliageMeshData.h"
+#include "Helpers/PCGHelpers.h"
+#include "Misc/PackageName.h"
+
+#define LOCTEXT_NAMESPACE "PVOutputSettings"
+
+FString FormatLongPackageNameErrorCode(const FPackageName::EErrorCode ErrorCode)
+{
+	switch (ErrorCode)
+	{
+	case FPackageName::EErrorCode::PackageNameUnknown:
+		return "Unknown Export Path error";
+	case FPackageName::EErrorCode::PackageNameEmptyPath:
+		return "Empty Export Path";
+	case FPackageName::EErrorCode::PackageNamePathNotMounted:
+		return "Mount Point for Export Path is not valid";
+	case FPackageName::EErrorCode::PackageNamePathIsMemoryOnly:
+		return "Export Path exists in memory only. ";
+	case FPackageName::EErrorCode::PackageNameSpacesNotAllowed:
+		return "Spaces are not allowed in Export Path";	
+	case FPackageName::EErrorCode::PackageNameContainsInvalidCharacters:
+		return FString::Printf(TEXT("Export Path contains one of \n%s\n Invalid Characters"), TEXT(R"(\\:*?\"<>|' ,.&!~\n\r\t@#)"));
+	case FPackageName::EErrorCode::LongPackageNames_PathTooShort:
+		return "Export Path is too small to be a valid path";
+	case FPackageName::EErrorCode::LongPackageNames_PathWithNoStartingSlash:
+		return "Export Path has to start with a slash [/]";
+	case FPackageName::EErrorCode::LongPackageNames_PathWithTrailingSlash:
+		return "Export Path cannot end with a slash [/]";
+	case FPackageName::EErrorCode::LongPackageNames_PathWithDoubleSlash:
+		return "Export Path contains a double slash [/]";
+	}
+	return FString();
+}
+
+UPVOutputSettings::UPVOutputSettings()
+{
+	if (PCGHelpers::IsNewObjectAndNotDefault(this))
+	{
+		static FString DefaultMeshName = "ProceduralVegetationOutput";
+		static FString DefaultAssetPath = "/Game/ProceduralVegetation/";
+		ExportSettings.Initialize(DefaultAssetPath, DefaultMeshName);
+	}
+}
+
+FPCGDataTypeIdentifier UPVOutputSettings::GetInputPinTypeIdentifier() const
+{
+	return FPCGDataTypeIdentifier::Construct(FPVDataTypeInfoMesh::AsId(), FPVDataTypeInfoFoliageMesh::AsId());
+}
+
+TArray<FPCGPinProperties> UPVOutputSettings::OutputPinProperties() const
+{
+	TArray<FPCGPinProperties> Properties;
+
+	FPCGPinProperties& Pin = Properties.Emplace_GetRef(PCGPinConstants::DefaultOutputLabel, EPCGDataType::Any);
+	Pin.SetRequiredPin();
+	Pin.SetAllowMultipleConnections(false);
+	Pin.bInvisiblePin = true;
+
+	return Properties;
+}
+
+FPCGElementPtr UPVOutputSettings::CreateElement() const
+{
+	return MakeShared<FPVOutputElement>();
+}
+
+bool FPVOutputElement::ExecuteInternal(FPCGContext* InContext) const
+{
+	const UPVOutputSettings* Settings = InContext->GetInputSettings<UPVOutputSettings>();
+	check(Settings);
+
+	for (const FPCGTaggedData& Input : InContext->InputData.GetInputsByPin(PCGPinConstants::DefaultInputLabel))
+	{
+		if (const UPVData* InputData = Cast<UPVData>(Input.Data))
+		{
+			const FString OutputPath = Settings->ExportSettings.ContentBrowserFolder.Path;
+			if (FPackageName::EErrorCode ErrorCode; !FPackageName::IsValidLongPackageName(OutputPath, false, &ErrorCode))
+			{
+				PCGLog::LogErrorOnGraph(FText::FromString(FormatLongPackageNameErrorCode(ErrorCode)), InContext);
+				return true;
+			}
+
+			FManagedArrayCollection OutCollection = InputData->GetCollection();
+
+			UPVData* ProceduralVegetationData = FPCGContext::NewObject_AnyThread<UPVData>(InContext);
+			ProceduralVegetationData->Initialize(MoveTemp(OutCollection));
+
+			InContext->OutputData.TaggedData.Emplace(ProceduralVegetationData);
+		}
+		else
+		{
+			PCGLog::InputOutput::LogInvalidInputDataError(InContext);
+			return true;
+		}
+	}
+
+	return true;
+}
+
+#undef LOCTEXT_NAMESPACE
diff --git a/Engine/Plugins/Experimental/ProceduralVegetationEditor/Source/ProceduralVegetation/Private/Nodes/PVPCGBaseSettings.cpp b/Engine/Plugins/Experimental/ProceduralVegetationEditor/Source/ProceduralVegetation/Private/Nodes/PVPCGBaseSettings.cpp
deleted file mode 100644
diff --git a/Engine/Plugins/Experimental/ProceduralVegetationEditor/Source/ProceduralVegetation/Private/Nodes/PVPCGBoneReduction.cpp b/Engine/Plugins/Experimental/ProceduralVegetationEditor/Source/ProceduralVegetation/Private/Nodes/PVPCGBoneReduction.cpp
deleted file mode 100644
diff --git a/Engine/Plugins/Experimental/ProceduralVegetationEditor/Source/ProceduralVegetation/Private/Nodes/PVPCGBoneReduction.h b/Engine/Plugins/Experimental/ProceduralVegetationEditor/Source/ProceduralVegetation/Private/Nodes/PVPCGBoneReduction.h
deleted file mode 100644
diff --git a/Engine/Plugins/Experimental/ProceduralVegetationEditor/Source/ProceduralVegetation/Private/Nodes/PVPCGCarve.cpp b/Engine/Plugins/Experimental/ProceduralVegetationEditor/Source/ProceduralVegetation/Private/Nodes/PVPCGCarve.cpp
deleted file mode 100644
diff --git a/Engine/Plugins/Experimental/ProceduralVegetationEditor/Source/ProceduralVegetation/Private/Nodes/PVPCGCarve.h b/Engine/Plugins/Experimental/ProceduralVegetationEditor/Source/ProceduralVegetation/Private/Nodes/PVPCGCarve.h
deleted file mode 100644
diff --git a/Engine/Plugins/Experimental/ProceduralVegetationEditor/Source/ProceduralVegetation/Private/Nodes/PVPCGFoliage.cpp b/Engine/Plugins/Experimental/ProceduralVegetationEditor/Source/ProceduralVegetation/Private/Nodes/PVPCGFoliage.cpp
deleted file mode 100644
diff --git a/Engine/Plugins/Experimental/ProceduralVegetationEditor/Source/ProceduralVegetation/Private/Nodes/PVPCGFoliage.h b/Engine/Plugins/Experimental/ProceduralVegetationEditor/Source/ProceduralVegetation/Private/Nodes/PVPCGFoliage.h
deleted file mode 100644
diff --git a/Engine/Plugins/Experimental/ProceduralVegetationEditor/Source/ProceduralVegetation/Private/Nodes/PVPCGFoliageDistributor.cpp b/Engine/Plugins/Experimental/ProceduralVegetationEditor/Source/ProceduralVegetation/Private/Nodes/PVPCGFoliageDistributor.cpp
deleted file mode 100644
diff --git a/Engine/Plugins/Experimental/ProceduralVegetationEditor/Source/ProceduralVegetation/Private/Nodes/PVPCGFoliageDistributor.h b/Engine/Plugins/Experimental/ProceduralVegetationEditor/Source/ProceduralVegetation/Private/Nodes/PVPCGFoliageDistributor.h
deleted file mode 100644
diff --git a/Engine/Plugins/Experimental/ProceduralVegetationEditor/Source/ProceduralVegetation/Private/Nodes/PVPCGGravity.cpp b/Engine/Plugins/Experimental/ProceduralVegetationEditor/Source/ProceduralVegetation/Private/Nodes/PVPCGGravity.cpp
deleted file mode 100644
diff --git a/Engine/Plugins/Experimental/ProceduralVegetationEditor/Source/ProceduralVegetation/Private/Nodes/PVPCGGravity.h b/Engine/Plugins/Experimental/ProceduralVegetationEditor/Source/ProceduralVegetation/Private/Nodes/PVPCGGravity.h
deleted file mode 100644
diff --git a/Engine/Plugins/Experimental/ProceduralVegetationEditor/Source/ProceduralVegetation/Private/Nodes/PVPCGImporter.cpp b/Engine/Plugins/Experimental/ProceduralVegetationEditor/Source/ProceduralVegetation/Private/Nodes/PVPCGImporter.cpp
deleted file mode 100644
diff --git a/Engine/Plugins/Experimental/ProceduralVegetationEditor/Source/ProceduralVegetation/Private/Nodes/PVPCGImporter.h b/Engine/Plugins/Experimental/ProceduralVegetationEditor/Source/ProceduralVegetation/Private/Nodes/PVPCGImporter.h
deleted file mode 100644
diff --git a/Engine/Plugins/Experimental/ProceduralVegetationEditor/Source/ProceduralVegetation/Private/Nodes/PVPCGMegaSeedLoad.cpp b/Engine/Plugins/Experimental/ProceduralVegetationEditor/Source/ProceduralVegetation/Private/Nodes/PVPCGMegaSeedLoad.cpp
deleted file mode 100644
diff --git a/Engine/Plugins/Experimental/ProceduralVegetationEditor/Source/ProceduralVegetation/Private/Nodes/PVPCGMegaSeedLoad.h b/Engine/Plugins/Experimental/ProceduralVegetationEditor/Source/ProceduralVegetation/Private/Nodes/PVPCGMegaSeedLoad.h
deleted file mode 100644
diff --git a/Engine/Plugins/Experimental/ProceduralVegetationEditor/Source/ProceduralVegetation/Private/Nodes/PVPCGMesher.cpp b/Engine/Plugins/Experimental/ProceduralVegetationEditor/Source/ProceduralVegetation/Private/Nodes/PVPCGMesher.cpp
deleted file mode 100644
diff --git a/Engine/Plugins/Experimental/ProceduralVegetationEditor/Source/ProceduralVegetation/Private/Nodes/PVPCGMesher.h b/Engine/Plugins/Experimental/ProceduralVegetationEditor/Source/ProceduralVegetation/Private/Nodes/PVPCGMesher.h
deleted file mode 100644
diff --git a/Engine/Plugins/Experimental/ProceduralVegetationEditor/Source/ProceduralVegetation/Private/Nodes/PVPCGOutput.cpp b/Engine/Plugins/Experimental/ProceduralVegetationEditor/Source/ProceduralVegetation/Private/Nodes/PVPCGOutput.cpp
deleted file mode 100644
diff --git a/Engine/Plugins/Experimental/ProceduralVegetationEditor/Source/ProceduralVegetation/Private/Nodes/PVPCGRemoveBranches.cpp b/Engine/Plugins/Experimental/ProceduralVegetationEditor/Source/ProceduralVegetation/Private/Nodes/PVPCGRemoveBranches.cpp
deleted file mode 100644
diff --git a/Engine/Plugins/Experimental/ProceduralVegetationEditor/Source/ProceduralVegetation/Private/Nodes/PVPCGRemoveBranches.h b/Engine/Plugins/Experimental/ProceduralVegetationEditor/Source/ProceduralVegetation/Private/Nodes/PVPCGRemoveBranches.h
deleted file mode 100644
diff --git a/Engine/Plugins/Experimental/ProceduralVegetationEditor/Source/ProceduralVegetation/Private/Nodes/PVPCGScale.cpp b/Engine/Plugins/Experimental/ProceduralVegetationEditor/Source/ProceduralVegetation/Private/Nodes/PVPCGScale.cpp
deleted file mode 100644
diff --git a/Engine/Plugins/Experimental/ProceduralVegetationEditor/Source/ProceduralVegetation/Private/Nodes/PVPCGScale.h b/Engine/Plugins/Experimental/ProceduralVegetationEditor/Source/ProceduralVegetation/Private/Nodes/PVPCGScale.h
deleted file mode 100644
diff --git a/Engine/Plugins/Experimental/ProceduralVegetationEditor/Source/ProceduralVegetation/Private/Nodes/PVPresetLoaderSettings.cpp b/Engine/Plugins/Experimental/ProceduralVegetationEditor/Source/ProceduralVegetation/Private/Nodes/PVPresetLoaderSettings.cpp
new file mode 100644
--- /dev/null
+++ b/Engine/Plugins/Experimental/ProceduralVegetationEditor/Source/ProceduralVegetation/Private/Nodes/PVPresetLoaderSettings.cpp
@@ -0,0 +1,64 @@
+// Copyright Epic Games, Inc. All Rights Reserved.
+
+#include "PVPresetLoaderSettings.h"
+#include "PCGContext.h"
+#include "DataTypes/PVGrowthData.h"
+#include "MegaSeedDataAsset.h"
+
+TArray<FPCGPinProperties> UPVPresetLoaderSettings::InputPinProperties() const
+{
+	return TArray<FPCGPinProperties>();
+}
+
+TArray<FPCGPinProperties> UPVPresetLoaderSettings::OutputPinProperties() const
+{
+	TArray<FPCGPinProperties> Properties;
+	if (PresetDataAsset)
+	{
+		for (const auto& [VariantName, VariantData] : PresetDataAsset->Variants)
+		{
+			FPCGPinProperties& Pin = Properties.Emplace_GetRef(*VariantName, GetOutputPinTypeIdentifier());
+			Pin.bAllowMultipleData = false;
+		}
+	}
+
+	return Properties;
+}
+
+FPCGDataTypeIdentifier UPVPresetLoaderSettings::GetOutputPinTypeIdentifier() const
+{
+	return FPCGDataTypeIdentifier{ FPVDataTypeInfoGrowth::AsId() };
+}
+
+FPCGElementPtr UPVPresetLoaderSettings::CreateElement() const
+{
+	return MakeShared<FPVPresetLoaderElement>();
+}
+
+bool FPVPresetLoaderElement::ExecuteInternal(FPCGContext* InContext) const
+{
+	TRACE_CPUPROFILER_EVENT_SCOPE(FPVPresetLoaderElement::ExecuteInternal);
+
+	check(InContext);
+
+	const UPVPresetLoaderSettings* Settings = InContext->GetInputSettings<UPVPresetLoaderSettings>();
+	check(Settings);
+
+	if (!Settings->PresetDataAsset)
+	{
+		PCGLog::LogWarningOnGraph(NSLOCTEXT("PVPresetLoaderSettings", "MissingDataAsset", "Preset data asset is not set"), InContext);
+		return true;
+	}
+
+	for (const auto& [VariantName, VariantData] : Settings->PresetDataAsset->Variants)
+	{
+		UPVGrowthData* OutVariantData = FPCGContext::NewObject_AnyThread<UPVGrowthData>(InContext);
+		OutVariantData->Initialize(CopyTemp(VariantData));
+
+		FPCGTaggedData& CollectionOutput = InContext->OutputData.TaggedData.Emplace_GetRef();
+		CollectionOutput.Data = OutVariantData;
+		CollectionOutput.Pin = *VariantName;
+	}
+
+	return true;
+}
diff --git a/Engine/Plugins/Experimental/ProceduralVegetationEditor/Source/ProceduralVegetation/Private/Nodes/PVPresetLoaderSettings.h b/Engine/Plugins/Experimental/ProceduralVegetationEditor/Source/ProceduralVegetation/Private/Nodes/PVPresetLoaderSettings.h
new file mode 100644
--- /dev/null
+++ b/Engine/Plugins/Experimental/ProceduralVegetationEditor/Source/ProceduralVegetation/Private/Nodes/PVPresetLoaderSettings.h
@@ -0,0 +1,38 @@
+// Copyright Epic Games, Inc. All Rights Reserved.
+#pragma once
+
+#include "CoreMinimal.h"
+#include "Nodes/PVBaseSettings.h"
+
+#include "PVPresetLoaderSettings.generated.h"
+
+UCLASS(BlueprintType, ClassGroup = (Procedural))
+class UPVPresetLoaderSettings : public UPVBaseSettings
+{
+	GENERATED_BODY()
+
+public:
+#if WITH_EDITOR
+	virtual FName GetDefaultNodeName() const override { return FName(TEXT("ProceduralVegetationPresetLoader")); }
+	virtual FText GetDefaultNodeTitle() const override { return NSLOCTEXT("PVPresetLoaderSettings", "NodeTitle", "Growth Data Loader"); }
+	virtual FText GetNodeTooltipText() const override { return NSLOCTEXT("PVPresetLoaderSettings", "NodeTooltip", "Load Procedural Vegetation Growth Data"); }
+
+	virtual TArray<EPVRenderType> GetRenderType() const override { return TArray{ EPVRenderType::Skeleton }; }
+#endif
+
+protected:
+	virtual TArray<FPCGPinProperties> InputPinProperties() const override;
+	virtual TArray<FPCGPinProperties> OutputPinProperties() const override;
+	virtual FPCGDataTypeIdentifier GetOutputPinTypeIdentifier() const override;
+	virtual FPCGElementPtr CreateElement() const override;
+
+public:
+	UPROPERTY(EditAnywhere, Category="Preset Data Asset")
+	TObjectPtr<class UMegaSeedDataAsset> PresetDataAsset = nullptr;
+};
+
+class FPVPresetLoaderElement : public FPVBaseElement
+{
+protected:
+	virtual bool ExecuteInternal(FPCGContext* InContext) const override;
+};
diff --git a/Engine/Plugins/Experimental/ProceduralVegetationEditor/Source/ProceduralVegetation/Private/Nodes/PVRemoveBranchesSettings.cpp b/Engine/Plugins/Experimental/ProceduralVegetationEditor/Source/ProceduralVegetation/Private/Nodes/PVRemoveBranchesSettings.cpp
new file mode 100644
--- /dev/null
+++ b/Engine/Plugins/Experimental/ProceduralVegetationEditor/Source/ProceduralVegetation/Private/Nodes/PVRemoveBranchesSettings.cpp
@@ -0,0 +1,56 @@
+﻿// Copyright Epic Games, Inc. All Rights Reserved.
+
+#include "PVRemoveBranchesSettings.h"
+#include "DataTypes/PVData.h"
+#include "DataTypes/PVGrowthData.h"
+#include "Data/PCGBasePointData.h"
+#include "GeometryCollection/ManagedArrayCollection.h"
+#include "Helpers/PCGSettingsHelpers.h"
+#include "Implementations/PVRemoveBranches.h"
+
+FPCGDataTypeIdentifier UPVRemoveBranchesSettings::GetInputPinTypeIdentifier() const
+{
+	return FPCGDataTypeIdentifier{ FPVDataTypeInfoGrowth::AsId() };
+}
+
+FPCGDataTypeIdentifier UPVRemoveBranchesSettings::GetOutputPinTypeIdentifier() const
+{
+	return FPCGDataTypeIdentifier{ FPVDataTypeInfoGrowth::AsId() };
+}
+
+FPCGElementPtr UPVRemoveBranchesSettings::CreateElement() const
+{
+	return MakeShared<FPVRemoveBranchesElement>();
+}
+
+bool FPVRemoveBranchesElement::ExecuteInternal(FPCGContext* InContext) const
+{
+	TRACE_CPUPROFILER_EVENT_SCOPE(FPVRemoveBranchesElement::Execute);
+
+	check(InContext);
+
+	const UPVRemoveBranchesSettings* Settings = InContext->GetInputSettings<UPVRemoveBranchesSettings>();
+	check(Settings);
+	
+	for (const FPCGTaggedData& Input : InContext->InputData.GetInputsByPin(PCGPinConstants::DefaultInputLabel))
+	{
+		if(const UPVGrowthData* InputData = Cast<UPVGrowthData>(Input.Data))
+		{
+			FManagedArrayCollection OutCollection = InputData->GetCollection();
+		
+			UPVGrowthData* OutManagedArrayCollectionData = FPCGContext::NewObject_AnyThread<UPVGrowthData>(InContext);
+			
+			FPVRemoveBranches::ApplyRemoveBranches(Settings->BranchRemoveBasis, Settings->Threshold, OutCollection);
+			
+			OutManagedArrayCollectionData->Initialize(MoveTemp(OutCollection));
+			InContext->OutputData.TaggedData.Emplace(OutManagedArrayCollectionData);
+		}
+		else
+		{
+			PCGLog::InputOutput::LogInvalidInputDataError(InContext);
+			return true;
+		}
+	}
+	
+	return true;
+}
diff --git a/Engine/Plugins/Experimental/ProceduralVegetationEditor/Source/ProceduralVegetation/Private/Nodes/PVRemoveBranchesSettings.h b/Engine/Plugins/Experimental/ProceduralVegetationEditor/Source/ProceduralVegetation/Private/Nodes/PVRemoveBranchesSettings.h
new file mode 100644
--- /dev/null
+++ b/Engine/Plugins/Experimental/ProceduralVegetationEditor/Source/ProceduralVegetation/Private/Nodes/PVRemoveBranchesSettings.h
@@ -0,0 +1,42 @@
+﻿// Copyright Epic Games, Inc. All Rights Reserved.
+#pragma once
+
+#include "CoreMinimal.h"
+#include "Nodes/PVBaseSettings.h"
+#include "Implementations/PVRemoveBranches.h"
+#include "PVRemoveBranchesSettings.generated.h"
+
+UCLASS(BlueprintType, ClassGroup = (Procedural))
+class UPVRemoveBranchesSettings : public UPVBaseSettings
+{
+	GENERATED_BODY()
+
+public:
+#if WITH_EDITOR
+	virtual FName GetDefaultNodeName() const override { return FName(TEXT("ProceduralVegetationRemoveBranches")); }
+	virtual FText GetDefaultNodeTitle() const override { return NSLOCTEXT("PVRemoveBranchesSettings", "NodeTitle", "Remove Branches"); }
+	virtual FText GetNodeTooltipText() const override { return NSLOCTEXT("PVRemoveBranchesSettings", "NodeTooltip", "Remove branches on the basis of selected criteria"); }
+	virtual FLinearColor GetNodeTitleColor() const override { return FLinearColor::Yellow; }
+
+	virtual TArray<EPVRenderType> GetRenderType() const override { return TArray{ EPVRenderType::Skeleton }; }
+#endif
+
+protected:
+	virtual FPCGDataTypeIdentifier GetInputPinTypeIdentifier() const override;
+	virtual FPCGDataTypeIdentifier GetOutputPinTypeIdentifier() const override;
+
+	virtual FPCGElementPtr CreateElement() const override;
+
+public:
+	UPROPERTY(EditAnywhere, Category = Settings)
+	ERemoveBranchesBasis BranchRemoveBasis = ERemoveBranchesBasis::Length;
+
+	UPROPERTY(EditAnywhere, Category = Settings, meta=(ClampMin = 0f, ClampMax = 10.0f, UIMin = 0f, UIMax = 1.0f))
+	float Threshold = 0;
+};
+
+class FPVRemoveBranchesElement : public FPVBaseElement
+{
+protected:
+	virtual bool ExecuteInternal(FPCGContext* InContext) const override;
+};
diff --git a/Engine/Plugins/Experimental/ProceduralVegetationEditor/Source/ProceduralVegetation/Private/Nodes/PVScaleSettings.cpp b/Engine/Plugins/Experimental/ProceduralVegetationEditor/Source/ProceduralVegetation/Private/Nodes/PVScaleSettings.cpp
new file mode 100644
--- /dev/null
+++ b/Engine/Plugins/Experimental/ProceduralVegetationEditor/Source/ProceduralVegetation/Private/Nodes/PVScaleSettings.cpp
@@ -0,0 +1,56 @@
+﻿// Copyright Epic Games, Inc. All Rights Reserved.
+
+#include "PVScaleSettings.h"
+#include "DataTypes/PVData.h"
+#include "DataTypes/PVGrowthData.h"
+#include "Data/PCGBasePointData.h"
+#include "GeometryCollection/ManagedArrayCollection.h"
+#include "Helpers/PCGSettingsHelpers.h"
+#include "Implementations/PVScale.h"
+
+FPCGDataTypeIdentifier UPVScaleSettings::GetInputPinTypeIdentifier() const
+{
+	return FPCGDataTypeIdentifier{ FPVDataTypeInfoGrowth::AsId() };
+}
+
+FPCGDataTypeIdentifier UPVScaleSettings::GetOutputPinTypeIdentifier() const
+{
+	return FPCGDataTypeIdentifier{ FPVDataTypeInfoGrowth::AsId() };
+}
+
+FPCGElementPtr UPVScaleSettings::CreateElement() const
+{
+	return MakeShared<FPVScaleElement>();
+}
+
+bool FPVScaleElement::ExecuteInternal(FPCGContext* InContext) const
+{
+	TRACE_CPUPROFILER_EVENT_SCOPE(FPVScaleElement::Execute);
+
+	check(InContext);
+
+	const UPVScaleSettings* Settings = InContext->GetInputSettings<UPVScaleSettings>();
+	check(Settings);
+	
+	for (const FPCGTaggedData& Input : InContext->InputData.GetInputsByPin(PCGPinConstants::DefaultInputLabel))
+	{
+		if(const UPVGrowthData* InputData = Cast<UPVGrowthData>(Input.Data))
+		{
+			FManagedArrayCollection OutCollection = InputData->GetCollection();
+		
+			UPVGrowthData* OutManagedArrayCollectionData = FPCGContext::NewObject_AnyThread<UPVGrowthData>(InContext);
+			
+			FPVScale::ApplyScale(Settings->Scale, OutCollection);
+			
+			OutManagedArrayCollectionData->Initialize(MoveTemp(OutCollection));
+			InContext->OutputData.TaggedData.Emplace(OutManagedArrayCollectionData);
+		}
+		else
+		{
+			PCGLog::InputOutput::LogInvalidInputDataError(InContext);
+			return true;
+		}
+	}
+	
+	return true;
+}
diff --git a/Engine/Plugins/Experimental/ProceduralVegetationEditor/Source/ProceduralVegetation/Private/Nodes/PVScaleSettings.h b/Engine/Plugins/Experimental/ProceduralVegetationEditor/Source/ProceduralVegetation/Private/Nodes/PVScaleSettings.h
new file mode 100644
--- /dev/null
+++ b/Engine/Plugins/Experimental/ProceduralVegetationEditor/Source/ProceduralVegetation/Private/Nodes/PVScaleSettings.h
@@ -0,0 +1,39 @@
+﻿// Copyright Epic Games, Inc. All Rights Reserved.
+#pragma once
+
+#include "CoreMinimal.h"
+#include "Nodes/PVBaseSettings.h"
+#include "PVScaleSettings.generated.h"
+
+UCLASS(BlueprintType, ClassGroup = (Procedural))
+class UPVScaleSettings : public UPVBaseSettings
+{
+	GENERATED_BODY()
+
+public:
+#if WITH_EDITOR
+	virtual FName GetDefaultNodeName() const override { return FName(TEXT("ProceduralVegetationScale")); }
+	virtual FText GetDefaultNodeTitle() const override { return NSLOCTEXT("PVScaleSettings", "NodeTitle", "Scale"); }
+	virtual FText GetNodeTooltipText() const override { return NSLOCTEXT("PVScaleSettings", "NodeTooltip", "Scale the tree based on the settings"); }
+	virtual FLinearColor GetNodeTitleColor() const override { return FLinearColor::Yellow; }
+
+	virtual TArray<EPVRenderType> GetRenderType() const override { return TArray{ EPVRenderType::Skeleton }; }
+#endif
+
+protected:
+	virtual FPCGDataTypeIdentifier GetInputPinTypeIdentifier() const override;
+	virtual FPCGDataTypeIdentifier GetOutputPinTypeIdentifier() const override;
+
+	virtual FPCGElementPtr CreateElement() const override;
+
+public:
+	UPROPERTY(EditAnywhere, Category = Settings, meta=(ClampMin = 0.01f, ClampMax = 100.0f, UIMin = 0.01f, UIMax = 10.0f))
+	float Scale = 1;
+
+};
+
+class FPVScaleElement : public FPVBaseElement
+{
+protected:
+	virtual bool ExecuteInternal(FPCGContext* InContext) const override;
+};
diff --git a/Engine/Plugins/Experimental/ProceduralVegetationEditor/Source/ProceduralVegetation/Private/PVPCGData.cpp b/Engine/Plugins/Experimental/ProceduralVegetationEditor/Source/ProceduralVegetation/Private/PVPCGData.cpp
deleted file mode 100644
diff --git a/Engine/Plugins/Experimental/ProceduralVegetationEditor/Source/ProceduralVegetation/Private/ProceduralVegetationLink.cpp b/Engine/Plugins/Experimental/ProceduralVegetationEditor/Source/ProceduralVegetation/Private/ProceduralVegetationLink.cpp
--- a/Engine/Plugins/Experimental/ProceduralVegetationEditor/Source/ProceduralVegetation/Private/ProceduralVegetationLink.cpp
+++ b/Engine/Plugins/Experimental/ProceduralVegetationEditor/Source/ProceduralVegetation/Private/ProceduralVegetationLink.cpp
@@ -1,11 +1,12 @@
 // Copyright Epic Games, Inc. All Rights Reserved.
 
 #include "ProceduralVegetationLink.h"
-#include "Nodes/PVPCGFoliage.h"
-#include "Nodes/PVPCGGravity.h"
-#include "Nodes/PVPCGMesher.h"
-#include "Nodes/PVPCGMegaSeedLoad.h"
-#include "Nodes/PVPCGOutput.h"
+#include "Nodes/PVFoliagePaletteSettings.h"
+#include "Nodes/PVGravitySettings.h"
+#include "Nodes/PVMeshBuilderSettings.h"
+#include "Nodes/PVPresetLoaderSettings.h"
+#include "Nodes/PVOutputSettings.h"
+#include "MegaSeedDataAsset.h"
 
 void SanitizeFolderName(FString& FolderName)
 {
@@ -39,41 +40,41 @@ void UProceduralVegetation::CreateGraphFromMegaSeed(const TObjectPtr<UMegaSeedDa
 
 	UPCGSettings* DefaultNodeSettings = nullptr;
 
-	UPVPCGMegaSeedLoadSettings* MegaSeedLoadNodeSettings = NewObject<UPVPCGMegaSeedLoadSettings>(Graph);
-	MegaSeedLoadNodeSettings->MegaSeedDataAsset = MegaSeedAsset;
-	UPCGNode* MegaSeedLoadNode = Graph->AddNodeCopy(MegaSeedLoadNodeSettings, DefaultNodeSettings);
+	UPVPresetLoaderSettings* PresetLoaderSettings = NewObject<UPVPresetLoaderSettings>(Graph);
+	PresetLoaderSettings->PresetDataAsset = MegaSeedAsset;
+	UPCGNode* PresetLoaderNode = Graph->AddNodeCopy(PresetLoaderSettings, DefaultNodeSettings);
 
 	FInt32Vector2 Position(0, 300);
 
-	MegaSeedLoadNode->SetNodePosition(Position.X, Position.Y);
+	PresetLoaderNode->SetNodePosition(Position.X, Position.Y);
 	Position.X += 100;
 
 	for (const auto& [VariantName, VariantData] : MegaSeedAsset->Variants)
 	{
-		const UPVPCGGravitySettings* GravityNodeSettings = NewObject<UPVPCGGravitySettings>(Graph);
-		const UPVPCGMesherSettings* MesherNodeSettings = NewObject<UPVPCGMesherSettings>(Graph);
-		const UPVPCGFoliageSettings* FoliageNodeSettings = NewObject<UPVPCGFoliageSettings>(Graph);
-		UPVPCGOutputSettings* OutputNodeSettings = NewObject<UPVPCGOutputSettings>(Graph);
+		const UPVGravitySettings* GravitySettings = NewObject<UPVGravitySettings>(Graph);
+		const UPVMeshBuilderSettings* MeshBuilderSettings = NewObject<UPVMeshBuilderSettings>(Graph);
+		const UPVFoliagePaletteSettings* FoliagePaletteSettings = NewObject<UPVFoliagePaletteSettings>(Graph);
+		UPVOutputSettings* OutputSettings = NewObject<UPVOutputSettings>(Graph);
 		FString MeshName = *VariantName;
 		SanitizeFolderName(MeshName);
-		OutputNodeSettings->ExportSettings.MeshName = *MeshName;
+		OutputSettings->ExportSettings.MeshName = *MeshName;
 
-		UPCGNode* GravityNode = Graph->AddNodeCopy(GravityNodeSettings, DefaultNodeSettings);
-		UPCGNode* MesherNode = Graph->AddNodeCopy(MesherNodeSettings, DefaultNodeSettings);
-		UPCGNode* FoliageNode = Graph->AddNodeCopy(FoliageNodeSettings, DefaultNodeSettings);
-		UPCGNode* OutputNode = Graph->AddNodeCopy(OutputNodeSettings, DefaultNodeSettings);
+		UPCGNode* GravityNode = Graph->AddNodeCopy(GravitySettings, DefaultNodeSettings);
+		UPCGNode* MeshBuilderNode = Graph->AddNodeCopy(MeshBuilderSettings, DefaultNodeSettings);
+		UPCGNode* FoliagePaletteNode = Graph->AddNodeCopy(FoliagePaletteSettings, DefaultNodeSettings);
+		UPCGNode* OutputNode = Graph->AddNodeCopy(OutputSettings, DefaultNodeSettings);
 
 		OutputNode->NodeTitle = *VariantName;
 
-		MegaSeedLoadNode
+		PresetLoaderNode
 			->AddEdgeTo(*VariantName, GravityNode, PCGPinConstants::DefaultInputLabel)
-			->AddEdgeTo(PCGPinConstants::DefaultOutputLabel, MesherNode, PCGPinConstants::DefaultInputLabel)
-			->AddEdgeTo(PCGPinConstants::DefaultOutputLabel, FoliageNode, PCGPinConstants::DefaultInputLabel)
+			->AddEdgeTo(PCGPinConstants::DefaultOutputLabel, MeshBuilderNode, PCGPinConstants::DefaultInputLabel)
+			->AddEdgeTo(PCGPinConstants::DefaultOutputLabel, FoliagePaletteNode, PCGPinConstants::DefaultInputLabel)
 			->AddEdgeTo(PCGPinConstants::DefaultOutputLabel, OutputNode, PCGPinConstants::DefaultInputLabel);
 
 		GravityNode->SetNodePosition(Position.X + 200, Position.Y);
-		MesherNode->SetNodePosition(Position.X + 400, Position.Y);
-		FoliageNode->SetNodePosition(Position.X + 600, Position.Y);
+		MeshBuilderNode->SetNodePosition(Position.X + 400, Position.Y);
+		FoliagePaletteNode->SetNodePosition(Position.X + 600, Position.Y);
 		OutputNode->SetNodePosition(Position.X + 800, Position.Y);
 
 		Position.Y += 150;
diff --git a/Engine/Plugins/Experimental/ProceduralVegetationEditor/Source/ProceduralVegetation/Public/DataTypes/PVData.h b/Engine/Plugins/Experimental/ProceduralVegetationEditor/Source/ProceduralVegetation/Public/DataTypes/PVData.h
new file mode 100644
--- /dev/null
+++ b/Engine/Plugins/Experimental/ProceduralVegetationEditor/Source/ProceduralVegetation/Public/DataTypes/PVData.h
@@ -0,0 +1,156 @@
+// Copyright Epic Games, Inc. All Rights Reserved.
+
+#pragma once
+
+#include "Data/PCGSpatialData.h"
+#include "GeometryCollection/ManagedArrayCollection.h"
+#include "Helpers/PVUtilities.h"
+#include "PVData.generated.h"
+
+struct FPCGContext;
+
+UENUM()
+enum class EPVRenderType : uint8
+{
+	None 		= 0,
+	Skeleton 	= 1 << 0,
+	Mesh 		= 1 << 1,
+	Foliage 	= 1 << 2,
+	Bones 		= 1 << 3,
+
+	// Add new renter type above
+	Count	
+};
+
+ENUM_CLASS_FLAGS(EPVRenderType)
+
+UENUM()
+enum class EPVDebugType
+{
+	Point,
+	Branches,
+	Foliage,
+	Custom
+};
+
+UENUM()
+enum class EPVDebugValueVisualizationMode
+{
+	Text,
+	Direction,
+	Point
+};
+
+USTRUCT()
+struct FPVVisualizationSettings
+{
+	GENERATED_BODY()
+
+	UPROPERTY(EditAnywhere, Category = DebugSettings)
+	EPVDebugType DebugType = EPVDebugType::Point;
+
+	UPROPERTY(EditAnywhere, Category = DebugSettings)
+	bool bShowAnchorPoints = true;
+
+	UPROPERTY(EditAnywhere, Category = DebugSettings)
+	FName AttributeToFilter;
+
+	UPROPERTY(EditAnywhere, Category = DebugSettings)
+	EPVDebugValueVisualizationMode VisualizationMode = EPVDebugValueVisualizationMode::Text;
+
+	UPROPERTY(EditAnywhere, Category = DebugSettings)
+	FColor Color = FColor::Black;
+	
+	FName GetDebugTypeString() const
+	{
+		switch (DebugType)
+		{
+		case EPVDebugType::Point:
+			return "Points";
+
+		case EPVDebugType::Branches:
+			return "Primitives";
+
+		case EPVDebugType::Foliage:
+			return "Foliage";
+
+		case EPVDebugType::Custom:
+			return "Custom";
+		}
+		return NAME_None;
+	}
+};
+
+USTRUCT()
+struct FPVDebugSettings
+{
+	GENERATED_BODY()
+
+	UPROPERTY(EditAnywhere, Category = DebugSettings)
+	TArray<FPVVisualizationSettings> VisualizationSettings;
+};
+
+UCLASS(BlueprintType, ClassGroup = (Procedural))
+class PROCEDURALVEGETATION_API UPVData : public UPCGSpatialData
+{
+	GENERATED_BODY()
+
+public:
+	UPVData(const FObjectInitializer& ObjectInitializer);
+	
+	void Initialize(FManagedArrayCollection&& InCollection, bool bCanTakeOwnership = false);
+	
+	// ~Begin UPCGData interface
+	virtual EPCGDataType GetDataType() const override { return EPCGDataType::Other; }
+	virtual void AddToCrc(FArchiveCrc32& Ar, bool bFullDataCrc) const override;
+	// ~End UPCGData interface
+
+	// ~Begin UPCGSpatialData interface
+	virtual int GetDimension() const override { return 3; }
+	virtual FBox GetBounds() const override;
+	virtual bool SamplePoint(const FTransform& Transform, const FBox& Bounds, FPCGPoint& OutPoint, UPCGMetadata* OutMetadata) const override;
+	//~End UPCGSpatialData interface
+
+	const FManagedArrayCollection& GetCollection() const { return Collection; }
+	FManagedArrayCollection& GetCollection() { return Collection; }
+	
+#if WITH_EDITOR
+	const TArray<EPVRenderType>& GetRenderType() const { return RenderType; }
+	void SetRenderType(TArray<EPVRenderType> InRenderType) const { RenderType = InRenderType; }
+#endif
+
+#if PV_SHOW_DEBUG_VISUALIZATION 
+	void SetDebugSettings(const FPVDebugSettings& InSettings) const { DebugSettings = InSettings; }
+	const FPVDebugSettings& GetDebugSettings() const { return DebugSettings; }
+#endif
+	
+protected:
+
+	FManagedArrayCollection Collection;
+
+#if WITH_EDITOR
+	mutable TArray<EPVRenderType> RenderType;
+
+#if PV_SHOW_DEBUG_VISUALIZATION 
+	mutable FPVDebugSettings DebugSettings;
+#endif
+	
+#endif
+	
+	// ~Begin UPCGData interface
+	virtual bool SupportsFullDataCrc() const override { return true; }
+	// ~End UPCGData interface
+
+	// ~Begin UPCGSpatialData interface
+	virtual UPCGSpatialData* CopyInternal(FPCGContext* Context) const override;
+public:
+	virtual const UPCGPointData* ToPointData(FPCGContext* Context, const FBox& InBounds) const override;
+	virtual const UPCGPointArrayData* ToPointArrayData(FPCGContext* Context, const FBox& InBounds) const override;
+	//~End UPCGSpatialData interface
+	
+private:
+	const UPCGBasePointData* ToBasePointData(FPCGContext* Context, const FBox& InBounds, TSubclassOf<UPCGBasePointData> PointDataClass) const;
+protected:
+	
+	mutable FBox CachedBounds = FBox(EForceInit::ForceInit);
+};
diff --git a/Engine/Plugins/Experimental/ProceduralVegetationEditor/Source/ProceduralVegetation/Public/DataTypes/PVFoliageMeshData.h b/Engine/Plugins/Experimental/ProceduralVegetationEditor/Source/ProceduralVegetation/Public/DataTypes/PVFoliageMeshData.h
new file mode 100644
--- /dev/null
+++ b/Engine/Plugins/Experimental/ProceduralVegetationEditor/Source/ProceduralVegetation/Public/DataTypes/PVFoliageMeshData.h
@@ -0,0 +1,26 @@
+// Copyright Epic Games, Inc. All Rights Reserved.
+
+#pragma once
+
+#include "PVData.h"
+
+#include "PVFoliageMeshData.generated.h"
+
+USTRUCT()
+struct FPVDataTypeInfoFoliageMesh : public FPCGDataTypeInfo
+{
+	GENERATED_BODY()
+
+	PCG_DECLARE_TYPE_INFO(PROCEDURALVEGETATION_API);
+};
+
+UCLASS(BlueprintType, ClassGroup = (Procedural))
+class UPVFoliageMeshData : public UPVData
+{
+	GENERATED_BODY()
+
+public:
+	//~ Begin UPCGData interface
+	PCG_ASSIGN_DEFAULT_TYPE_INFO(FPVDataTypeInfoFoliageMesh)
+	//~ End UPCGData interface
+};
\ No newline at end of file
diff --git a/Engine/Plugins/Experimental/ProceduralVegetationEditor/Source/ProceduralVegetation/Public/DataTypes/PVGrowthData.h b/Engine/Plugins/Experimental/ProceduralVegetationEditor/Source/ProceduralVegetation/Public/DataTypes/PVGrowthData.h
new file mode 100644
--- /dev/null
+++ b/Engine/Plugins/Experimental/ProceduralVegetationEditor/Source/ProceduralVegetation/Public/DataTypes/PVGrowthData.h
@@ -0,0 +1,26 @@
+// Copyright Epic Games, Inc. All Rights Reserved.
+
+#pragma once
+
+#include "PVData.h"
+
+#include "PVGrowthData.generated.h"
+
+USTRUCT()
+struct FPVDataTypeInfoGrowth : public FPCGDataTypeInfo
+{
+	GENERATED_BODY()
+
+	PCG_DECLARE_TYPE_INFO(PROCEDURALVEGETATION_API);
+};
+
+UCLASS(BlueprintType, ClassGroup = (Procedural))
+class UPVGrowthData : public UPVData
+{
+	GENERATED_BODY()
+
+public:
+	//~ Begin UPCGData interface
+	PCG_ASSIGN_DEFAULT_TYPE_INFO(FPVDataTypeInfoGrowth)
+	//~ End UPCGData interface
+};
\ No newline at end of file
diff --git a/Engine/Plugins/Experimental/ProceduralVegetationEditor/Source/ProceduralVegetation/Public/DataTypes/PVMeshData.h b/Engine/Plugins/Experimental/ProceduralVegetationEditor/Source/ProceduralVegetation/Public/DataTypes/PVMeshData.h
new file mode 100644
--- /dev/null
+++ b/Engine/Plugins/Experimental/ProceduralVegetationEditor/Source/ProceduralVegetation/Public/DataTypes/PVMeshData.h
@@ -0,0 +1,26 @@
+// Copyright Epic Games, Inc. All Rights Reserved.
+
+#pragma once
+
+#include "PVData.h"
+
+#include "PVMeshData.generated.h"
+
+USTRUCT()
+struct FPVDataTypeInfoMesh : public FPCGDataTypeInfo
+{
+	GENERATED_BODY()
+
+	PCG_DECLARE_TYPE_INFO(PROCEDURALVEGETATION_API);
+};
+
+UCLASS(BlueprintType, ClassGroup = (Procedural))
+class UPVMeshData : public UPVData
+{
+	GENERATED_BODY()
+
+public:
+	//~ Begin UPCGData interface
+	PCG_ASSIGN_DEFAULT_TYPE_INFO(FPVDataTypeInfoMesh)
+	//~ End UPCGData interface
+};
\ No newline at end of file
diff --git a/Engine/Plugins/Experimental/ProceduralVegetationEditor/Source/ProceduralVegetation/Public/Helpers/PVExportHelper.h b/Engine/Plugins/Experimental/ProceduralVegetationEditor/Source/ProceduralVegetation/Public/Helpers/PVExportHelper.h
--- a/Engine/Plugins/Experimental/ProceduralVegetationEditor/Source/ProceduralVegetation/Public/Helpers/PVExportHelper.h
+++ b/Engine/Plugins/Experimental/ProceduralVegetationEditor/Source/ProceduralVegetation/Public/Helpers/PVExportHelper.h
@@ -23,7 +23,7 @@ enum class EPVExportMeshType: uint8
 };
 
 USTRUCT()
-struct PROCEDURALVEGETATION_API FPVExportSettings
+struct PROCEDURALVEGETATION_API FPVExportParams
 {
 	GENERATED_BODY()
 
@@ -133,6 +133,6 @@ namespace PV::Export
 		const TObjectPtr<UProceduralVegetation> InProceduralVegetation,
 		const FManagedArrayCollection& Collection,
 		const FManagedArrayCollection& FoliageCollection,
-		const FPVExportSettings& Settings
+		const FPVExportParams& ExportParams
 	);
 }
diff --git a/Engine/Plugins/Experimental/ProceduralVegetationEditor/Source/ProceduralVegetation/Public/Nodes/PVBaseSettings.h b/Engine/Plugins/Experimental/ProceduralVegetationEditor/Source/ProceduralVegetation/Public/Nodes/PVBaseSettings.h
new file mode 100644
--- /dev/null
+++ b/Engine/Plugins/Experimental/ProceduralVegetationEditor/Source/ProceduralVegetation/Public/Nodes/PVBaseSettings.h
@@ -0,0 +1,43 @@
+﻿// Copyright Epic Games, Inc. All Rights Reserved.
+#pragma once
+
+#include "CoreMinimal.h"
+#include "PCGSettings.h"
+#include "DataTypes/PVData.h"
+#include "Helpers/PVUtilities.h"
+#include "PVBaseSettings.generated.h"
+
+UCLASS(BlueprintType, Abstract, HideCategories=(Debug), ClassGroup = (Procedural))
+class UPVBaseSettings : public UPCGSettings
+{
+	GENERATED_BODY()
+
+public:
+#if WITH_EDITOR
+	virtual EPCGSettingsType GetType() const override { return EPCGSettingsType::Spatial; }
+	virtual TArray<EPVRenderType> GetRenderType() const { return TArray<EPVRenderType>(); }
+#endif
+
+#if PV_SHOW_DEBUG_VISUALIZATION
+	FPVDebugSettings GetDebugVisualizationSettings() const { return DebugVisualizationSettings; }
+#endif
+
+protected:
+	virtual TArray<FPCGPinProperties> InputPinProperties() const override;
+	virtual TArray<FPCGPinProperties> OutputPinProperties() const override;
+	virtual FPCGElementPtr CreateElement() const override;
+
+	virtual FPCGDataTypeIdentifier GetInputPinTypeIdentifier() const;
+	virtual FPCGDataTypeIdentifier GetOutputPinTypeIdentifier() const;
+private:
+	
+	UPROPERTY(EditAnywhere, Category = DebugSettings, meta=(EditCondition = "bDebug", EditConditionHides))
+	FPVDebugSettings DebugVisualizationSettings;
+};
+
+class FPVBaseElement : public IPCGElement
+{
+protected:
+	virtual void PostExecuteInternal(FPCGContext* InContext) const override;
+	virtual bool ExecuteInternal(FPCGContext* Context) const override;
+};
\ No newline at end of file
diff --git a/Engine/Plugins/Experimental/ProceduralVegetationEditor/Source/ProceduralVegetation/Public/Nodes/PVOutputSettings.h b/Engine/Plugins/Experimental/ProceduralVegetationEditor/Source/ProceduralVegetation/Public/Nodes/PVOutputSettings.h
new file mode 100644
--- /dev/null
+++ b/Engine/Plugins/Experimental/ProceduralVegetationEditor/Source/ProceduralVegetation/Public/Nodes/PVOutputSettings.h
@@ -0,0 +1,38 @@
+﻿// Copyright Epic Games, Inc. All Rights Reserved.
+#pragma once
+
+#include "CoreMinimal.h"
+#include "Helpers/PVExportHelper.h"
+#include "Nodes/PVBaseSettings.h"
+#include "PVOutputSettings.generated.h"
+
+UCLASS(ClassGroup = (Procedural))
+class UPVOutputSettings : public UPVBaseSettings
+{
+	GENERATED_BODY()
+
+public:
+	UPVOutputSettings();
+
+#if WITH_EDITOR
+	virtual FName GetDefaultNodeName() const override { return FName(TEXT("ProceduralVegetationOutput")); }
+	virtual FText GetDefaultNodeTitle() const override { return NSLOCTEXT("PVOutputSettings", "NodeTitle", "Procedural Vegetation Output"); }
+	virtual FText GetNodeTooltipText() const override { return NSLOCTEXT("PVOutputSettings", "NodeTooltip", "Output node for Procedural Vegetation"); }
+	virtual EPCGSettingsType GetType() const override { return EPCGSettingsType::InputOutput; }
+#endif
+
+protected:
+	virtual FPCGDataTypeIdentifier GetInputPinTypeIdentifier() const override;
+	virtual TArray<FPCGPinProperties> OutputPinProperties() const override;
+
+	virtual FPCGElementPtr CreateElement() const override;
+	
+public:
+	UPROPERTY(EditAnywhere, Category = ExportSettings, meta = (ShowOnlyInnerProperties))
+	FPVExportParams ExportSettings;
+};
+
+class FPVOutputElement : public FPVBaseElement
+{
+	virtual bool ExecuteInternal(FPCGContext* Context) const override;
+};
\ No newline at end of file
diff --git a/Engine/Plugins/Experimental/ProceduralVegetationEditor/Source/ProceduralVegetation/Public/Nodes/PVPCGBaseSettings.h b/Engine/Plugins/Experimental/ProceduralVegetationEditor/Source/ProceduralVegetation/Public/Nodes/PVPCGBaseSettings.h
deleted file mode 100644
diff --git a/Engine/Plugins/Experimental/ProceduralVegetationEditor/Source/ProceduralVegetation/Public/Nodes/PVPCGOutput.h b/Engine/Plugins/Experimental/ProceduralVegetationEditor/Source/ProceduralVegetation/Public/Nodes/PVPCGOutput.h
deleted file mode 100644
diff --git a/Engine/Plugins/Experimental/ProceduralVegetationEditor/Source/ProceduralVegetation/Public/PVPCGData.h b/Engine/Plugins/Experimental/ProceduralVegetationEditor/Source/ProceduralVegetation/Public/PVPCGData.h
deleted file mode 100644
diff --git a/Engine/Plugins/Experimental/ProceduralVegetationEditor/Source/ProceduralVegetationEditor/Private/PVEditor.cpp b/Engine/Plugins/Experimental/ProceduralVegetationEditor/Source/ProceduralVegetationEditor/Private/PVEditor.cpp
--- a/Engine/Plugins/Experimental/ProceduralVegetationEditor/Source/ProceduralVegetationEditor/Private/PVEditor.cpp
+++ b/Engine/Plugins/Experimental/ProceduralVegetationEditor/Source/ProceduralVegetationEditor/Private/PVEditor.cpp
@@ -7,8 +7,8 @@
 #include "Subsystems/PCGEngineSubsystem.h"
 #include "PVEditorSchema.h"
 #include "ProceduralVegetation.h"
-#include "PVPCGData.h"
-#include "Nodes/PVPCGOutput.h"
+#include "DataTypes/PVData.h"
+#include "Nodes/PVOutputSettings.h"
 #include "ToolMenus.h"
 #include "ToolMenuSection.h"
 #include "Dataflow/DataflowCollectionSpreadSheetWidget.h"
@@ -294,7 +294,7 @@ FReply FPVEditor::ExportInternal()
 		
 		for (const UPCGNode* Node : Nodes)
 		{
-			if(const UPVPCGOutputSettings* OutputSettings = Cast<UPVPCGOutputSettings>(Node->GetSettings()))
+			if(const UPVOutputSettings* OutputSettings = Cast<UPVOutputSettings>(Node->GetSettings()))
 			{
 				if(OutputSettings->ExportSettings.bShouldExport)
 				{
@@ -308,7 +308,7 @@ FReply FPVEditor::ExportInternal()
 		
 		for(auto Node : NodesToExport)
 		{
-			if(const UPVPCGOutputSettings* OutputSettings = Cast<UPVPCGOutputSettings>(Node->GetSettings()))
+			if(const UPVOutputSettings* OutputSettings = Cast<UPVOutputSettings>(Node->GetSettings()))
 			{
 				FManagedArrayCollection OutCollection;
 				FName PinName = NAME_None;
@@ -355,7 +355,7 @@ bool FPVEditor::GetCollectionFromNode(const UPCGNode* InNode, FName& OutPinName,
 		
 		if(TaggedData != nullptr)
 		{
-			const UPVPCGData* ProceduralVegetationData = Cast<UPVPCGData>(TaggedData->Data);
+			const UPVData* ProceduralVegetationData = Cast<UPVData>(TaggedData->Data);
 		
 			if(ProceduralVegetationData != nullptr)
 			{
diff --git a/Engine/Plugins/Experimental/ProceduralVegetationEditor/Source/ProceduralVegetationEditor/Private/PVEditorSchema.cpp b/Engine/Plugins/Experimental/ProceduralVegetationEditor/Source/ProceduralVegetationEditor/Private/PVEditorSchema.cpp
--- a/Engine/Plugins/Experimental/ProceduralVegetationEditor/Source/ProceduralVegetationEditor/Private/PVEditorSchema.cpp
+++ b/Engine/Plugins/Experimental/ProceduralVegetationEditor/Source/ProceduralVegetationEditor/Private/PVEditorSchema.cpp
@@ -1,7 +1,7 @@
 ﻿// Copyright Epic Games, Inc. All Rights Reserved.
 
 #include "PVEditorSchema.h"
-#include "Nodes/PVPCGBaseSettings.h"
+#include "Nodes/PVBaseSettings.h"
 
 EPCGElementType UPVEditorSchema::GetElementTypeFiltering() const
 {
@@ -15,7 +15,7 @@ const FPCGGraphEditorCustomization& UPVEditorSchema::GetGraphEditorCustomization
 		.bFilterSettings = true,
 		.FilteredSettingsTypes =
 	{
-			UPVPCGBaseSettings::StaticClass()
+			UPVBaseSettings::StaticClass()
 		}
 	};
 
diff --git a/Engine/Plugins/Experimental/ProceduralVegetationEditor/Source/ProceduralVegetationEditor/Private/ProceduralVegetationEditorModule.cpp b/Engine/Plugins/Experimental/ProceduralVegetationEditor/Source/ProceduralVegetationEditor/Private/ProceduralVegetationEditorModule.cpp
--- a/Engine/Plugins/Experimental/ProceduralVegetationEditor/Source/ProceduralVegetationEditor/Private/ProceduralVegetationEditorModule.cpp
+++ b/Engine/Plugins/Experimental/ProceduralVegetationEditor/Source/ProceduralVegetationEditor/Private/ProceduralVegetationEditorModule.cpp
@@ -9,7 +9,12 @@
 #include "Importer/MegaSeedImportWorkflowFactory.h"
 #include "Modules/ModuleManager.h"
 #include "PVEditorCommands.h"
-#include "PVPCGData.h"
+#include "PVEditorSettings.h"
+#include "Styling/SlateStyleRegistry.h"
+#include "DataTypes/PVData.h"
+#include "DataTypes/PVGrowthData.h"
+#include "DataTypes/PVMeshData.h"
+#include "DataTypes/PVFoliageMeshData.h"
 #include "Selection.h"
 #include "Animation/SkeletalMeshActor.h"
 #include "Visualizations/PVPCGDataVisualization.h"
@@ -37,10 +42,11 @@ void FProceduralVegetationEditorModule::StartupModule()
 	RegisterMenus();
 	FPVEditorCommands::Register();
 	RegisterLevelEditorMenuExtension();
+	RegisterPinColorAndIcons();
 
 	// Visualizations
 	FPCGDataVisualizationRegistry& DataVisRegistry = FPCGModule::GetMutablePCGDataVisualizationRegistry();
-	DataVisRegistry.RegisterPCGDataVisualization(UPVPCGData::StaticClass(), MakeUnique<const FPVPCGDataVisualization>());
+	DataVisRegistry.RegisterPCGDataVisualization(UPVData::StaticClass(), MakeUnique<const FPVPCGDataVisualization>());
 
 	// FAB
 	MegaSeedImportWorkflowFactory = MakeShared<FMegaSeedImportWorkflowFactory>();
@@ -50,6 +56,7 @@ void FProceduralVegetationEditorModule::StartupModule()
 void FProceduralVegetationEditorModule::ShutdownModule()
 {
 	UnregisterLevelEditorMenuExtension();
+	UnregisterPinColorAndIcons();
 	FPVEditorCommands::Unregister();
 	FFabWorkflowFactoryRegistry::UnregisterFactory(MegaSeedImportWorkflowFactory);
 }
@@ -153,6 +160,24 @@ void FProceduralVegetationEditorModule::UnregisterLevelEditorMenuExtension()
 	}
 }
 
+void FProceduralVegetationEditorModule::RegisterPinColorAndIcons()
+{
+	FPCGDataTypeRegistry& TypeRegistry = FPCGModule::GetMutableDataTypeRegistry();
+
+	TypeRegistry.RegisterPinColorFunction(FPVDataTypeInfoGrowth::AsId(), [](const FPCGDataTypeIdentifier&) { return GetDefault<UPVEditorSettings>()->GrowthDataPinColor; });
+	TypeRegistry.RegisterPinColorFunction(FPVDataTypeInfoMesh::AsId(), [](const FPCGDataTypeIdentifier&) { return GetDefault<UPVEditorSettings>()->MeshDataPinColor; });
+	TypeRegistry.RegisterPinColorFunction(FPVDataTypeInfoFoliageMesh::AsId(), [](const FPCGDataTypeIdentifier&) { return GetDefault<UPVEditorSettings>()->FoliageMeshDataPinColor; });
+}
+
+void FProceduralVegetationEditorModule::UnregisterPinColorAndIcons()
+{
+	FPCGDataTypeRegistry& TypeRegistry = FPCGModule::GetMutableDataTypeRegistry();
+
+	TypeRegistry.UnregisterPinColorFunction(FPVDataTypeInfoGrowth::AsId());
+	TypeRegistry.UnregisterPinColorFunction(FPVDataTypeInfoMesh::AsId());
+	TypeRegistry.UnregisterPinColorFunction(FPVDataTypeInfoFoliageMesh::AsId());
+}
+
 TSharedRef<FExtender> FProceduralVegetationEditorModule::OnExtendLevelEditorActorSelectionMenu(const TSharedRef<FUICommandList> InCommandList, TArray<AActor*> InSelectedActors)
 {
 	TSharedRef<FExtender> Extender = MakeShared<FExtender>();
diff --git a/Engine/Plugins/Experimental/ProceduralVegetationEditor/Source/ProceduralVegetationEditor/Private/Visualizations/PVPCGDataVisualization.cpp b/Engine/Plugins/Experimental/ProceduralVegetationEditor/Source/ProceduralVegetationEditor/Private/Visualizations/PVPCGDataVisualization.cpp
--- a/Engine/Plugins/Experimental/ProceduralVegetationEditor/Source/ProceduralVegetationEditor/Private/Visualizations/PVPCGDataVisualization.cpp
+++ b/Engine/Plugins/Experimental/ProceduralVegetationEditor/Source/ProceduralVegetationEditor/Private/Visualizations/PVPCGDataVisualization.cpp
@@ -16,7 +16,7 @@
 #include "GeometryCollection/GeometryCollectionObject.h"
 #include "GeometryCollection/Facades/CollectionRenderingFacade.h"
 #include "Helpers/PVUtilities.h"
-#include "PVPCGData.h"
+#include "DataTypes/PVData.h"
 #include "Components/TextRenderComponent.h"
 #include "DebugVisualization/PVDebugVisualizer.h"
 #include "Visualizations/PVSkeletonVisualizerComponent.h"
@@ -51,7 +51,7 @@ void FPVPCGDataVisualization::ExecuteDebugDisplay(FPCGContext* Context, const UP
 
 FPCGSetupSceneFunc FPVPCGDataVisualization::GetViewportSetupFunc(const UPCGSettingsInterface* SettingsInterface, const UPCGData* Data) const
 {
-	return [this, WeakData=TWeakObjectPtr<const UPVPCGData>(Cast<UPVPCGData>(Data)), WeakSettings = TWeakObjectPtr<const UPCGSettingsInterface>(Cast<UPCGSettingsInterface>(SettingsInterface))](
+	return [this, WeakData=TWeakObjectPtr<const UPVData>(Cast<UPVData>(Data)), WeakSettings = TWeakObjectPtr<const UPCGSettingsInterface>(Cast<UPCGSettingsInterface>(SettingsInterface))](
 		FPCGSceneSetupParams& InOutParams)
 	{
 		check(InOutParams.Scene);
diff --git a/Engine/Plugins/Experimental/ProceduralVegetationEditor/Source/ProceduralVegetationEditor/Private/Visualizations/PVPCGDataVisualization.h b/Engine/Plugins/Experimental/ProceduralVegetationEditor/Source/ProceduralVegetationEditor/Private/Visualizations/PVPCGDataVisualization.h
--- a/Engine/Plugins/Experimental/ProceduralVegetationEditor/Source/ProceduralVegetationEditor/Private/Visualizations/PVPCGDataVisualization.h
+++ b/Engine/Plugins/Experimental/ProceduralVegetationEditor/Source/ProceduralVegetationEditor/Private/Visualizations/PVPCGDataVisualization.h
@@ -4,7 +4,7 @@
 
 #include "PCGDataVisualization.h"
 #include "DataVisualizations/PCGSpatialDataVisualization.h"
-#include "PVPCGData.h"
+#include "DataTypes/PVData.h"
 
 class AActor;
 class UPCGData;
diff --git a/Engine/Plugins/Experimental/ProceduralVegetationEditor/Source/ProceduralVegetationEditor/Private/Widgets/SPVExportSelectionDialog.cpp b/Engine/Plugins/Experimental/ProceduralVegetationEditor/Source/ProceduralVegetationEditor/Private/Widgets/SPVExportSelectionDialog.cpp
--- a/Engine/Plugins/Experimental/ProceduralVegetationEditor/Source/ProceduralVegetationEditor/Private/Widgets/SPVExportSelectionDialog.cpp
+++ b/Engine/Plugins/Experimental/ProceduralVegetationEditor/Source/ProceduralVegetationEditor/Private/Widgets/SPVExportSelectionDialog.cpp
@@ -6,7 +6,7 @@
 #include "Editor.h"
 #include "SPVExportItem.h"
 #include "ProceduralVegetation.h"
-#include "Nodes/PVPCGOutput.h"
+#include "Nodes/PVOutputSettings.h"
 #include "Widgets/Layout/SScrollBox.h"
 
 #define LOCTEXT_NAMESPACE "SPVExportSelectionDialog"
@@ -69,7 +69,7 @@ void SPVExportSelectionDialog::AddExportItems()
 	
 	Graph->ForEachNode([this](UPCGNode* Node)
 	{
-		const UPVPCGOutputSettings* OutputSettings = Cast<UPVPCGOutputSettings>(Node->GetSettings());
+		const UPVOutputSettings* OutputSettings = Cast<UPVOutputSettings>(Node->GetSettings());
 
 		if (OutputSettings)
 		{
@@ -119,7 +119,7 @@ void SPVExportSelectionDialog::OnOutputSelectionChanged(const FString ItemName,
 	{
 		if(ItemName == Node->GetNodeTitle(EPCGNodeTitleType::FullTitle).ToString())
 		{
-			if (UPVPCGOutputSettings* OutputSettings = Cast<UPVPCGOutputSettings>(Node->GetSettings()))
+			if (UPVOutputSettings* OutputSettings = Cast<UPVOutputSettings>(Node->GetSettings()))
 			{
 				OutputSettings->ExportSettings.bShouldExport = NewState ==ECheckBoxState::Checked;
 			}
diff --git a/Engine/Plugins/Experimental/ProceduralVegetationEditor/Source/ProceduralVegetationEditor/Public/PVEditorSettings.h b/Engine/Plugins/Experimental/ProceduralVegetationEditor/Source/ProceduralVegetationEditor/Public/PVEditorSettings.h
new file mode 100644
--- /dev/null
+++ b/Engine/Plugins/Experimental/ProceduralVegetationEditor/Source/ProceduralVegetationEditor/Public/PVEditorSettings.h
@@ -0,0 +1,25 @@
+// Copyright Epic Games, Inc. All Rights Reserved.
+
+#pragma once
+
+#include "Engine/DeveloperSettings.h"
+
+#include "PVEditorSettings.generated.h"
+
+UCLASS(MinimalAPI, config=EditorPerProjectUserSettings, meta = (DisplayName = "Procedural Vegetation Editor Settings"))
+class UPVEditorSettings : public UObject
+{
+	GENERATED_BODY()
+public:
+	/** Color used for data pins of type GrowthData */
+	UPROPERTY(EditAnywhere, config, Category = Node, meta = (HideAlphaChannel))
+	FLinearColor GrowthDataPinColor = FColor(239, 239, 239);
+
+	/** Color used for data pins of type MeshData */
+	UPROPERTY(EditAnywhere, config, Category = Node, meta = (HideAlphaChannel))
+	FLinearColor MeshDataPinColor = FColor(3, 234, 234);
+
+	/** Color used for data pins of type FoliageMeshData */
+	UPROPERTY(EditAnywhere, config, Category = Node, meta = (HideAlphaChannel))
+	FLinearColor FoliageMeshDataPinColor = FColor(0, 255, 0);
+};
diff --git a/Engine/Plugins/Experimental/ProceduralVegetationEditor/Source/ProceduralVegetationEditor/Public/ProceduralVegetationEditorModule.h b/Engine/Plugins/Experimental/ProceduralVegetationEditor/Source/ProceduralVegetationEditor/Public/ProceduralVegetationEditorModule.h
--- a/Engine/Plugins/Experimental/ProceduralVegetationEditor/Source/ProceduralVegetationEditor/Public/ProceduralVegetationEditorModule.h
+++ b/Engine/Plugins/Experimental/ProceduralVegetationEditor/Source/ProceduralVegetationEditor/Public/ProceduralVegetationEditorModule.h
@@ -29,6 +29,9 @@ private:
 
 	void RegisterLevelEditorMenuExtension();
 	void UnregisterLevelEditorMenuExtension();
+	void RegisterPinColorAndIcons();
+	void UnregisterPinColorAndIcons();
+
 	static TSharedRef<FExtender> OnExtendLevelEditorActorSelectionMenu(const TSharedRef<FUICommandList> InCommandList, TArray<AActor*> InSelectedActors);
 	static void ExecuteOpenInPVEditorOnSelectedActors();
 	static UProceduralVegetation* GetProceduralVegetationFromAsset(UObject* InObject);
