Fix build break due to cl45667402

diff --git a/Engine/Plugins/Experimental/ProceduralVegetationEditor/Source/ProceduralVegetationEditor/Private/Visualizations/DebugVisualization/PVDebugVisualizationBase.cpp b/Engine/Plugins/Experimental/ProceduralVegetationEditor/Source/ProceduralVegetationEditor/Private/Visualizations/DebugVisualization/PVDebugVisualizationBase.cpp
--- a/Engine/Plugins/Experimental/ProceduralVegetationEditor/Source/ProceduralVegetationEditor/Private/Visualizations/DebugVisualization/PVDebugVisualizationBase.cpp
+++ b/Engine/Plugins/Experimental/ProceduralVegetationEditor/Source/ProceduralVegetationEditor/Private/Visualizations/DebugVisualization/PVDebugVisualizationBase.cpp
@@ -4,10 +4,12 @@
 
 #if PV_SHOW_DEBUG_VISUALIZATION 
 
+#include "UObject/Package.h"
+#include "Engine/StaticMesh.h"
 #include "AdvancedPreviewScene.h"
 #include "PVDataAttributeValueInterface.h"
 #include "PVDebugVisualizer.h"
-#include "PVPCGData.h"
+#include "DataTypes/PVData.h"
 #include "Components/InstancedStaticMeshComponent.h"
 #include "Components/TextRenderComponent.h"
 #include "Visualizations/PVLineBatchComponent.h"
diff --git a/Engine/Plugins/Experimental/ProceduralVegetationEditor/Source/ProceduralVegetationEditor/Private/Visualizations/DebugVisualization/PVDebugVisualizationBase.h b/Engine/Plugins/Experimental/ProceduralVegetationEditor/Source/ProceduralVegetationEditor/Private/Visualizations/DebugVisualization/PVDebugVisualizationBase.h
--- a/Engine/Plugins/Experimental/ProceduralVegetationEditor/Source/ProceduralVegetationEditor/Private/Visualizations/DebugVisualization/PVDebugVisualizationBase.h
+++ b/Engine/Plugins/Experimental/ProceduralVegetationEditor/Source/ProceduralVegetationEditor/Private/Visualizations/DebugVisualization/PVDebugVisualizationBase.h
@@ -15,7 +15,7 @@ class UPVSkeletonVisualizerComponent;
 enum class EManagedArrayType : uint8;
 class FPVAttributeValueInterface;
 struct FPCGSceneSetupParams;
-class UPVPCGData;
+class UPVData;
 
 class FPVDebugVisualizationBase
 {
diff --git a/Engine/Plugins/Experimental/ProceduralVegetationEditor/Source/ProceduralVegetationEditor/Private/Visualizations/DebugVisualization/PVDebugVisualizer.cpp b/Engine/Plugins/Experimental/ProceduralVegetationEditor/Source/ProceduralVegetationEditor/Private/Visualizations/DebugVisualization/PVDebugVisualizer.cpp
--- a/Engine/Plugins/Experimental/ProceduralVegetationEditor/Source/ProceduralVegetationEditor/Private/Visualizations/DebugVisualization/PVDebugVisualizer.cpp
+++ b/Engine/Plugins/Experimental/ProceduralVegetationEditor/Source/ProceduralVegetationEditor/Private/Visualizations/DebugVisualization/PVDebugVisualizer.cpp
@@ -6,7 +6,7 @@
 
 #include "PVDebugVisualizationBase.h"
 #include "PVDebugVisualizations.h"
-#include "PVPCGData.h"
+#include "DataTypes/PVData.h"
 
 FVisualizerDrawContext::FVisualizerDrawContext(const FManagedArrayCollection& InCollection, const FPVVisualizationSettings& InSettings,
 	FPCGSceneSetupParams& InSceneParams)
@@ -16,7 +16,7 @@ FVisualizerDrawContext::FVisualizerDrawContext(const FManagedArrayCollection& In
 {
 }
 
-void FPVDebugVisualizer::Draw(const UPVPCGData* InData, FPCGSceneSetupParams& InOutParams)
+void FPVDebugVisualizer::Draw(const UPVData* InData, FPCGSceneSetupParams& InOutParams)
 {
 	check(InData);
 	
diff --git a/Engine/Plugins/Experimental/ProceduralVegetationEditor/Source/ProceduralVegetationEditor/Private/Visualizations/DebugVisualization/PVDebugVisualizer.h b/Engine/Plugins/Experimental/ProceduralVegetationEditor/Source/ProceduralVegetationEditor/Private/Visualizations/DebugVisualization/PVDebugVisualizer.h
--- a/Engine/Plugins/Experimental/ProceduralVegetationEditor/Source/ProceduralVegetationEditor/Private/Visualizations/DebugVisualization/PVDebugVisualizer.h
+++ b/Engine/Plugins/Experimental/ProceduralVegetationEditor/Source/ProceduralVegetationEditor/Private/Visualizations/DebugVisualization/PVDebugVisualizer.h
@@ -11,7 +11,7 @@ struct FManagedArrayCollection;
 struct FPVVisualizationSettings;
 class FPVDebugVisualizationBase;
 struct FPCGSceneSetupParams;
-class UPVPCGData;
+class UPVData;
 enum class EPVDebugType;
 
 struct FVisualizerDrawContext
@@ -25,7 +25,7 @@ struct FVisualizerDrawContext
 class FPVDebugVisualizer
 {
 public:
-	static void Draw(const UPVPCGData* InData,FPCGSceneSetupParams& InOutParams);
+	static void Draw(const UPVData* InData,FPCGSceneSetupParams& InOutParams);
 private:
 	static FPVDebugVisualizationPtr CreateVisualizer(EPVDebugType InDebugType);
 };
