[PVE] Fix issue where aggressive mesh simplification results in twisted geometry in some instances

diff --git a/Engine/Plugins/Experimental/ProceduralVegetationEditor/Source/ProceduralVegetation/Private/Implementations/PVMeshBuilder.cpp b/Engine/Plugins/Experimental/ProceduralVegetationEditor/Source/ProceduralVegetation/Private/Implementations/PVMeshBuilder.cpp
--- a/Engine/Plugins/Experimental/ProceduralVegetationEditor/Source/ProceduralVegetation/Private/Implementations/PVMeshBuilder.cpp
+++ b/Engine/Plugins/Experimental/ProceduralVegetationEditor/Source/ProceduralVegetation/Private/Implementations/PVMeshBuilder.cpp
@@ -310,8 +310,8 @@ TSet<int32> FPVMeshBuilder::ComputePointGradients(const PV::Facades::FPointFacad
 		const float PointGroundGradient =
 			MeshBuilderParams.PointGroundRetentionGradient.GetRichCurveConst()->Eval(1.0f - GroundGradient) * MeshBuilderParams.PointGroundRetention;
 
-		const float ScaleRemovalGradient = (ScaleGradient + PointHullGradient + PointMainTrunkGradient + PointGroundGradient) * MaxPointScale /
-			100.0f;
+		const float ScaleRemovalGradient =
+			(ScaleGradient + PointHullGradient + PointMainTrunkGradient + PointGroundGradient) * MaxPointScale / 100.0f;
 
 		// Compute mesh divisions gradient
 		const float DivisionsScaleGradient = MeshBuilderParams.DivisionScaleRetentionGradient.GetRichCurveConst()->Eval(1.0f - ScaleGradient);
@@ -407,6 +407,11 @@ void FPVMeshBuilder::PerformPathSimplification(const PV::Facades::FBranchFacade&
 
 		// Ensure long chains of removed points do not exceed a maximum threshold
 		BranchPoints = BranchFacade.GetPoints(BranchIndex);
+		if (BranchPoints.Num() == 0)
+		{
+			return;
+		}
+
 		FVector3f PreviousPointPosition = PointFacade.GetPosition(BranchPoints[0]);
 		float TravelDistance = 0.0f;
 
@@ -601,6 +606,11 @@ void FPVMeshBuilder::GenerateBranchMeshData(const bool bPrimitiveIsTrunk, const
                                        const TArray<int32>& TargetMeshDivisions, const FManagedArrayCollection& Collection,
                                        FLocalDynamicMeshData& OutLocalMeshData)
 {
+	if (PrimitivePoints.Num() < 2)
+	{
+		return;
+	}
+
 	const PV::Facades::FPointFacade PointFacade(Collection);
 	const PV::Facades::FBudVectorsFacade BudVectorFacade(Collection);
 	const PV::Facades::FPlantProfileFacade PlantProfileFacade(Collection);
@@ -634,17 +644,16 @@ void FPVMeshBuilder::GenerateBranchMeshData(const bool bPrimitiveIsTrunk, const
 	TArray<int32> CurrentIndices;
 	TArray<int32> PreviousIndices;
 
-	FVector3f PreviousPosition = PointPositions[PrimitivePoints[0]];
-	FVector3f PreviousUpVector = FVector3f::ForwardVector;
-
 	bool bShouldApplyDisplacement =
 		DisplacementData.Values.Num() > 0
 		&& DisplacementData.TextureWidth > 0
 		&& DisplacementData.TextureHeight > 0
 		&& GenerationNumber <= MeshBuilderParams.DisplacementGenerationUpperLimit;
 
-	const int32 MaxIndex = PrimitivePoints.Num() - 1;
-	for (int32 Index = 0; Index <= MaxIndex; ++Index)
+	FVector3f PreviousUpVector = FVector3f::ForwardVector;
+	FVector3f PreviousAimVector = FVector3f::UpVector;
+
+	for (int32 Index = 0; Index < PrimitivePoints.Num(); ++Index)
 	{
 		const int32& PointIndex = PrimitivePoints[Index];
 		const float NjordPixelId = PointIndex;
@@ -660,30 +669,62 @@ void FPVMeshBuilder::GenerateBranchMeshData(const bool bPrimitiveIsTrunk, const
 			: 1.0f / Interval;
 		TextureCoordV /= URatio;
 
-		const int32 PointDivisions = TargetMeshDivisions[PointIndex];
+		FVector3f AimVector = PointBudDirections[PointIndex][0].GetUnsafeNormal();
+		// The AimVector needs to be adjusted in case curvature has changed during mesh simplification
 		const FVector3f& PointPosition = PointPositions[PointIndex];
-		const float PointScale = PointScales[PointIndex];
-		FVector3f NormalVector = PointBudDirections[PointIndex][0].GetUnsafeNormal();
-		const float PlantGradient = PointFacade.GetPlantGradient(PointIndex);
-		const float ProfileFallOff = MeshBuilderParams.PlantProfileFallOff.GetRichCurveConst()->Eval(PlantGradient);
+		FVector3f UpdatedAimVector;
+		if (Index == PrimitivePoints.Num() - 1)
+		[[unlikely]]
+		{
+			const int32 PreviousPointIndex = PrimitivePoints[Index - 1];
+			const FVector3f& PreviousPointPosition = PointPositions[PreviousPointIndex];
+			UpdatedAimVector = (PointPosition - PreviousPointPosition).GetSafeNormal();
+		}
+		else
+		[[likely]]
+		{
+			const int32 NextPointIndex = PrimitivePoints[Index + 1];
+			const FVector3f& NextPointPosition = PointPositions[NextPointIndex];
+			UpdatedAimVector = (NextPointPosition - PointPosition).GetSafeNormal();
+		}
+		const FQuat4f RotationQuat = FQuat4f::FindBetweenNormals(AimVector, UpdatedAimVector);
 
 		const FVector3f PointUpOriginal = PointBudDirections[PointIndex][5].GetUnsafeNormal();
-		FVector3f PointUp = FVector3f::VectorPlaneProject(PointUpOriginal, NormalVector).GetSafeNormal();
+		const FVector3f PointUpUpdated = RotationQuat.RotateVector(PointUpOriginal);
+
+		FVector3f PointUp = PointUpUpdated;
+
+		FVector3f CrossVector = PointUp.Cross(UpdatedAimVector);
+		PointUp = CrossVector.Cross(UpdatedAimVector);
 
-		if (Index > 0 && FVector3f::DotProduct(PointUp, PreviousUpVector) < 0.0f)
+		if (Index > 0)
 		{
-			PointUp = -PointUp;
+			if (Index > 0 && FVector3f::DotProduct(PointUp, PreviousUpVector) < 0.0f)
+			{
+				PointUp = -PointUp;
+			}
+
+			const FVector3f InterpolatedAimVector = FMath::Lerp(PreviousAimVector, UpdatedAimVector, 0.5f).GetSafeNormal();
+			const FQuat4f AimVectorRotationQuat = FQuat4f::FindBetweenNormals(UpdatedAimVector, InterpolatedAimVector);
+			PointUp = AimVectorRotationQuat.RotateVector(PointUp);
+			UpdatedAimVector = InterpolatedAimVector;
 		}
+
+		PreviousAimVector = UpdatedAimVector;
 		PreviousUpVector = PointUp;
 
-		if (PointScale == 0.0f)
+		const float PlantGradient = PointFacade.GetPlantGradient(PointIndex);
+		const float ProfileFallOff = MeshBuilderParams.PlantProfileFallOff.GetRichCurveConst()->Eval(PlantGradient);
+		if (const float PointScale = PointScales[PointIndex];
+			PointScale == 0.0f)
 		[[unlikely]]
 		{
-			CurrentIndices.Add(OutLocalMeshData.Vertices.Emplace(PointPosition, NormalVector, FVector2f(0.0f, TextureCoordV), NjordPixelId));
+			CurrentIndices.Add(OutLocalMeshData.Vertices.Emplace(PointPosition, UpdatedAimVector, FVector2f(0.0f, TextureCoordV), NjordPixelId));
 		}
 		else
 		[[likely]]
 		{
+			const int32 PointDivisions = TargetMeshDivisions[PointIndex];
 			const float RotationDelta = UE_TWO_PI / PointDivisions;
 			for (int32 Division = 0; Division < PointDivisions; ++Division)
 			{
@@ -698,7 +739,7 @@ void FPVMeshBuilder::GenerateBranchMeshData(const bool bPrimitiveIsTrunk, const
 					ProfileMultiplier = ProfileMultiplier * MeshBuilderParams.PlantProfileScale;
 				}
 
-				const FVector3f Direction = FQuat4f(NormalVector, RotationDelta * Division).RotateVector(PointUp);
+				const FVector3f Direction = FQuat4f(UpdatedAimVector, RotationDelta * Division).RotateVector(PointUp).GetSafeNormal();
 				FVector3f RadialPoint = PointPosition + (Direction * PointScale * ProfileMultiplier);
 
 				if (bShouldApplyDisplacement)
@@ -736,8 +777,6 @@ void FPVMeshBuilder::GenerateBranchMeshData(const bool bPrimitiveIsTrunk, const
 				TextureCoordV),
 			NjordPixelId));
 
-		PreviousPosition = PointPosition;
-
 		TriangulateRings(PreviousIndices, CurrentIndices, PolyGroupIndex, OutLocalMeshData);
 
 		PreviousIndices = MoveTemp(CurrentIndices);
