[PVE] Fix Sample assets not showing up

diff --git a/Engine/Plugins/Experimental/ProceduralVegetationEditor/Source/ProceduralVegetationEditor/Private/Widgets/SPVCreateDialog.cpp b/Engine/Plugins/Experimental/ProceduralVegetationEditor/Source/ProceduralVegetationEditor/Private/Widgets/SPVCreateDialog.cpp
--- a/Engine/Plugins/Experimental/ProceduralVegetationEditor/Source/ProceduralVegetationEditor/Private/Widgets/SPVCreateDialog.cpp
+++ b/Engine/Plugins/Experimental/ProceduralVegetationEditor/Source/ProceduralVegetationEditor/Private/Widgets/SPVCreateDialog.cpp
@@ -23,7 +23,7 @@ void SPVCreateDialog::Construct(const FArguments& InArgs)
 	TSharedRef<SVerticalBox> Container = SNew(SVerticalBox);
 
 	FName SampleTreesVirtualPath;
-	IContentBrowserDataModule::Get().GetSubsystem()->ConvertInternalPathToVirtual("/ProceduralVegetationEditor/SampleTrees", SampleTreesVirtualPath);
+	IContentBrowserDataModule::Get().GetSubsystem()->ConvertInternalPathToVirtual("/ProceduralVegetationEditor", SampleTreesVirtualPath);
 
 	const FContentBrowserModule& ContentBrowserModule = FModuleManager::LoadModuleChecked<FContentBrowserModule>("ContentBrowser");
 
