Second rename pass for Procedural Vegetation Editor plugin

- Rename MegaSeedDataAsset to PVPresetDataAsset
- Change node title of PVPresetLoaderSettings from "Growth Data Loader" to "Procedural Vegetation Preset Loader"
- Change graph corner text from "PVE Graph" to "Procedural Vegetation"
- Remove "MegaPlant" from PVDebugVisualizations
- Replace "PVPCG" with "PV"

diff --git a/Engine/Plugins/Experimental/ProceduralVegetationEditor/Config/BaseProceduralVegetationEditor.ini b/Engine/Plugins/Experimental/ProceduralVegetationEditor/Config/BaseProceduralVegetationEditor.ini
--- a/Engine/Plugins/Experimental/ProceduralVegetationEditor/Config/BaseProceduralVegetationEditor.ini
+++ b/Engine/Plugins/Experimental/ProceduralVegetationEditor/Config/BaseProceduralVegetationEditor.ini
@@ -31,7 +31,7 @@
 
 +EnumRedirects=(OldName="MegaPlantExportMeshType",NewName="/Script/ProceduralVegetation.PVExportMeshType")
 
-+PropertyRedirects=(OldName="PCGMegaPlantRecipeLoadSettings.RecipeDataAsset",NewName="PVPresetLoaderSettings.PresetDataAsset")
++PropertyRedirects=(OldName="PCGMegaPlantRecipeLoadSettings.RecipeDataAsset",NewName="PVPresetLoaderSettings.Preset")
 
 ; Class/file name refactor
 +ClassRedirects=(OldName="/Script/ProceduralVegetation.PVPCGRemoveBranchesSettings",NewName="/Script/ProceduralVegetation.PVRemoveBranchesSettings")
@@ -47,9 +47,11 @@
 +ClassRedirects=(OldName="/Script/ProceduralVegetation.PVPCGBaseSettings",NewName="/Script/ProceduralVegetation.PVBaseSettings")
 +ClassRedirects=(OldName="/Script/ProceduralVegetation.PVPCGBoneReductionSettings",NewName="/Script/ProceduralVegetation.PVBoneReductionSettings")
 
++ClassRedirects=(OldName="/Script/ProceduralVegetation.MegaSeedDataAsset",NewName="/Script/ProceduralVegetation.ProceduralVegetationPreset")
+
 +StructRedirects=(OldName="/Script/ProceduralVegetation.PVMesherSettings",NewName="/Script/ProceduralVegetation.PVMeshBuilderParams")
 +StructRedirects=(OldName="/Script/ProceduralVegetation.PVGravitySettings",NewName="/Script/ProceduralVegetation.PVGravityParams")
 +StructRedirects=(OldName="/Script/ProceduralVegetation.PVCarveSettings",NewName="/Script/ProceduralVegetation.PVCarveParams")
 +StructRedirects=(OldName="/Script/ProceduralVegetation.PVExportSettings",NewName="/Script/ProceduralVegetation.PVExportParams")
 
-+PropertyRedirects=(OldName="PVPCGMegaSeedLoadSettings.RecipeDataAsset",NewName="PVPresetLoaderSettings.PresetDataAsset")
++PropertyRedirects=(OldName="PVPCGMegaSeedLoadSettings.RecipeDataAsset",NewName="PVPresetLoaderSettings.Preset")
diff --git a/Engine/Plugins/Experimental/ProceduralVegetationEditor/Source/ProceduralVegetation/Private/MegaSeedDataAsset.cpp b/Engine/Plugins/Experimental/ProceduralVegetationEditor/Source/ProceduralVegetation/Private/MegaSeedDataAsset.cpp
deleted file mode 100644
diff --git a/Engine/Plugins/Experimental/ProceduralVegetationEditor/Source/ProceduralVegetation/Private/Nodes/PVImporterSettings.cpp b/Engine/Plugins/Experimental/ProceduralVegetationEditor/Source/ProceduralVegetation/Private/Nodes/PVImporterSettings.cpp
--- a/Engine/Plugins/Experimental/ProceduralVegetationEditor/Source/ProceduralVegetation/Private/Nodes/PVImporterSettings.cpp
+++ b/Engine/Plugins/Experimental/ProceduralVegetationEditor/Source/ProceduralVegetation/Private/Nodes/PVImporterSettings.cpp
@@ -38,7 +38,7 @@ bool FPVImporterElement::ExecuteInternal(FPCGContext* InContext) const
 	FString ErrorMessage;
 
 	FManagedArrayCollection Collection;
-	bool bSuccess = PV::LoadMegaSeedJsonToCollection(Collection, SkeletonFilePath, ErrorMessage);
+	bool bSuccess = PV::LoadMegaPlantsJsonToCollection(Collection, SkeletonFilePath, ErrorMessage);
 	if (!bSuccess)
 	{
 		UE_LOG(LogProceduralVegetation, Error, TEXT("%s"), *ErrorMessage);
diff --git a/Engine/Plugins/Experimental/ProceduralVegetationEditor/Source/ProceduralVegetation/Private/Nodes/PVPresetLoaderSettings.cpp b/Engine/Plugins/Experimental/ProceduralVegetationEditor/Source/ProceduralVegetation/Private/Nodes/PVPresetLoaderSettings.cpp
--- a/Engine/Plugins/Experimental/ProceduralVegetationEditor/Source/ProceduralVegetation/Private/Nodes/PVPresetLoaderSettings.cpp
+++ b/Engine/Plugins/Experimental/ProceduralVegetationEditor/Source/ProceduralVegetation/Private/Nodes/PVPresetLoaderSettings.cpp
@@ -3,7 +3,7 @@
 #include "PVPresetLoaderSettings.h"
 #include "PCGContext.h"
 #include "DataTypes/PVGrowthData.h"
-#include "MegaSeedDataAsset.h"
+#include "ProceduralVegetationPreset.h"
 
 TArray<FPCGPinProperties> UPVPresetLoaderSettings::InputPinProperties() const
 {
@@ -13,9 +13,9 @@ TArray<FPCGPinProperties> UPVPresetLoaderSettings::InputPinProperties() const
 TArray<FPCGPinProperties> UPVPresetLoaderSettings::OutputPinProperties() const
 {
 	TArray<FPCGPinProperties> Properties;
-	if (PresetDataAsset)
+	if (Preset)
 	{
-		for (const auto& [VariantName, VariantData] : PresetDataAsset->Variants)
+		for (const auto& [VariantName, VariantData] : Preset->Variants)
 		{
 			FPCGPinProperties& Pin = Properties.Emplace_GetRef(*VariantName, GetOutputPinTypeIdentifier());
 			Pin.bAllowMultipleData = false;
@@ -44,13 +44,13 @@ bool FPVPresetLoaderElement::ExecuteInternal(FPCGContext* InContext) const
 	const UPVPresetLoaderSettings* Settings = InContext->GetInputSettings<UPVPresetLoaderSettings>();
 	check(Settings);
 
-	if (!Settings->PresetDataAsset)
+	if (!Settings->Preset)
 	{
 		PCGLog::LogWarningOnGraph(NSLOCTEXT("PVPresetLoaderSettings", "MissingDataAsset", "Preset data asset is not set"), InContext);
 		return true;
 	}
 
-	for (const auto& [VariantName, VariantData] : Settings->PresetDataAsset->Variants)
+	for (const auto& [VariantName, VariantData] : Settings->Preset->Variants)
 	{
 		UPVGrowthData* OutVariantData = FPCGContext::NewObject_AnyThread<UPVGrowthData>(InContext);
 		OutVariantData->Initialize(CopyTemp(VariantData));
diff --git a/Engine/Plugins/Experimental/ProceduralVegetationEditor/Source/ProceduralVegetation/Private/Nodes/PVPresetLoaderSettings.h b/Engine/Plugins/Experimental/ProceduralVegetationEditor/Source/ProceduralVegetation/Private/Nodes/PVPresetLoaderSettings.h
--- a/Engine/Plugins/Experimental/ProceduralVegetationEditor/Source/ProceduralVegetation/Private/Nodes/PVPresetLoaderSettings.h
+++ b/Engine/Plugins/Experimental/ProceduralVegetationEditor/Source/ProceduralVegetation/Private/Nodes/PVPresetLoaderSettings.h
@@ -14,8 +14,8 @@ class UPVPresetLoaderSettings : public UPVBaseSettings
 public:
 #if WITH_EDITOR
 	virtual FName GetDefaultNodeName() const override { return FName(TEXT("ProceduralVegetationPresetLoader")); }
-	virtual FText GetDefaultNodeTitle() const override { return NSLOCTEXT("PVPresetLoaderSettings", "NodeTitle", "Growth Data Loader"); }
-	virtual FText GetNodeTooltipText() const override { return NSLOCTEXT("PVPresetLoaderSettings", "NodeTooltip", "Load Procedural Vegetation Growth Data"); }
+	virtual FText GetDefaultNodeTitle() const override { return NSLOCTEXT("PVPresetLoaderSettings", "NodeTitle", "Procedural Vegetation Preset Loader"); }
+	virtual FText GetNodeTooltipText() const override { return NSLOCTEXT("PVPresetLoaderSettings", "NodeTooltip", "Load Procedural Vegetation Preset"); }
 
 	virtual TArray<EPVRenderType> GetRenderType() const override { return TArray{ EPVRenderType::Skeleton }; }
 #endif
@@ -27,8 +27,8 @@ protected:
 	virtual FPCGElementPtr CreateElement() const override;
 
 public:
-	UPROPERTY(EditAnywhere, Category="Preset Data Asset")
-	TObjectPtr<class UMegaSeedDataAsset> PresetDataAsset = nullptr;
+	UPROPERTY(EditAnywhere, Category="Preset")
+	TObjectPtr<class UProceduralVegetationPreset> Preset = nullptr;
 };
 
 class FPVPresetLoaderElement : public FPVBaseElement
diff --git a/Engine/Plugins/Experimental/ProceduralVegetationEditor/Source/ProceduralVegetation/Private/ProceduralVegetationLink.cpp b/Engine/Plugins/Experimental/ProceduralVegetationEditor/Source/ProceduralVegetation/Private/ProceduralVegetationLink.cpp
--- a/Engine/Plugins/Experimental/ProceduralVegetationEditor/Source/ProceduralVegetation/Private/ProceduralVegetationLink.cpp
+++ b/Engine/Plugins/Experimental/ProceduralVegetationEditor/Source/ProceduralVegetation/Private/ProceduralVegetationLink.cpp
@@ -6,7 +6,7 @@
 #include "Nodes/PVMeshBuilderSettings.h"
 #include "Nodes/PVPresetLoaderSettings.h"
 #include "Nodes/PVOutputSettings.h"
-#include "MegaSeedDataAsset.h"
+#include "ProceduralVegetationPreset.h"
 
 void SanitizeFolderName(FString& FolderName)
 {
@@ -34,14 +34,14 @@ void UProceduralVegetation::CreateGraph(const UProceduralVegetationGraph* InGrap
 }
 
 #if WITH_EDITOR
-void UProceduralVegetation::CreateGraphFromMegaSeed(const TObjectPtr<UMegaSeedDataAsset> MegaSeedAsset)
+void UProceduralVegetation::CreateGraphFromPreset(const TObjectPtr<UProceduralVegetationPreset> InPreset)
 {
 	Graph = NewObject<UProceduralVegetationGraph>(this, UProceduralVegetationGraph::StaticClass());
 
 	UPCGSettings* DefaultNodeSettings = nullptr;
 
 	UPVPresetLoaderSettings* PresetLoaderSettings = NewObject<UPVPresetLoaderSettings>(Graph);
-	PresetLoaderSettings->PresetDataAsset = MegaSeedAsset;
+	PresetLoaderSettings->Preset = InPreset;
 	UPCGNode* PresetLoaderNode = Graph->AddNodeCopy(PresetLoaderSettings, DefaultNodeSettings);
 
 	FInt32Vector2 Position(0, 300);
@@ -49,7 +49,7 @@ void UProceduralVegetation::CreateGraphFromMegaSeed(const TObjectPtr<UMegaSeedDa
 	PresetLoaderNode->SetNodePosition(Position.X, Position.Y);
 	Position.X += 100;
 
-	for (const auto& [VariantName, VariantData] : MegaSeedAsset->Variants)
+	for (const auto& [VariantName, VariantData] : InPreset->Variants)
 	{
 		const UPVGravitySettings* GravitySettings = NewObject<UPVGravitySettings>(Graph);
 		const UPVMeshBuilderSettings* MeshBuilderSettings = NewObject<UPVMeshBuilderSettings>(Graph);
diff --git a/Engine/Plugins/Experimental/ProceduralVegetationEditor/Source/ProceduralVegetation/Private/ProceduralVegetationPreset.cpp b/Engine/Plugins/Experimental/ProceduralVegetationEditor/Source/ProceduralVegetation/Private/ProceduralVegetationPreset.cpp
new file mode 100644
--- /dev/null
+++ b/Engine/Plugins/Experimental/ProceduralVegetationEditor/Source/ProceduralVegetation/Private/ProceduralVegetationPreset.cpp
@@ -0,0 +1,17 @@
+﻿// Copyright Epic Games, Inc. All Rights Reserved.
+
+#include "ProceduralVegetationPreset.h"
+
+#include "Helpers/PVJSONHelper.h"
+
+void UProceduralVegetationPreset::LoadFromVariantFiles(const TMap<FString, FString>& InVariantFiles)
+{
+	for (const auto& [VariantName, VariantFile] : InVariantFiles)
+	{
+		FManagedArrayCollection Collection;
+		if (FString OutError; PV::LoadMegaPlantsJsonToCollection(Collection, VariantFile, OutError))
+		{
+			Variants.Add(VariantName, Collection);
+		}
+	}
+}
diff --git a/Engine/Plugins/Experimental/ProceduralVegetationEditor/Source/ProceduralVegetation/Public/Facades/PVPlantFacade.h b/Engine/Plugins/Experimental/ProceduralVegetationEditor/Source/ProceduralVegetation/Public/Facades/PVPlantFacade.h
--- a/Engine/Plugins/Experimental/ProceduralVegetationEditor/Source/ProceduralVegetation/Public/Facades/PVPlantFacade.h
+++ b/Engine/Plugins/Experimental/ProceduralVegetationEditor/Source/ProceduralVegetation/Public/Facades/PVPlantFacade.h
@@ -9,8 +9,8 @@
 namespace PV::Facades
 {
 	/**
-	 * FPlantFacade is used to access and manipulate the data from Plants Group with in the MegaPlant's FManagedArrayCollection
-	 * A MegaPlant asset can consist of multiple plants. Each entry in the group references that plant's associated branch
+	 * FPlantFacade is used to access and manipulate the data from Plants Group with in the ProceduralVegetation's FManagedArrayCollection
+	 * A ProceduralVegetation asset can consist of multiple plants. Each entry in the group references that plant's associated branch
 	 * numbers and indices.
 	 */
 	class PROCEDURALVEGETATION_API FPlantFacade
diff --git a/Engine/Plugins/Experimental/ProceduralVegetationEditor/Source/ProceduralVegetation/Public/Helpers/PVJSONHelper.h b/Engine/Plugins/Experimental/ProceduralVegetationEditor/Source/ProceduralVegetation/Public/Helpers/PVJSONHelper.h
--- a/Engine/Plugins/Experimental/ProceduralVegetationEditor/Source/ProceduralVegetation/Public/Helpers/PVJSONHelper.h
+++ b/Engine/Plugins/Experimental/ProceduralVegetationEditor/Source/ProceduralVegetation/Public/Helpers/PVJSONHelper.h
@@ -375,17 +375,17 @@ namespace PV
 		return true;
 	}
 
-	static bool LoadMegaSeedJsonToCollection(FManagedArrayCollection& Collection, const FString& FilePath, FString& OutErrorMessage)
+	static bool LoadMegaPlantsJsonToCollection(FManagedArrayCollection& Collection, const FString& FilePath, FString& OutErrorMessage)
 	{
-		FString MegaSeedData;
-		if (!FFileHelper::LoadFileToString(MegaSeedData, *FilePath))
+		FString MegaPlantsData;
+		if (!FFileHelper::LoadFileToString(MegaPlantsData, *FilePath))
 		{
 			OutErrorMessage = TEXT("Failed to open skeleton file");
 			return false;
 		}
 
 		TSharedPtr<FJsonObject> LoadedData = MakeShareable(new FJsonObject);
-		if (!FJsonSerializer::Deserialize(TJsonReaderFactory<>::Create(MegaSeedData), LoadedData))
+		if (!FJsonSerializer::Deserialize(TJsonReaderFactory<>::Create(MegaPlantsData), LoadedData))
 		{
 			OutErrorMessage = TEXT("Failed to deserialize JSON file for skeleton");
 			return false;
diff --git a/Engine/Plugins/Experimental/ProceduralVegetationEditor/Source/ProceduralVegetation/Public/MegaSeedDataAsset.h b/Engine/Plugins/Experimental/ProceduralVegetationEditor/Source/ProceduralVegetation/Public/MegaSeedDataAsset.h
deleted file mode 100644
diff --git a/Engine/Plugins/Experimental/ProceduralVegetationEditor/Source/ProceduralVegetation/Public/ProceduralVegetation.h b/Engine/Plugins/Experimental/ProceduralVegetationEditor/Source/ProceduralVegetation/Public/ProceduralVegetation.h
--- a/Engine/Plugins/Experimental/ProceduralVegetationEditor/Source/ProceduralVegetation/Public/ProceduralVegetation.h
+++ b/Engine/Plugins/Experimental/ProceduralVegetationEditor/Source/ProceduralVegetation/Public/ProceduralVegetation.h
@@ -5,7 +5,7 @@
 #include "PCGGraph.h"
 #include "ProceduralVegetation.generated.h"
 
-class UMegaSeedDataAsset;
+class UProceduralVegetationPreset;
 
 UCLASS(MinimalAPI)
 class UProceduralVegetationGraph : public UPCGGraph
@@ -42,6 +42,6 @@ public:
 	void CreateGraph(const UProceduralVegetationGraph* InGraph = nullptr);
 
 #if WITH_EDITOR
-	void CreateGraphFromMegaSeed(const TObjectPtr<UMegaSeedDataAsset> MegaSeedAsset);
+	void CreateGraphFromPreset(const TObjectPtr<UProceduralVegetationPreset> InPreset);
 #endif
 };
diff --git a/Engine/Plugins/Experimental/ProceduralVegetationEditor/Source/ProceduralVegetation/Public/ProceduralVegetationPreset.h b/Engine/Plugins/Experimental/ProceduralVegetationEditor/Source/ProceduralVegetation/Public/ProceduralVegetationPreset.h
new file mode 100644
--- /dev/null
+++ b/Engine/Plugins/Experimental/ProceduralVegetationEditor/Source/ProceduralVegetation/Public/ProceduralVegetationPreset.h
@@ -0,0 +1,22 @@
+// Copyright Epic Games, Inc. All Rights Reserved.
+
+#pragma once
+
+#include "Engine/DataAsset.h"
+
+#include "GeometryCollection/ManagedArrayCollection.h"
+
+#include "ProceduralVegetationPreset.generated.h"
+
+UCLASS(BlueprintType)
+class PROCEDURALVEGETATION_API UProceduralVegetationPreset : public UDataAsset
+{
+	GENERATED_BODY()
+
+public:
+	UPROPERTY(EditAnywhere, Category="Preset Data")
+	TMap<FString, FManagedArrayCollection> Variants;
+
+public:
+	void LoadFromVariantFiles(const TMap<FString, FString>& InVariantFiles);
+};
diff --git a/Engine/Plugins/Experimental/ProceduralVegetationEditor/Source/ProceduralVegetationEditor/Private/AssetDefinition_ProceduralVegetation.cpp b/Engine/Plugins/Experimental/ProceduralVegetationEditor/Source/ProceduralVegetationEditor/Private/AssetDefinition_ProceduralVegetation.cpp
--- a/Engine/Plugins/Experimental/ProceduralVegetationEditor/Source/ProceduralVegetationEditor/Private/AssetDefinition_ProceduralVegetation.cpp
+++ b/Engine/Plugins/Experimental/ProceduralVegetationEditor/Source/ProceduralVegetationEditor/Private/AssetDefinition_ProceduralVegetation.cpp
@@ -6,7 +6,7 @@
 
 FText UAssetDefinition_ProceduralVegetation::GetAssetDisplayName() const
 {
-	return INVTEXT("Procedural Vegetation");
+	return NSLOCTEXT("AssetDefinition_ProceduralVegetation", "AssetDisplayName", "Procedural Vegetation");
 }
 
 FLinearColor UAssetDefinition_ProceduralVegetation::GetAssetColor() const
diff --git a/Engine/Plugins/Experimental/ProceduralVegetationEditor/Source/ProceduralVegetationEditor/Private/Importer/MegaSeedImportWorkflow.cpp b/Engine/Plugins/Experimental/ProceduralVegetationEditor/Source/ProceduralVegetationEditor/Private/Importer/MegaSeedImportWorkflow.cpp
--- a/Engine/Plugins/Experimental/ProceduralVegetationEditor/Source/ProceduralVegetationEditor/Private/Importer/MegaSeedImportWorkflow.cpp
+++ b/Engine/Plugins/Experimental/ProceduralVegetationEditor/Source/ProceduralVegetationEditor/Private/Importer/MegaSeedImportWorkflow.cpp
@@ -3,7 +3,7 @@
 #include "MegaSeedImportWorkflow.h"
 #include "FabDownloader.h"
 #include "MegaSeedImporter.h"
-#include "MegaSeedDataAsset.h"
+#include "ProceduralVegetationPreset.h"
 #include "AssetRegistry/IAssetRegistry.h"
 #include "Async/Async.h"
 #include "Dom/JsonObject.h"
@@ -86,7 +86,7 @@ void FMegaSeedImportWorkflow::ImportContent(const TArray<FString>& SourceFiles)
 	const FString MegaSeedPackageName = ImportLocation / MegaSeedName;
 
 	UPackage* MegaSeedPackage = CreatePackage(*MegaSeedPackageName);
-	UMegaSeedDataAsset* MegaSeedAsset = NewObject<UMegaSeedDataAsset>(MegaSeedPackage, *MegaSeedName, RF_Public | RF_Standalone);
+	UProceduralVegetationPreset* Preset = NewObject<UProceduralVegetationPreset>(MegaSeedPackage, *MegaSeedName, RF_Public | RF_Standalone);
 
 	TMap<FString, FString> VariantFiles;
 	for (const auto& Variation : MegaSeedVariations)
@@ -97,9 +97,9 @@ void FMegaSeedImportWorkflow::ImportContent(const TArray<FString>& SourceFiles)
 		const FString VariationPath = MegaSeedDownloadFolder / VariationUri;
 		VariantFiles.Add(VariationName, VariationPath);
 	}
-	MegaSeedAsset->LoadFromVariantFiles(VariantFiles);
+	Preset->LoadFromVariantFiles(VariantFiles);
 
-	IAssetRegistry::Get()->AssetCreated(MegaSeedAsset);
+	IAssetRegistry::Get()->AssetCreated(Preset);
 	MegaSeedPackage->SetDirtyFlag(true);
 
 	const FString ProceduralVegetationName = AssetImportFolder + "_Graph";
@@ -121,7 +121,7 @@ void FMegaSeedImportWorkflow::ImportContent(const TArray<FString>& SourceFiles)
 		const FString FoliagePath = MegaSeedDownloadFolder / FoliageUri;
 
 		const FString FoliageImportLocation = ImportLocation / "Foliage";
-		for (auto& [VariantName, VariantCollection] : MegaSeedAsset->Variants)
+		for (auto& [VariantName, VariantCollection] : Preset->Variants)
 		{
 			PV::Facades::FFoliageFacade FoliageFacade(VariantCollection);
 			FoliageFacade.SetFoliagePath(FoliageImportLocation);
@@ -133,21 +133,21 @@ void FMegaSeedImportWorkflow::ImportContent(const TArray<FString>& SourceFiles)
 		FMegaSeedImporter::ImportFoliage(
 			FoliagePath,
 			FoliageImportLocation,
-			[this, ProceduralVegetation, MegaSeedAsset, PreBakedMeshesPaths](const TArray<UObject*>& InImportedObjects)
+			[this, ProceduralVegetation, Preset, PreBakedMeshesPaths](const TArray<UObject*>& InImportedObjects)
 				{
 					AsyncTask(
 						ENamedThreads::GameThread,
-						[this, ProceduralVegetation, MegaSeedAsset]()
+						[this, ProceduralVegetation, Preset]()
 							{
 #if WITH_EDITOR
-								ProceduralVegetation->CreateGraphFromMegaSeed(MegaSeedAsset);
+								ProceduralVegetation->CreateGraphFromPreset(Preset);
 								ProceduralVegetation->MarkPackageDirty();
 #endif
 								ExpireImportNotification(true);
 								CompleteWorkflow();
 							}
 					);
-					for (auto& [VariantName, VariantCollection] : MegaSeedAsset->Variants)
+					for (auto& [VariantName, VariantCollection] : Preset->Variants)
 					{
 						PV::Facades::FFoliageFacade FoliageFacade(VariantCollection);
 						for (int32 FoliageIndex = 0; FoliageIndex < FoliageFacade.NumFoliageNames(); FoliageIndex++)
diff --git a/Engine/Plugins/Experimental/ProceduralVegetationEditor/Source/ProceduralVegetationEditor/Private/PVEditor.cpp b/Engine/Plugins/Experimental/ProceduralVegetationEditor/Source/ProceduralVegetationEditor/Private/PVEditor.cpp
--- a/Engine/Plugins/Experimental/ProceduralVegetationEditor/Source/ProceduralVegetationEditor/Private/PVEditor.cpp
+++ b/Engine/Plugins/Experimental/ProceduralVegetationEditor/Source/ProceduralVegetationEditor/Private/PVEditor.cpp
@@ -104,7 +104,7 @@ FName FPVEditor::GetToolkitFName() const
 
 FText FPVEditor::GetBaseToolkitName() const
 {
-	return INVTEXT("Procedural Vegetation Editor");
+	return LOCTEXT("PVEditorToolkitName", "Procedural Vegetation Editor");
 }
 
 FText FPVEditor::GetToolkitToolTipText() const
@@ -169,14 +169,14 @@ IPCGBaseSubsystem* FPVEditor::GetSubsystem() const
 TAttribute<FGraphAppearanceInfo> FPVEditor::GetAppearanceInfo() const
 {
 	FGraphAppearanceInfo AppearanceInfo;
-	AppearanceInfo.CornerText = INVTEXT("PVE Graph");
+	AppearanceInfo.CornerText = LOCTEXT("PVEditorCornerText", "Procedural Vegetation");
 
 	return AppearanceInfo;
 }
 
 TSharedRef<FTabManager::FLayout> FPVEditor::GetDefaultLayout() const
 {
-	return FTabManager::NewLayout("Standalone_PVPCGGraphEditor_DefaultLayout_v1.0.4")
+	return FTabManager::NewLayout("Standalone_PVGraphEditor_DefaultLayout_v1.0.4")
 	->AddArea // Main PCG Graph Editor Area
 	(
 		FTabManager::NewPrimaryArea()->SetOrientation(Orient_Vertical)
diff --git a/Engine/Plugins/Experimental/ProceduralVegetationEditor/Source/ProceduralVegetationEditor/Private/PVEditorCommands.cpp b/Engine/Plugins/Experimental/ProceduralVegetationEditor/Source/ProceduralVegetationEditor/Private/PVEditorCommands.cpp
--- a/Engine/Plugins/Experimental/ProceduralVegetationEditor/Source/ProceduralVegetationEditor/Private/PVEditorCommands.cpp
+++ b/Engine/Plugins/Experimental/ProceduralVegetationEditor/Source/ProceduralVegetationEditor/Private/PVEditorCommands.cpp
@@ -3,12 +3,12 @@
 #include "PVEditorCommands.h"
 #include "Styling/AppStyle.h"
 
-#define LOCTEXT_NAMESPACE "PVPCGEditorCommands"
+#define LOCTEXT_NAMESPACE "PVEditorCommands"
 
 FPVEditorCommands::FPVEditorCommands()
 	: TCommands<FPVEditorCommands>(
 		"PVPCGEditor",
-		NSLOCTEXT("Contexts", "PVPCGEditor", "Procedural Vegetation Editor"),
+		LOCTEXT("ContextDesc", "Procedural Vegetation Editor"),
 		NAME_None,
 		FAppStyle::GetAppStyleSetName())
 {
diff --git a/Engine/Plugins/Experimental/ProceduralVegetationEditor/Source/ProceduralVegetationEditor/Private/ProceduralVegetationEditorModule.cpp b/Engine/Plugins/Experimental/ProceduralVegetationEditor/Source/ProceduralVegetationEditor/Private/ProceduralVegetationEditorModule.cpp
--- a/Engine/Plugins/Experimental/ProceduralVegetationEditor/Source/ProceduralVegetationEditor/Private/ProceduralVegetationEditorModule.cpp
+++ b/Engine/Plugins/Experimental/ProceduralVegetationEditor/Source/ProceduralVegetationEditor/Private/ProceduralVegetationEditorModule.cpp
@@ -17,7 +17,7 @@
 #include "DataTypes/PVFoliageMeshData.h"
 #include "Selection.h"
 #include "Animation/SkeletalMeshActor.h"
-#include "Visualizations/PVPCGDataVisualization.h"
+#include "Visualizations/PVDataVisualization.h"
 #include "Engine/StaticMeshActor.h"
 #include "ToolMenus.h"
 #include "ToolMenuContext.h"
@@ -46,7 +46,7 @@ void FProceduralVegetationEditorModule::StartupModule()
 
 	// Visualizations
 	FPCGDataVisualizationRegistry& DataVisRegistry = FPCGModule::GetMutablePCGDataVisualizationRegistry();
-	DataVisRegistry.RegisterPCGDataVisualization(UPVData::StaticClass(), MakeUnique<const FPVPCGDataVisualization>());
+	DataVisRegistry.RegisterPCGDataVisualization(UPVData::StaticClass(), MakeUnique<const FPVDataVisualization>());
 
 	// FAB
 	MegaSeedImportWorkflowFactory = MakeShared<FMegaSeedImportWorkflowFactory>();
diff --git a/Engine/Plugins/Experimental/ProceduralVegetationEditor/Source/ProceduralVegetationEditor/Private/Visualizations/DebugVisualization/PVDebugVisualizationBase.cpp b/Engine/Plugins/Experimental/ProceduralVegetationEditor/Source/ProceduralVegetationEditor/Private/Visualizations/DebugVisualization/PVDebugVisualizationBase.cpp
--- a/Engine/Plugins/Experimental/ProceduralVegetationEditor/Source/ProceduralVegetationEditor/Private/Visualizations/DebugVisualization/PVDebugVisualizationBase.cpp
+++ b/Engine/Plugins/Experimental/ProceduralVegetationEditor/Source/ProceduralVegetationEditor/Private/Visualizations/DebugVisualization/PVDebugVisualizationBase.cpp
@@ -14,9 +14,9 @@
 #include "Components/TextRenderComponent.h"
 #include "Visualizations/PVLineBatchComponent.h"
 
-#define LOCTEXT_NAMESPACE "MegaPlantDebugVisualization"
+#define LOCTEXT_NAMESPACE "PVDebugVisualization"
 
-DEFINE_LOG_CATEGORY(LogMegaPlantDebugVisualization);
+DEFINE_LOG_CATEGORY(LogPVDebugVisualization);
 
 const float FPVDebugVisualizationBase::DefaultDirectionStrength = 20;
 const float FPVDebugVisualizationBase::DefaultTextScale = 2.5;
diff --git a/Engine/Plugins/Experimental/ProceduralVegetationEditor/Source/ProceduralVegetationEditor/Private/Visualizations/DebugVisualization/PVDebugVisualizationBase.h b/Engine/Plugins/Experimental/ProceduralVegetationEditor/Source/ProceduralVegetationEditor/Private/Visualizations/DebugVisualization/PVDebugVisualizationBase.h
--- a/Engine/Plugins/Experimental/ProceduralVegetationEditor/Source/ProceduralVegetationEditor/Private/Visualizations/DebugVisualization/PVDebugVisualizationBase.h
+++ b/Engine/Plugins/Experimental/ProceduralVegetationEditor/Source/ProceduralVegetationEditor/Private/Visualizations/DebugVisualization/PVDebugVisualizationBase.h
@@ -47,6 +47,6 @@ protected:
 };
 typedef TSharedPtr<FPVDebugVisualizationBase> FPVDebugVisualizationPtr;
 
-DECLARE_LOG_CATEGORY_EXTERN(LogMegaPlantDebugVisualization, Log, All);
+DECLARE_LOG_CATEGORY_EXTERN(LogPVDebugVisualization, Log, All);
 
 #endif
\ No newline at end of file
diff --git a/Engine/Plugins/Experimental/ProceduralVegetationEditor/Source/ProceduralVegetationEditor/Private/Visualizations/DebugVisualization/PVDebugVisualizations.cpp b/Engine/Plugins/Experimental/ProceduralVegetationEditor/Source/ProceduralVegetationEditor/Private/Visualizations/DebugVisualization/PVDebugVisualizations.cpp
--- a/Engine/Plugins/Experimental/ProceduralVegetationEditor/Source/ProceduralVegetationEditor/Private/Visualizations/DebugVisualization/PVDebugVisualizations.cpp
+++ b/Engine/Plugins/Experimental/ProceduralVegetationEditor/Source/ProceduralVegetationEditor/Private/Visualizations/DebugVisualization/PVDebugVisualizations.cpp
@@ -8,7 +8,7 @@
 #include "Facades/PVFoliageFacade.h"
 #include "Facades/PVPointFacade.h"
 
-TArray<FVector3f> FMegaPlantPointDebugVisualization::GetPivotPositions(const FManagedArrayCollection& InCollection)
+TArray<FVector3f> FPVPointDebugVisualization::GetPivotPositions(const FManagedArrayCollection& InCollection)
 {
 	const PV::Facades::FPointFacade PointFacade(InCollection);
 	const TManagedArray<FVector3f>& PointPositions = PointFacade.GetPositions();
@@ -24,7 +24,7 @@ TArray<FVector3f> FMegaPlantPointDebugVisualization::GetPivotPositions(const FMa
 	return Positions;
 }
 
-void FMegaPlantPointDebugVisualization::GetPivot(const FManagedArrayCollection& InCollection, const int InIndex, FVector3f& OutPos, float& OutScale)
+void FPVPointDebugVisualization::GetPivot(const FManagedArrayCollection& InCollection, const int InIndex, FVector3f& OutPos, float& OutScale)
 {
 	const PV::Facades::FPointFacade PointFacade(InCollection);
 
@@ -32,7 +32,7 @@ void FMegaPlantPointDebugVisualization::GetPivot(const FManagedArrayCollection&
 	OutScale = PointFacade.GetPointScale(InIndex);
 }
 
-TArray<FVector3f> FMegaPlantFoliageDebugVisualization::GetPivotPositions(const FManagedArrayCollection& InCollection)
+TArray<FVector3f> FPVFoliageDebugVisualization::GetPivotPositions(const FManagedArrayCollection& InCollection)
 {
 	const PV::Facades::FFoliageFacade FoliageFacade(InCollection);
 	const TManagedArray<FVector3f>& PointPositions =  FoliageFacade.GetPivotPositions();
@@ -49,7 +49,7 @@ TArray<FVector3f> FMegaPlantFoliageDebugVisualization::GetPivotPositions(const F
 
 }
 
-void FMegaPlantFoliageDebugVisualization::GetPivot(const FManagedArrayCollection& InCollection, const int InIndex, FVector3f& OutPos, float& OutScale)
+void FPVFoliageDebugVisualization::GetPivot(const FManagedArrayCollection& InCollection, const int InIndex, FVector3f& OutPos, float& OutScale)
 {
 	const PV::Facades::FFoliageFacade FoliageFacade(InCollection);
 	const PV::Facades::FBranchFacade BranchFacade(InCollection);
@@ -70,7 +70,7 @@ void FMegaPlantFoliageDebugVisualization::GetPivot(const FManagedArrayCollection
 	}
 }
 
-TArray<FVector3f> FMegaPlantBranchDebugVisualization::GetPivotPositions(const FManagedArrayCollection& InCollection)
+TArray<FVector3f> FPVBranchDebugVisualization::GetPivotPositions(const FManagedArrayCollection& InCollection)
 {
 	const PV::Facades::FBranchFacade BranchFacade(InCollection);
 	const PV::Facades::FPointFacade PointFacade(InCollection);
@@ -90,7 +90,7 @@ TArray<FVector3f> FMegaPlantBranchDebugVisualization::GetPivotPositions(const FM
 	return Positions;
 }
 
-void FMegaPlantBranchDebugVisualization::GetPivot(const FManagedArrayCollection& InCollection, const int InIndex, FVector3f& OutPos, float& OutScale)
+void FPVBranchDebugVisualization::GetPivot(const FManagedArrayCollection& InCollection, const int InIndex, FVector3f& OutPos, float& OutScale)
 {
 	const PV::Facades::FBranchFacade BranchFacade(InCollection);
 	const PV::Facades::FPointFacade PointFacade(InCollection);
diff --git a/Engine/Plugins/Experimental/ProceduralVegetationEditor/Source/ProceduralVegetationEditor/Private/Visualizations/DebugVisualization/PVDebugVisualizations.h b/Engine/Plugins/Experimental/ProceduralVegetationEditor/Source/ProceduralVegetationEditor/Private/Visualizations/DebugVisualization/PVDebugVisualizations.h
--- a/Engine/Plugins/Experimental/ProceduralVegetationEditor/Source/ProceduralVegetationEditor/Private/Visualizations/DebugVisualization/PVDebugVisualizations.h
+++ b/Engine/Plugins/Experimental/ProceduralVegetationEditor/Source/ProceduralVegetationEditor/Private/Visualizations/DebugVisualization/PVDebugVisualizations.h
@@ -11,28 +11,28 @@
 
 struct FManagedArrayCollection;
 
-class FMegaPlantPointDebugVisualization final : public FPVDebugVisualizationBase
+class FPVPointDebugVisualization final : public FPVDebugVisualizationBase
 {
 public:
-	FMegaPlantPointDebugVisualization(){};
+	FPVPointDebugVisualization(){};
 protected:
 	virtual TArray<FVector3f> GetPivotPositions(const FManagedArrayCollection& InCollection) override;
 	virtual void GetPivot(const FManagedArrayCollection& InCollection, const int InIndex, FVector3f& OutPos, float& OutScale) override;
 };
 
-class FMegaPlantFoliageDebugVisualization final : public FPVDebugVisualizationBase
+class FPVFoliageDebugVisualization final : public FPVDebugVisualizationBase
 {
 public:
-	FMegaPlantFoliageDebugVisualization(){};
+	FPVFoliageDebugVisualization(){};
 protected:
 	virtual TArray<FVector3f> GetPivotPositions(const FManagedArrayCollection& InCollection) override;
 	virtual void GetPivot(const FManagedArrayCollection& InCollection, const int InIndex, FVector3f& OutPos, float& OutScale) override;
 };
 
-class FMegaPlantBranchDebugVisualization final : public FPVDebugVisualizationBase
+class FPVBranchDebugVisualization final : public FPVDebugVisualizationBase
 {
 public:
-	FMegaPlantBranchDebugVisualization(){};
+	FPVBranchDebugVisualization(){};
 protected:
 	virtual TArray<FVector3f> GetPivotPositions(const FManagedArrayCollection& InCollection) override;
 	virtual void GetPivot(const FManagedArrayCollection& InCollection, const int InIndex, FVector3f& OutPos, float& OutScale) override;
diff --git a/Engine/Plugins/Experimental/ProceduralVegetationEditor/Source/ProceduralVegetationEditor/Private/Visualizations/DebugVisualization/PVDebugVisualizer.cpp b/Engine/Plugins/Experimental/ProceduralVegetationEditor/Source/ProceduralVegetationEditor/Private/Visualizations/DebugVisualization/PVDebugVisualizer.cpp
--- a/Engine/Plugins/Experimental/ProceduralVegetationEditor/Source/ProceduralVegetationEditor/Private/Visualizations/DebugVisualization/PVDebugVisualizer.cpp
+++ b/Engine/Plugins/Experimental/ProceduralVegetationEditor/Source/ProceduralVegetationEditor/Private/Visualizations/DebugVisualization/PVDebugVisualizer.cpp
@@ -40,13 +40,13 @@ FPVDebugVisualizationPtr FPVDebugVisualizer::CreateVisualizer(EPVDebugType InDeb
 	switch (InDebugType)
 	{
 		case EPVDebugType::Point:
-			return MakeShared<FMegaPlantPointDebugVisualization>();
+			return MakeShared<FPVPointDebugVisualization>();
 
 		case EPVDebugType::Foliage:
-			return MakeShared<FMegaPlantFoliageDebugVisualization>();
+			return MakeShared<FPVFoliageDebugVisualization>();
 
 		case EPVDebugType::Branches:
-			return MakeShared<FMegaPlantBranchDebugVisualization>();
+			return MakeShared<FPVBranchDebugVisualization>();
 
 		default:
 			return nullptr;
diff --git a/Engine/Plugins/Experimental/ProceduralVegetationEditor/Source/ProceduralVegetationEditor/Private/Visualizations/PVDataVisualization.cpp b/Engine/Plugins/Experimental/ProceduralVegetationEditor/Source/ProceduralVegetationEditor/Private/Visualizations/PVDataVisualization.cpp
new file mode 100644
--- /dev/null
+++ b/Engine/Plugins/Experimental/ProceduralVegetationEditor/Source/ProceduralVegetationEditor/Private/Visualizations/PVDataVisualization.cpp
@@ -0,0 +1,227 @@
+// Copyright Epic Games, Inc. All Rights Reserved.
+
+#include "PVDataVisualization.h"
+#include "AdvancedPreviewScene.h"
+#include "GeometryCollectionToDynamicMesh.h"
+#include "PVBoneComponent.h"
+#include "PCGComponent.h"
+#include "PCGContext.h"
+#include "PCGDataVisualization.h"
+#include "PCGEditorSettings.h"
+#include "PlanarCut.h"
+#include "Components/DynamicMeshComponent.h"
+#include "Components/InstancedStaticMeshComponent.h"
+#include "Components/LineBatchComponent.h"
+#include "Facades/PVFoliageFacade.h"
+#include "GeometryCollection/GeometryCollectionObject.h"
+#include "GeometryCollection/Facades/CollectionRenderingFacade.h"
+#include "Helpers/PVUtilities.h"
+#include "DataTypes/PVData.h"
+#include "Components/TextRenderComponent.h"
+#include "DebugVisualization/PVDebugVisualizer.h"
+#include "Visualizations/PVSkeletonVisualizerComponent.h"
+
+#define LOCTEXT_NAMESPACE "PCGManagedArrayCollectionDataVisualization"
+
+class UPVBoneComponent;
+
+DEFINE_LOG_CATEGORY(LogProceduralVegetationDataVisualization);
+
+FPVDataVisualization::FPVDataVisualization()
+{
+	UPCGEditorSettings* PCGEditorSettings = GetMutableDefault<UPCGEditorSettings>();
+	PCGEditorSettings->bAutoFocusViewport = false;
+
+	RegisterVisualizations();
+}
+
+void FPVDataVisualization::RegisterVisualizations()
+{
+	RenderMap.Add(EPVRenderType::Skeleton, FPVRenderCallback::CreateRaw(this, &FPVDataVisualization::SkeletonRenderer));
+	RenderMap.Add(EPVRenderType::Mesh, FPVRenderCallback::CreateRaw(this, &FPVDataVisualization::MeshRenderer));
+	RenderMap.Add(EPVRenderType::Foliage, FPVRenderCallback::CreateRaw(this, &FPVDataVisualization::FoliageRenderer));
+	RenderMap.Add(EPVRenderType::Bones, FPVRenderCallback::CreateRaw(this, &FPVDataVisualization::BonesRenderer));
+}
+
+void FPVDataVisualization::ExecuteDebugDisplay(FPCGContext* Context, const UPCGSettingsInterface* SettingsInterface, const UPCGData* Data,
+                                                         AActor* TargetActor) const
+{
+	return;
+}
+
+FPCGSetupSceneFunc FPVDataVisualization::GetViewportSetupFunc(const UPCGSettingsInterface* SettingsInterface, const UPCGData* Data) const
+{
+	return [this, WeakData=TWeakObjectPtr<const UPVData>(Cast<UPVData>(Data)), WeakSettings = TWeakObjectPtr<const UPCGSettingsInterface>(Cast<UPCGSettingsInterface>(SettingsInterface))](
+		FPCGSceneSetupParams& InOutParams)
+	{
+		check(InOutParams.Scene);
+
+		if (!WeakData.IsValid())
+		{
+			UE_LOG(LogProceduralVegetationDataVisualization, Error, TEXT("Failed to setup data viewport, the data was lost or invalid."));
+			return;
+		}
+
+		const FManagedArrayCollection& Collection = WeakData->GetCollection();
+
+		FBoxSphereBounds Bounds(ForceInit);
+
+		for (const EPVRenderType& RenderType : WeakData->GetRenderType())
+		{
+			if (RenderMap.Contains(RenderType))
+			{
+				InOutParams.Scene->GetLineBatcher()->Flush();
+				RenderMap[RenderType].Execute(Collection, InOutParams, Bounds);
+			}
+		}
+
+		if (!InOutParams.Scene->IsGridEnabled())
+		{
+			InOutParams.Scene->HandleToggleGrid();
+		}
+		if (!InOutParams.Scene->IsEnvironmentEnabled())
+		{
+			InOutParams.Scene->HandleToggleEnvironment();
+		}
+		if (!InOutParams.Scene->IsPostProcessingEnabled())
+		{
+			InOutParams.Scene->HandleTogglePostProcessing();
+		}
+
+		UStaticMesh* MannequinMesh = LoadObject<UStaticMesh>(
+			nullptr,
+			TEXT("/ProceduralVegetationEditor/Mannequin/Viewport_Mannequin_T_Pose.Viewport_Mannequin_T_Pose")
+		);
+
+		constexpr bool bShowMannequin = true;
+		if (bShowMannequin)
+		{
+			UStaticMeshComponent* MannequinComponent = NewObject<UStaticMeshComponent>(GetTransientPackage(), NAME_None, RF_Transient);
+			MannequinComponent->SetStaticMesh(MannequinMesh);
+
+			InOutParams.ManagedResources.Add(MannequinComponent);
+			InOutParams.Scene->AddComponent(MannequinComponent, FTransform(FVector(-Bounds.BoxExtent.X * 2.0, 0, 0)));
+		}
+
+#if PV_SHOW_DEBUG_VISUALIZATION
+			if (WeakSettings.IsValid() && WeakSettings->bDebug)
+			{
+				FPVDebugVisualizer::Draw(WeakData.Get(), InOutParams);
+			}
+#endif
+
+		InOutParams.Scene->SetFloorVisibility(true);
+		InOutParams.Scene->SetFloorOffset(1.0f);
+
+		InOutParams.FocusBounds = FBoxSphereBounds(Bounds.Origin, Bounds.BoxExtent, Bounds.SphereRadius);
+	};
+}
+
+void FPVDataVisualization::MeshRenderer(const FManagedArrayCollection& InCollection, FPCGSceneSetupParams& InOutParams,
+                                                  FBoxSphereBounds& OutBounds)
+{
+	TRACE_CPUPROFILER_EVENT_SCOPE(FPVDataVisualization::MeshRenderer);
+
+	TArray<UMaterialInterface*> Materials;
+	
+	if (InCollection.HasAttribute("MaterialPath", FGeometryCollection::MaterialGroup))
+	{
+		const TManagedArray<FString>& MaterialPaths = InCollection.GetAttribute<FString>("MaterialPath", FGeometryCollection::MaterialGroup);
+		for (const FString& MaterialPath : MaterialPaths.GetConstArray())
+		{
+			Materials.Add(LoadObject<UMaterialInterface>(nullptr, MaterialPath));
+		}
+	}
+
+	UE::Geometry::FGeometryCollectionToDynamicMeshes Converter;
+	Converter.Init(InCollection, {});
+	for (const auto& MeshInfo : Converter.Meshes)
+	{
+		FDynamicMesh3* DynamicMesh = MeshInfo.Mesh.Get();
+		if (DynamicMesh->VertexCount())
+		{
+			TObjectPtr<UDynamicMeshComponent> DynamicMeshComponent = NewObject<UDynamicMeshComponent>(GetTransientPackage(), NAME_None, RF_Transient);
+			DynamicMeshComponent->SetMesh(MoveTemp(*DynamicMesh));
+			DynamicMeshComponent->SetTangentsType(EDynamicMeshComponentTangentsMode::AutoCalculated);
+			DynamicMeshComponent->ConfigureMaterialSet(Materials);
+			DynamicMeshComponent->UpdateBounds();
+
+			InOutParams.ManagedResources.Add(DynamicMeshComponent);
+			InOutParams.Scene->AddComponent(DynamicMeshComponent, FTransform::Identity);
+
+			OutBounds = OutBounds + DynamicMeshComponent->CalcLocalBounds();
+		}
+	}
+}
+
+void FPVDataVisualization::SkeletonRenderer(const FManagedArrayCollection& InCollection, FPCGSceneSetupParams& InOutParams,
+                                                      FBoxSphereBounds& OutBounds)
+{
+	TRACE_CPUPROFILER_EVENT_SCOPE(FPVDataVisualization::SkeletonRenderer);
+
+	const TObjectPtr<UPVSkeletonVisualizerComponent> VisualizerComponent = NewObject<UPVSkeletonVisualizerComponent>(
+		GetTransientPackage(), NAME_None, RF_Transient);
+	VisualizerComponent->SetCollection(&InCollection);
+	VisualizerComponent->UpdateBounds();
+
+	InOutParams.ManagedResources.Add(VisualizerComponent);
+	InOutParams.Scene->AddComponent(VisualizerComponent, FTransform::Identity);
+
+	OutBounds = OutBounds + VisualizerComponent->CalcLocalBounds();
+}
+
+void FPVDataVisualization::FoliageRenderer(const FManagedArrayCollection& InCollection, FPCGSceneSetupParams& InOutParams,
+                                                     FBoxSphereBounds& OutBounds)
+{
+	TRACE_CPUPROFILER_EVENT_SCOPE(FPVDataVisualization::FoliageRenderer);
+
+	PV::Facades::FFoliageFacade FoliageFacade(InCollection);
+
+	auto OnComponentCreated = [&InOutParams](UMeshComponent* InInstancedComponent)
+	{
+		InOutParams.ManagedResources.Add(InInstancedComponent);
+		InInstancedComponent->RegisterComponentWithWorld(InOutParams.Scene->GetWorld());
+		InOutParams.Scene->AddComponent(InInstancedComponent, FTransform::Identity);
+	};
+
+	TMap<FString, TObjectPtr<UMeshComponent>> OutInstancedComponentMap;
+
+	PV::Utilities::AddFoliageInstances(FoliageFacade, nullptr,
+	                                           PV::Utilities::FFoliageComponentCreatedCallback::CreateLambda(OnComponentCreated),
+	                                           OutInstancedComponentMap);
+}
+
+void FPVDataVisualization::BonesRenderer(const FManagedArrayCollection& InCollection, FPCGSceneSetupParams& InOutParams,
+	FBoxSphereBounds& OutBounds)
+{
+	TRACE_CPUPROFILER_EVENT_SCOPE(FPVDataVisualization::BonesRenderer);
+
+	const TObjectPtr<UPVBoneComponent> BoneComponent = NewObject<UPVBoneComponent>(
+		GetTransientPackage(), NAME_None, RF_Transient);
+	BoneComponent->SetCollection(&InCollection);
+	BoneComponent->UpdateBounds();
+
+	InOutParams.ManagedResources.Add(BoneComponent);
+	InOutParams.Scene->AddComponent(BoneComponent, FTransform::Identity);
+
+	FString InTextToDraw = FString::Format(TEXT("{0} Bones"),{BoneComponent->GetBoneCount()});
+	const FVector3f& Pos = (FVector3f::ForwardVector * 100);
+
+	DrawDebugString(BoneComponent->GetWorld(), FVector::Zero(), InTextToDraw);
+	
+	TObjectPtr<UTextRenderComponent> Component = NewObject<UTextRenderComponent>(GetTransientPackage(), NAME_None, RF_Transient);
+	Component->SetText(FText::FromString(InTextToDraw));
+	Component->SetTextRenderColor(FColor::Blue);
+	Component->SetWorldSize(40);
+	Component->SetGenerateOverlapEvents(false);
+						
+	InOutParams.ManagedResources.Add(Component);
+	FTransform TextTransform = FTransform(FVector(Pos.X, Pos.Y, Pos.Z));
+	FRotator TextRotator = FRotator(0,90,0);
+	TextTransform.SetRotation(TextRotator.Quaternion());
+	InOutParams.Scene->AddComponent(Component, TextTransform);
+
+	OutBounds = OutBounds + BoneComponent->CalcLocalBounds();
+}
+
+#undef LOCTEXT_NAMESPACE
diff --git a/Engine/Plugins/Experimental/ProceduralVegetationEditor/Source/ProceduralVegetationEditor/Private/Visualizations/PVDataVisualization.h b/Engine/Plugins/Experimental/ProceduralVegetationEditor/Source/ProceduralVegetationEditor/Private/Visualizations/PVDataVisualization.h
new file mode 100644
--- /dev/null
+++ b/Engine/Plugins/Experimental/ProceduralVegetationEditor/Source/ProceduralVegetationEditor/Private/Visualizations/PVDataVisualization.h
@@ -0,0 +1,36 @@
+// Copyright Epic Games, Inc. All Rights Reserved.
+
+#pragma once
+
+#include "PCGDataVisualization.h"
+#include "DataVisualizations/PCGSpatialDataVisualization.h"
+#include "DataTypes/PVData.h"
+
+class AActor;
+class UPCGData;
+struct FPCGContext;
+
+DECLARE_DELEGATE_ThreeParams(FPVRenderCallback, const FManagedArrayCollection&, FPCGSceneSetupParams&, FBoxSphereBounds&);
+
+class FPVDataVisualization : public IPCGSpatialDataVisualization
+{
+public:
+	FPVDataVisualization();
+	
+	// ~Begin IPCGDataVisualization interface
+	virtual void ExecuteDebugDisplay(FPCGContext* Context, const UPCGSettingsInterface* SettingsInterface, const UPCGData* Data, AActor* TargetActor) const override;
+	virtual FPCGSetupSceneFunc GetViewportSetupFunc(const UPCGSettingsInterface* SettingsInterface, const UPCGData* Data) const override;
+	// ~End IPCGDataVisualization interface
+
+protected:
+	TMap<EPVRenderType, FPVRenderCallback> RenderMap;
+	
+	void RegisterVisualizations();
+
+	void MeshRenderer(const FManagedArrayCollection& InCollection, FPCGSceneSetupParams& InOutParams, FBoxSphereBounds& OutBounds);
+	void SkeletonRenderer(const FManagedArrayCollection& InCollection, FPCGSceneSetupParams& InOutParams, FBoxSphereBounds& OutBounds);
+	void FoliageRenderer(const FManagedArrayCollection& InCollection, FPCGSceneSetupParams& InOutParams, FBoxSphereBounds& OutBounds);
+	void BonesRenderer(const FManagedArrayCollection& InCollection, FPCGSceneSetupParams& InOutParams, FBoxSphereBounds& OutBounds);
+};
+
+DECLARE_LOG_CATEGORY_EXTERN(LogProceduralVegetationDataVisualization, Log, All);
\ No newline at end of file
diff --git a/Engine/Plugins/Experimental/ProceduralVegetationEditor/Source/ProceduralVegetationEditor/Private/Visualizations/PVPCGDataVisualization.cpp b/Engine/Plugins/Experimental/ProceduralVegetationEditor/Source/ProceduralVegetationEditor/Private/Visualizations/PVPCGDataVisualization.cpp
deleted file mode 100644
diff --git a/Engine/Plugins/Experimental/ProceduralVegetationEditor/Source/ProceduralVegetationEditor/Private/Visualizations/PVPCGDataVisualization.h b/Engine/Plugins/Experimental/ProceduralVegetationEditor/Source/ProceduralVegetationEditor/Private/Visualizations/PVPCGDataVisualization.h
deleted file mode 100644
