[PVE] Major changes
- Remove FAB dependency and import workflows
- Add PVEFloatRamp with customizations
- Remove extra/redundant properties in PVE detail panel
- Mesher node name now changes with the export mesh name

diff --git a/Engine/Plugins/Experimental/ProceduralVegetationEditor/Config/BaseProceduralVegetationEditor.ini b/Engine/Plugins/Experimental/ProceduralVegetationEditor/Config/BaseProceduralVegetationEditor.ini
--- a/Engine/Plugins/Experimental/ProceduralVegetationEditor/Config/BaseProceduralVegetationEditor.ini
+++ b/Engine/Plugins/Experimental/ProceduralVegetationEditor/Config/BaseProceduralVegetationEditor.ini
@@ -55,3 +55,5 @@
 +StructRedirects=(OldName="/Script/ProceduralVegetation.PVExportSettings",NewName="/Script/ProceduralVegetation.PVExportParams")
 
 +PropertyRedirects=(OldName="PVPCGMegaSeedLoadSettings.RecipeDataAsset",NewName="PVPresetLoaderSettings.Preset")
+
++StructRedirects=(OldName="/Script/Engine.RuntimeFloatCurve",NewName="/Script/ProceduralVegetation.PVFloatRamp")
\ No newline at end of file
diff --git a/Engine/Plugins/Experimental/ProceduralVegetationEditor/ProceduralVegetationEditor.uplugin b/Engine/Plugins/Experimental/ProceduralVegetationEditor/ProceduralVegetationEditor.uplugin
--- a/Engine/Plugins/Experimental/ProceduralVegetationEditor/ProceduralVegetationEditor.uplugin
+++ b/Engine/Plugins/Experimental/ProceduralVegetationEditor/ProceduralVegetationEditor.uplugin
@@ -50,10 +50,6 @@
 			"Name": "PCG",
 			"Enabled": true
 		},
-		{
-			"Name": "Fab",
-			"Enabled": true
-		},
 		{
 			"Name": "DynamicWind",
 			"Enabled": true
diff --git a/Engine/Plugins/Experimental/ProceduralVegetationEditor/Source/ProceduralVegetation/Private/Implementations/PVFoliage.cpp b/Engine/Plugins/Experimental/ProceduralVegetationEditor/Source/ProceduralVegetation/Private/Implementations/PVFoliage.cpp
--- a/Engine/Plugins/Experimental/ProceduralVegetationEditor/Source/ProceduralVegetation/Private/Implementations/PVFoliage.cpp
+++ b/Engine/Plugins/Experimental/ProceduralVegetationEditor/Source/ProceduralVegetation/Private/Implementations/PVFoliage.cpp
@@ -3,15 +3,17 @@
 #include "PVFoliage.h"
 
 #include "ProceduralVegetationModule.h"
+#include "PVFloatRamp.h"
+
 #include "Facades/PVAttributesNames.h"
 #include "Facades/PVBranchFacade.h"
 #include "Facades/PVFoliageFacade.h"
 #include "Facades/PVPointFacade.h"
 
 void FPVFoliage::DistributeFoliage(FManagedArrayCollection& OutCollection, const FManagedArrayCollection& SourceCollection,
-                                          const FDistributionSettings& DistributionSettings, const FScaleSettings& ScaleSettings,
-                                          const FVectorSettings& VectorSettings, const FPhyllotaxySettings& PhyllotaxySettings,
-                                          const FMiscSettings& MiscSettings)
+                                   const FDistributionSettings& DistributionSettings, const FScaleSettings& ScaleSettings,
+                                   const FVectorSettings& VectorSettings, const FPhyllotaxySettings& PhyllotaxySettings,
+                                   const FMiscSettings& MiscSettings)
 {
 	PV::Facades::FFoliageFacade FoliageFacadeOutput(OutCollection);
 
diff --git a/Engine/Plugins/Experimental/ProceduralVegetationEditor/Source/ProceduralVegetation/Private/Implementations/PVFoliage.h b/Engine/Plugins/Experimental/ProceduralVegetationEditor/Source/ProceduralVegetation/Private/Implementations/PVFoliage.h
--- a/Engine/Plugins/Experimental/ProceduralVegetationEditor/Source/ProceduralVegetation/Private/Implementations/PVFoliage.h
+++ b/Engine/Plugins/Experimental/ProceduralVegetationEditor/Source/ProceduralVegetation/Private/Implementations/PVFoliage.h
@@ -1,12 +1,17 @@
 // Copyright Epic Games, Inc. All Rights Reserved.
 
 #pragma once
+
 #include "CoreMinimal.h"
-#include "GeometryCollection/ManagedArrayCollection.h"
+
 #include "Curves/CurveFloat.h"
 
+#include "GeometryCollection/ManagedArrayCollection.h"
+
 #include "PVFoliage.generated.h"
 
+struct FPVFloatRamp;
+
 UENUM(BlueprintType)
 enum class EPhyllotaxyType : uint8
 {
@@ -32,7 +37,7 @@ struct FDistributionSettings
 	const float EthyleneThreshold;
 	const bool OverrideDistribution;
 	const float InstanceSpacing;
-	const FRuntimeFloatCurve* InstanceSpacingRamp;
+	const FPVFloatRamp* InstanceSpacingRamp;
 	const float InstanceSpacingRampEffect;
 	const int32 MaxPerBranch;
 };
@@ -45,14 +50,14 @@ struct FScaleSettings
 	const float MaxScale;
 	const float RandomScaleMin;
 	const float RandomScaleMax;
-	const FRuntimeFloatCurve* ScaleRamp;
+	const FPVFloatRamp* ScaleRamp;
 };
 
 struct FVectorSettings
 {
 	const bool OverrideAxilAngle;
 	const float AxilAngle;
-	const FRuntimeFloatCurve* AxilAngleRamp;
+	const FPVFloatRamp* AxilAngleRamp;
 	const float AxilAngleRampUpperValue;
 	const float AxilAngleRampEffect;
 };
diff --git a/Engine/Plugins/Experimental/ProceduralVegetationEditor/Source/ProceduralVegetation/Private/Implementations/PVMeshBuilder.h b/Engine/Plugins/Experimental/ProceduralVegetationEditor/Source/ProceduralVegetation/Private/Implementations/PVMeshBuilder.h
--- a/Engine/Plugins/Experimental/ProceduralVegetationEditor/Source/ProceduralVegetation/Private/Implementations/PVMeshBuilder.h
+++ b/Engine/Plugins/Experimental/ProceduralVegetationEditor/Source/ProceduralVegetation/Private/Implementations/PVMeshBuilder.h
@@ -1,19 +1,27 @@
 // Copyright Epic Games, Inc. All Rights Reserved.
 
 #pragma once
+
 #include "CoreMinimal.h"
+#include "PVFloatRamp.h"
+#include "PVMaterialSettings.h"
 #include "UDynamicMesh.h"
+
+#include "Curves/CurveFloat.h"
+
+#include "Engine/Texture2D.h"
+
 #include "Facades/PVBranchFacade.h"
+#include "Facades/PVPlantFacade.h"
 #include "Facades/PVPointFacade.h"
+
 #include "Generators/MeshShapeGenerator.h"
+
 #include "GeometryCollection/GeometryCollection.h"
 #include "GeometryCollection/ManagedArrayCollection.h"
-#include "Materials/MaterialInterface.h"
+
 #include "Materials/MaterialInstance.h"
-#include "PVMaterialSettings.h"
-#include "Curves/CurveFloat.h"
-#include "Engine/Texture2D.h"
-#include "Facades/PVPlantFacade.h"
+#include "Materials/MaterialInterface.h"
 
 #include "PVMeshBuilder.generated.h"
 
@@ -72,14 +80,14 @@ struct FPVMeshBuilderParams
 		meta=(ClampMin=0.0f, ClampMax=1.0f, UIMin=0.0f, UIMax=1.0f))
 	float PointGroundRetention = 0.0f;
 
-	UPROPERTY(EditAnywhere, Category="Point Removal Settings", DisplayName="Hull Retention Gradient", meta=(ClampMin=0.0f, ClampMax=1.0f))
-	FRuntimeFloatCurve PointHullRetentionGradient;
+	UPROPERTY(EditAnywhere, Category="Point Removal Settings", DisplayName="Hull Retention Gradient", meta=(XAxisMin=0.0f, XAxisMax=1.0f, YAxisMin=0.0f, YAxisMax=1.0f))
+	FPVFloatRamp PointHullRetentionGradient;
 
-	UPROPERTY(EditAnywhere, Category="Point Removal Settings", DisplayName="Main Trunk Retention Gradient", meta=(ClampMin=0.0f, ClampMax=1.0f))
-	FRuntimeFloatCurve PointMainTrunkRetentionGradient;
+	UPROPERTY(EditAnywhere, Category="Point Removal Settings", DisplayName="Main Trunk Retention Gradient", meta=(XAxisMin=0.0f, XAxisMax=1.0f, YAxisMin=0.0f, YAxisMax=1.0f))
+	FPVFloatRamp PointMainTrunkRetentionGradient;
 
-	UPROPERTY(EditAnywhere, Category="Point Removal Settings", DisplayName="Ground Retention Gradient", meta=(ClampMin=0.0f, ClampMax=1.0f))
-	FRuntimeFloatCurve PointGroundRetentionGradient;
+	UPROPERTY(EditAnywhere, Category="Point Removal Settings", DisplayName="Ground Retention Gradient", meta=(XAxisMin=0.0f, XAxisMax=1.0f, YAxisMin=0.0f, YAxisMax=1.0f))
+	FPVFloatRamp PointGroundRetentionGradient;
 
 	UPROPERTY(EditAnywhere, Category="Segment Removal Settings", DisplayName="Hull Retention",
 		meta=(ClampMin=0.0f, ClampMax=1.0f, UIMin=0.0f, UIMax=1.0f))
@@ -97,17 +105,17 @@ struct FPVMeshBuilderParams
 		meta=(ClampMin=0.0f, ClampMax=1.0f, UIMin=0.0f, UIMax=1.0f))
 	float SegmentScaleRetention = 0.0f;
 
-	UPROPERTY(EditAnywhere, Category="Segment Removal Settings", DisplayName="Hull Retention Gradient", meta=(ClampMin=0.0f, ClampMax=1.0f))
-	FRuntimeFloatCurve SegmentHullRetentionGradient;
+	UPROPERTY(EditAnywhere, Category="Segment Removal Settings", DisplayName="Hull Retention Gradient", meta=(XAxisMin=0.0f, XAxisMax=1.0f, YAxisMin=0.0f, YAxisMax=1.0f))
+	FPVFloatRamp SegmentHullRetentionGradient;
 
-	UPROPERTY(EditAnywhere, Category="Segment Removal Settings", DisplayName="Main Trunk Retention Gradient", meta=(ClampMin=0.0f, ClampMax=1.0f))
-	FRuntimeFloatCurve SegmentMainTrunkRetentionGradient;
+	UPROPERTY(EditAnywhere, Category="Segment Removal Settings", DisplayName="Main Trunk Retention Gradient", meta=(XAxisMin=0.0f, XAxisMax=1.0f, YAxisMin=0.0f, YAxisMax=1.0f))
+	FPVFloatRamp SegmentMainTrunkRetentionGradient;
 
-	UPROPERTY(EditAnywhere, Category="Segment Removal Settings", DisplayName="Ground Retention Gradient", meta=(ClampMin=0.0f, ClampMax=1.0f))
-	FRuntimeFloatCurve SegmentGroundRetentionGradient;
+	UPROPERTY(EditAnywhere, Category="Segment Removal Settings", DisplayName="Ground Retention Gradient", meta=(XAxisMin=0.0f, XAxisMax=1.0f, YAxisMin=0.0f, YAxisMax=1.0f))
+	FPVFloatRamp SegmentGroundRetentionGradient;
 
-	UPROPERTY(EditAnywhere, Category="Segment Removal Settings", DisplayName="Scale Retention Gradient", meta=(ClampMin=0.0f, ClampMax=1.0f))
-	FRuntimeFloatCurve SegmentScaleRetentionGradient;
+	UPROPERTY(EditAnywhere, Category="Segment Removal Settings", DisplayName="Scale Retention Gradient", meta=(XAxisMin=0.0f, XAxisMax=1.0f, YAxisMin=0.0f, YAxisMax=1.0f))
+	FPVFloatRamp SegmentScaleRetentionGradient;
 
 	UPROPERTY(EditAnywhere, Category="Division Removal Settings", DisplayName="Hull Retention",
 		meta=(ClampMin=0.0f, ClampMax=1.0f, UIMin=0.0f, UIMax=1.0f))
@@ -121,17 +129,17 @@ struct FPVMeshBuilderParams
 		meta=(ClampMin=0.0f, ClampMax=1.0f, UIMin=0.0f, UIMax=1.0f))
 	float DivisionGroundRetention = 0.0f;
 
-	UPROPERTY(EditAnywhere, Category="Division Removal Settings", DisplayName="Hull Retention Gradient", meta=(ClampMin=0.0f, ClampMax=1.0f))
-	FRuntimeFloatCurve DivisionHullRetentionGradient;
+	UPROPERTY(EditAnywhere, Category="Division Removal Settings", DisplayName="Hull Retention Gradient", meta=(XAxisMin=0.0f, XAxisMax=1.0f, YAxisMin=0.0f, YAxisMax=1.0f))
+	FPVFloatRamp DivisionHullRetentionGradient;
 
-	UPROPERTY(EditAnywhere, Category="Division Removal Settings", DisplayName="Main Trunk Retention Gradient", meta=(ClampMin=0.0f, ClampMax=1.0f))
-	FRuntimeFloatCurve DivisionMainTrunkRetentionGradient;
+	UPROPERTY(EditAnywhere, Category="Division Removal Settings", DisplayName="Main Trunk Retention Gradient", meta=(XAxisMin=0.0f, XAxisMax=1.0f, YAxisMin=0.0f, YAxisMax=1.0f))
+	FPVFloatRamp DivisionMainTrunkRetentionGradient;
 
-	UPROPERTY(EditAnywhere, Category="Division Removal Settings", DisplayName="Ground Retention Gradient", meta=(ClampMin=0.0f, ClampMax=1.0f))
-	FRuntimeFloatCurve DivisionGroundRetentionGradient;
+	UPROPERTY(EditAnywhere, Category="Division Removal Settings", DisplayName="Ground Retention Gradient", meta=(XAxisMin=0.0f, XAxisMax=1.0f, YAxisMin=0.0f, YAxisMax=1.0f))
+	FPVFloatRamp DivisionGroundRetentionGradient;
 
-	UPROPERTY(EditAnywhere, Category="Division Removal Settings", DisplayName="Scale Retention Gradient", meta=(ClampMin=0.0f, ClampMax=1.0f))
-	FRuntimeFloatCurve DivisionScaleRetentionGradient;
+	UPROPERTY(EditAnywhere, Category="Division Removal Settings", DisplayName="Scale Retention Gradient", meta=(XAxisMin=0.0f, XAxisMax=1.0f, YAxisMin=0.0f, YAxisMax=1.0f))
+	FPVFloatRamp DivisionScaleRetentionGradient;
 
 	UPROPERTY(EditAnywhere, Category="Plant Profile Settings", DisplayName="Profile",
 		meta=(EditCondition="bIsPlantProfileDropdownEnabled", EditConditionHides, HideEditConditionToggle, GetOptions="GetPlantProfileOptions"))
@@ -144,8 +152,8 @@ struct FPVMeshBuilderParams
 	bool bIsPlantProfileDropdownEnabled = false;
 
 	UPROPERTY(EditAnywhere, Category="Plant Profile Settings", DisplayName="Profile Falloff",
-		meta=(ClampMin=0.0f, ClampMax=1.0f, EditCondition="bIsPlantProfileDropdownEnabled", EditConditionHides, HideEditConditionToggle))
-	FRuntimeFloatCurve PlantProfileFallOff;
+		meta=(XAxisMin=0.0f, XAxisMax=1.0f, YAxisMin=0.0f, YAxisMax=1.0f, EditCondition="bIsPlantProfileDropdownEnabled", EditConditionHides, HideEditConditionToggle))
+	FPVFloatRamp PlantProfileFallOff;
 
 	UPROPERTY(EditAnywhere, Category="Plant Profile Settings", DisplayName="Profile Scale",
 		meta=(ClampMin=0.0f, ClampMax=3.0f, UIMin=0.0f, UIMax=3.0f, EditCondition="bIsPlantProfileDropdownEnabled", EditConditionHides,
@@ -194,7 +202,7 @@ struct FPVMeshBuilderParams
 	}
 
 private:
-	static void InitializeLinearCurve(FRuntimeFloatCurve& Curve)
+	static void InitializeLinearCurve(FPVFloatRamp& Curve)
 	{
 		if (FRichCurve* RichCurve = Curve.GetRichCurve())
 		{
diff --git a/Engine/Plugins/Experimental/ProceduralVegetationEditor/Source/ProceduralVegetation/Private/Nodes/PVFoliageDistributorSettings.h b/Engine/Plugins/Experimental/ProceduralVegetationEditor/Source/ProceduralVegetation/Private/Nodes/PVFoliageDistributorSettings.h
--- a/Engine/Plugins/Experimental/ProceduralVegetationEditor/Source/ProceduralVegetation/Private/Nodes/PVFoliageDistributorSettings.h
+++ b/Engine/Plugins/Experimental/ProceduralVegetationEditor/Source/ProceduralVegetation/Private/Nodes/PVFoliageDistributorSettings.h
@@ -2,11 +2,15 @@
 #pragma once
 
 #include "CoreMinimal.h"
-#include "Nodes/PVBaseSettings.h"
-#include "Implementations/PVFoliage.h"
+#include "PVFloatRamp.h"
+
 #include "DataTypes/PVData.h"
-#include "PVFoliageDistributorSettings.generated.h"
 
+#include "Implementations/PVFoliage.h"
+
+#include "Nodes/PVBaseSettings.h"
+
+#include "PVFoliageDistributorSettings.generated.h"
 
 UCLASS(BlueprintType, ClassGroup = (Procedural))
 class UPVFoliageDistributorSettings : public UPVBaseSettings
@@ -45,8 +49,8 @@ private:
 	UPROPERTY(EditAnywhere, Category="Distribution Settings", meta=(ClampMin=0.0f, ClampMax=1.0f, UIMin=0.0f, UIMax=1.0f))
 	float InstanceSpacing = 0.3f;
 
-	UPROPERTY(EditAnywhere, Category="Distribution Settings", meta=(ClampMin=0.0f, ClampMax=1.0f))
-	FRuntimeFloatCurve InstanceSpacingRamp;
+	UPROPERTY(EditAnywhere, Category="Distribution Settings", meta=(XAxisMin=0.0f, XAxisMax=1.0f, YAxisMin=0.0f, YAxisMax=1.0f))
+	FPVFloatRamp InstanceSpacingRamp;
 
 	UPROPERTY(EditAnywhere, Category="Distribution Settings", meta=(ClampMin=0.0f, ClampMax=1.0f, UIMin=0.0f, UIMax=1.0f))
 	float InstanceSpacingRampEffect = 0.0f;
@@ -72,8 +76,8 @@ private:
 	UPROPERTY(EditAnywhere, Category="Scale Settings", meta=(ClampMin=1.0f, ClampMax=2.0f, UIMin=1.0f, UIMax=2.0f))
 	float RandomScaleMax = 1.0f;
 
-	UPROPERTY(EditAnywhere, Category="Scale Settings", meta=(ClampMin=0.0f, ClampMax=1.0f))
-	FRuntimeFloatCurve ScaleRamp;
+	UPROPERTY(EditAnywhere, Category="Scale Settings", meta=(XAxisMin=0.0f, XAxisMax=1.0f, YAxisMin=0.0f, YAxisMax=1.0f))
+	FPVFloatRamp ScaleRamp;
 
 	UPROPERTY(EditAnywhere, Category="Vector Settings")
 	bool OverrideAxilAngle = true;
@@ -81,8 +85,8 @@ private:
 	UPROPERTY(EditAnywhere, Category="Vector Settings", meta=(ClampMin=0.0f, ClampMax=90.0f, UIMin=0.0f, UIMax=90.0f))
 	float AxilAngle = 45.0f;
 
-	UPROPERTY(EditAnywhere, Category="Vector Settings", meta=(ClampMin=0.0f, ClampMax=1.0f))
-	FRuntimeFloatCurve AxilAngleRamp;
+	UPROPERTY(EditAnywhere, Category="Vector Settings", meta=(XAxisMin=0.0f, XAxisMax=1.0f, YAxisMin=0.0f, YAxisMax=1.0f))
+	FPVFloatRamp AxilAngleRamp;
 
 	UPROPERTY(EditAnywhere, Category="Vector Settings", meta=(ClampMin=0.0f, ClampMax=90.0f, UIMin=0.0f, UIMax=90.0f))
 	float AxilAngleRampUpperValue = 45.0f;
diff --git a/Engine/Plugins/Experimental/ProceduralVegetationEditor/Source/ProceduralVegetation/Private/Nodes/PVImporterSettings.cpp b/Engine/Plugins/Experimental/ProceduralVegetationEditor/Source/ProceduralVegetation/Private/Nodes/PVImporterSettings.cpp
--- a/Engine/Plugins/Experimental/ProceduralVegetationEditor/Source/ProceduralVegetation/Private/Nodes/PVImporterSettings.cpp
+++ b/Engine/Plugins/Experimental/ProceduralVegetationEditor/Source/ProceduralVegetation/Private/Nodes/PVImporterSettings.cpp
@@ -3,12 +3,15 @@
 #include "PVImporterSettings.h"
 
 #include "ProceduralVegetationModule.h"
+
 #include "DataTypes/PVData.h"
 #include "DataTypes/PVGrowthData.h"
+
 #include "GeometryCollection/ManagedArrayCollection.h"
-#include "Helpers/PVJSONHelper.h"
+
 #include "Helpers/PCGHelpers.h"
 #include "Helpers/PCGSettingsHelpers.h"
+#include "Helpers/PVJSONHelper.h"
 
 TArray<FPCGPinProperties> UPVImporterSettings::InputPinProperties() const
 {
@@ -34,7 +37,11 @@ bool FPVImporterElement::ExecuteInternal(FPCGContext* InContext) const
 	const UPVImporterSettings* Settings = InContext->GetInputSettings<UPVImporterSettings>();
 	check(Settings);
 	
-	const FString& SkeletonFilePath = Settings->SkeletonFile.FilePath;
+	FString SkeletonFilePath = Settings->SkeletonFile.FilePath;
+	if (SkeletonFilePath.StartsWith("Content"))
+	{
+		SkeletonFilePath = FPaths::ProjectDir() / SkeletonFilePath;
+	}
 	FString ErrorMessage;
 
 	FManagedArrayCollection Collection;
diff --git a/Engine/Plugins/Experimental/ProceduralVegetationEditor/Source/ProceduralVegetation/Private/Nodes/PVImporterSettings.h b/Engine/Plugins/Experimental/ProceduralVegetationEditor/Source/ProceduralVegetation/Private/Nodes/PVImporterSettings.h
--- a/Engine/Plugins/Experimental/ProceduralVegetationEditor/Source/ProceduralVegetation/Private/Nodes/PVImporterSettings.h
+++ b/Engine/Plugins/Experimental/ProceduralVegetationEditor/Source/ProceduralVegetation/Private/Nodes/PVImporterSettings.h
@@ -2,6 +2,7 @@
 #pragma once
 
 #include "CoreMinimal.h"
+
 #include "Nodes/PVBaseSettings.h"
 
 #include "PVImporterSettings.generated.h"
@@ -27,7 +28,7 @@ protected:
 	virtual FPCGElementPtr CreateElement() const override;
 
 public:
-	UPROPERTY(EditAnywhere, Category="Skeleton")
+	UPROPERTY(EditAnywhere, Category="Skeleton", meta=(RelativeToGameDir))
 	FFilePath SkeletonFile;
 };
 
diff --git a/Engine/Plugins/Experimental/ProceduralVegetationEditor/Source/ProceduralVegetation/Public/Nodes/PVOutputSettings.h b/Engine/Plugins/Experimental/ProceduralVegetationEditor/Source/ProceduralVegetation/Public/Nodes/PVOutputSettings.h
--- a/Engine/Plugins/Experimental/ProceduralVegetationEditor/Source/ProceduralVegetation/Public/Nodes/PVOutputSettings.h
+++ b/Engine/Plugins/Experimental/ProceduralVegetationEditor/Source/ProceduralVegetation/Public/Nodes/PVOutputSettings.h
@@ -15,8 +15,7 @@ public:
 	UPVOutputSettings();
 
 #if WITH_EDITOR
-	virtual FName GetDefaultNodeName() const override { return FName(TEXT("ProceduralVegetationOutput")); }
-	virtual FText GetDefaultNodeTitle() const override { return NSLOCTEXT("PVOutputSettings", "NodeTitle", "Procedural Vegetation Output"); }
+	virtual FName GetDefaultNodeName() const override { return ExportSettings.MeshName.IsNone() ? "ProceduralVegetation Output" : ExportSettings.MeshName; }
 	virtual FText GetNodeTooltipText() const override { return NSLOCTEXT("PVOutputSettings", "NodeTooltip", "Output node for Procedural Vegetation"); }
 	virtual EPCGSettingsType GetType() const override { return EPCGSettingsType::InputOutput; }
 #endif
diff --git a/Engine/Plugins/Experimental/ProceduralVegetationEditor/Source/ProceduralVegetation/Public/PVFloatRamp.h b/Engine/Plugins/Experimental/ProceduralVegetationEditor/Source/ProceduralVegetation/Public/PVFloatRamp.h
new file mode 100644
--- /dev/null
+++ b/Engine/Plugins/Experimental/ProceduralVegetationEditor/Source/ProceduralVegetation/Public/PVFloatRamp.h
@@ -0,0 +1,25 @@
+// Copyright Epic Games, Inc. All Rights Reserved.
+
+#pragma once
+
+#include "Curves/RichCurve.h"
+
+#include "PVFloatRamp.generated.h"
+
+USTRUCT(BlueprintType)
+struct PROCEDURALVEGETATION_API FPVFloatRamp
+{
+	GENERATED_BODY()
+
+	UPROPERTY()
+	FRichCurve EditorCurveData;
+
+	FRichCurve* GetRichCurve()
+	{
+		return &EditorCurveData;
+	}
+	const FRichCurve* GetRichCurveConst() const
+	{
+		return &EditorCurveData;
+	}
+};
diff --git a/Engine/Plugins/Experimental/ProceduralVegetationEditor/Source/ProceduralVegetationEditor/Private/Customizations/PVFloatRampCustomization.cpp b/Engine/Plugins/Experimental/ProceduralVegetationEditor/Source/ProceduralVegetationEditor/Private/Customizations/PVFloatRampCustomization.cpp
new file mode 100644
--- /dev/null
+++ b/Engine/Plugins/Experimental/ProceduralVegetationEditor/Source/ProceduralVegetationEditor/Private/Customizations/PVFloatRampCustomization.cpp
@@ -0,0 +1,713 @@
+// Copyright Epic Games, Inc. All Rights Reserved.
+
+#include "PVFloatRampCustomization.h"
+
+#include "CurveEditorCommands.h"
+#include "DetailWidgetRow.h"
+#include "Editor.h"
+#include "IDetailChildrenBuilder.h"
+#include "PVFloatRamp.h"
+#include "ScopedTransaction.h"
+#include "SCurveEditor.h"
+
+#include "Widgets/Colors/SColorBlock.h"
+#include "Widgets/Input/SButton.h"
+#include "Widgets/Layout/SSeparator.h"
+#include "Widgets/Text/STextBlock.h"
+
+#define CAN_EXECUTE_ACTION(ActionName) CommandList->CanExecuteAction(CurveEditorCommands->ActionName.ToSharedRef())
+#define EXECUTE_ACTION(ActionName) CommandList->ExecuteAction(CurveEditorCommands->ActionName.ToSharedRef())
+
+#define LOCTEXT_NAMESPACE "PVEFloatRampCustomization"
+
+const FVector2D FPVFloatRampCustomization::DEFAULT_WINDOW_SIZE = FVector2D(800, 500);
+
+FPVFloatRampOwner::FPVFloatRampOwner(FRichCurve* InCurve, const FName InCurveName)
+	: FloatCurve(InCurve)
+	, CurveOwner(nullptr)
+{
+	FloatCurveEditInfo.CurveToEdit = FloatCurve;
+	FloatCurveEditInfo.CurveName = InCurveName;
+	ConstFloatCurveEditInfo.CurveToEdit = FloatCurve;
+	ConstFloatCurveEditInfo.CurveName = InCurveName;
+}
+
+FPVFloatRampOwner::FPVFloatRampOwner(const FPVFloatRampOwner& Other)
+{
+	this->operator=(Other);
+}
+
+FPVFloatRampOwner& FPVFloatRampOwner::operator=(const FPVFloatRampOwner& Other)
+{
+	FloatCurve = Other.FloatCurve;
+
+	FloatCurveEditInfo.CurveToEdit = FloatCurve;
+	FloatCurveEditInfo.CurveName = Other.FloatCurveEditInfo.CurveName;
+	ConstFloatCurveEditInfo.CurveToEdit = FloatCurve;
+	ConstFloatCurveEditInfo.CurveName = Other.ConstFloatCurveEditInfo.CurveName;
+
+	return *this;
+}
+
+TArray<FRichCurveEditInfoConst> FPVFloatRampOwner::GetCurves() const
+{
+	return {ConstFloatCurveEditInfo};
+}
+
+void FPVFloatRampOwner::GetCurves(TAdderReserverRef<FRichCurveEditInfoConst> Curves) const
+{
+	Curves.Add(ConstFloatCurveEditInfo);
+}
+
+TArray<FRichCurveEditInfo> FPVFloatRampOwner::GetCurves()
+{
+	return {FloatCurveEditInfo};
+}
+
+const FName& FPVFloatRampOwner::GetCurveName() const
+{
+	return FloatCurveEditInfo.CurveName;
+}
+
+void FPVFloatRampOwner::SetCurveName(const FName InName)
+{
+	FloatCurveEditInfo.CurveName = InName;
+	ConstFloatCurveEditInfo.CurveName = InName;
+}
+
+void FPVFloatRampOwner::SetCurve(FRichCurve* InCurve)
+{
+	FloatCurve = InCurve;
+
+	FloatCurveEditInfo.CurveToEdit = FloatCurve;
+	ConstFloatCurveEditInfo.CurveToEdit = FloatCurve;
+}
+
+FRichCurve& FPVFloatRampOwner::GetRichCurve()
+{
+	return *FloatCurve;
+}
+
+const FRichCurve& FPVFloatRampOwner::GetRichCurve() const
+{
+	return *FloatCurve;
+}
+
+void FPVFloatRampOwner::SetOwner(UObject* InOwner)
+{
+	CurveOwner = InOwner;
+}
+
+void FPVFloatRampOwner::ModifyOwner()
+{
+	if (CurveOwner)
+	{
+		CurveOwner->Modify(true);
+	}
+}
+
+TArray<const UObject*> FPVFloatRampOwner::GetOwners() const
+{
+	if (CurveOwner)
+	{
+		return {CurveOwner};
+	}
+	return {};
+}
+
+void FPVFloatRampOwner::OnCurveChanged(const TArray<FRichCurveEditInfo>& ChangedCurveEditInfos)
+{
+	if (ChangedCurveEditInfos.Num() > 0)
+	{
+		FRealCurve* ChangedRamp = ChangedCurveEditInfos[0].CurveToEdit;
+		OnFloatRampChanged().Broadcast(static_cast<FRichCurve*>(ChangedRamp));
+	}
+}
+
+bool FPVFloatRampOwner::IsValidCurve(FRichCurveEditInfo CurveInfo)
+{
+	return CurveInfo == FloatCurveEditInfo;
+}
+
+FPVFloatRampOwner::FOnFloatRampChanged& FPVFloatRampOwner::OnFloatRampChanged()
+{
+	return OnFloatRampChangedDelegate;
+}
+
+void FPVFloatRampOwner::MakeTransactional()
+{
+	if (CurveOwner)
+	{
+		CurveOwner->SetFlags(CurveOwner->GetFlags() | RF_Transactional);
+	}
+}
+
+void ClampKeyTime(FRealCurve* Curve, const FKeyHandle& KeyHandle, const float Min, const float Max)
+{
+	const float KeyTime = Curve->GetKeyTime(KeyHandle);
+	if (KeyTime < Min || KeyTime > Max)
+	{
+		Curve->SetKeyTime(KeyHandle, FMath::Clamp(KeyTime, Min, Max));
+	}
+};
+
+void ClampKeyValue(FRealCurve* Curve, const FKeyHandle& KeyHandle, const float Min, const float Max)
+{
+	const float KeyValue = Curve->GetKeyValue(KeyHandle);
+	if (KeyValue < Min || KeyValue > Max)
+	{
+		Curve->SetKeyValue(KeyHandle, FMath::Clamp(KeyValue, Min, Max), false);
+	}
+};
+
+FPVFloatRampCustomization::FPVFloatRampCustomization()
+	: CurveEditorCommands(&FCurveEditorCommands::Get())
+{
+	check(CurveEditorCommands);
+
+	if (GEditor)
+	{
+		GEditor->RegisterForUndo(this);
+	}
+}
+
+FPVFloatRampCustomization::~FPVFloatRampCustomization()
+{
+	if (CurveEditor.IsValid() && CurveEditor->GetCurveOwner() == FloatRampOwner.Get())
+	{
+		CurveEditor->SetCurveOwner(nullptr, false);
+	}
+
+	DestroyPopOutWindow();
+
+	if (GEditor)
+	{
+		GEditor->UnregisterForUndo(this);
+	}
+}
+
+TSharedRef<IPropertyTypeCustomization> FPVFloatRampCustomization::MakeInstance()
+{
+	return MakeShareable(new FPVFloatRampCustomization);
+}
+
+void FPVFloatRampCustomization::CustomizeHeader(
+	TSharedRef<IPropertyHandle> InPropertyHandle,
+	FDetailWidgetRow& HeaderRow,
+	IPropertyTypeCustomizationUtils& InCustomizationUtils
+)
+{
+	PropertyHandle = InPropertyHandle;
+
+	TArray<UObject*> OuterObjects;
+	PropertyHandle->GetOuterObjects(OuterObjects);
+	TArray<void*> StructPtrs;
+	PropertyHandle->AccessRawData(StructPtrs);
+
+	if (StructPtrs.Num() == 1)
+	{
+		FloatRamp = static_cast<FPVFloatRamp*>(StructPtrs[0]);
+		const FName CurveName = *InPropertyHandle->GetPropertyDisplayName().ToString();
+
+		FloatRampOwner = MakeUnique<FPVFloatRampOwner>(&FloatRamp->EditorCurveData, CurveName);
+		check(FloatRampOwner);
+
+		if (OuterObjects.Num() == 1)
+		{
+			FloatRampOwner->SetOwner(OuterObjects[0]);
+		}
+
+		const SCurveEditor::FArguments CurveEditorArgs = GetCurveEditorArgs(InPropertyHandle);
+		const float XAxisMin = CurveEditorArgs._ViewMinInput.IsBound()
+			? CurveEditorArgs._ViewMinInput.Get() + 0.05f
+			: -FLT_MAX;
+		const float XAxisMax = CurveEditorArgs._ViewMaxInput.IsBound()
+			? CurveEditorArgs._ViewMaxInput.Get() - 0.05f
+			: FLT_MAX;
+		const float YAxisMin = CurveEditorArgs._ViewMinOutput.IsBound()
+			? CurveEditorArgs._ViewMinOutput.Get() + 0.05f
+			: -FLT_MAX;
+		const float YAxisMax = CurveEditorArgs._ViewMaxOutput.IsBound()
+			? CurveEditorArgs._ViewMaxOutput.Get() - 0.05f
+			: FLT_MAX;
+
+		FloatRampOwner->OnFloatRampChanged().AddLambda([=, this](FRichCurve* Curve)
+			{
+				for (auto KeyIter = Curve->GetKeyHandleIterator(); KeyIter; ++KeyIter)
+				{
+					if (XAxisMin != -FLT_MAX || XAxisMax != FLT_MAX)
+						ClampKeyTime(Curve, *KeyIter, XAxisMin, XAxisMax);
+					if (YAxisMin != -FLT_MAX || YAxisMax != FLT_MAX)
+						ClampKeyValue(Curve, *KeyIter, YAxisMin, YAxisMax);
+				}
+			}
+		);
+
+		bHasHorizontalRange = CurveEditorArgs._ViewMinInput.IsBound() && CurveEditorArgs._ViewMaxInput.IsBound();
+		bHasVerticalRange = CurveEditorArgs._ViewMinOutput.IsBound() && CurveEditorArgs._ViewMaxOutput.IsBound();
+
+		CurveEditor = CreateCurveEditor(CurveEditorArgs);
+		check(CurveEditor);
+
+		CommandList = CurveEditor->GetCommands();
+
+		CurveEditor->SetRequireFocusToZoom(true);
+		CurveEditor->SetCurveOwner(FloatRampOwner.Get(), true);
+		CurveEditor->SetPropertyUtils(InCustomizationUtils.GetPropertyUtilities());
+		TSharedPtr<IPropertyHandle> RootHandle = PropertyHandle;
+		while (RootHandle->GetParentHandle().IsValid())
+		{
+			RootHandle = RootHandle->GetParentHandle();
+		}
+		CurveEditor->RegisterToPropertyChangedEvent(RootHandle);
+
+
+		HeaderRow
+			.NameContent()
+			[
+				InPropertyHandle->CreatePropertyNameWidget()
+			]
+			.ValueContent()
+			[
+				SNew(STextBlock)
+				.Text_Lambda(
+					[this]
+						{
+							return FText::Format(LOCTEXT("PVE_CurveKeyCount", "Curve with {0} Key(s)"), FloatRamp->GetRichCurve()->GetNumKeys());
+						}
+				)
+				.Font(FAppStyle::GetFontStyle("PropertyWindow.NormalFont"))
+				.TextStyle(FAppStyle::Get(), "DetailsView.CategoryTextStyle")
+			];
+	}
+}
+
+void FPVFloatRampCustomization::CustomizeChildren(
+	TSharedRef<IPropertyHandle> InPropertyHandle,
+	IDetailChildrenBuilder& ChildBuilder,
+	IPropertyTypeCustomizationUtils& InCustomizationUtils
+)
+{
+	if (!PropertyHandle.IsValid())
+	{
+		PropertyHandle = InPropertyHandle;
+	}
+
+	ChildBuilder
+		.AddCustomRow(FText::FromName(FloatRampOwner->GetCurveName()))
+		.WholeRowContent()
+		.HAlign(HAlign_Fill)
+		.MinDesiredWidth(300)
+		[
+			SNew(SBorder)
+			.BorderImage(nullptr)
+			.Padding(0.0f)
+			.OnMouseDoubleClick(this, &FPVFloatRampCustomization::OnRampPreviewDoubleClick)
+			[
+				SNew(SHorizontalBox)
+				+ SHorizontalBox::Slot()
+				.AutoWidth()
+				.HAlign(HAlign_Left)
+				[
+					SNew(SSeparator)
+					.Orientation(Orient_Vertical)
+					.Thickness(1.0f)
+				]
+				+ SHorizontalBox::Slot()
+				.HAlign(HAlign_Fill)
+				[
+					CreateRampValueWidget(CurveEditor.ToSharedRef(), false)
+				]
+			]
+		];
+}
+
+void FPVFloatRampCustomization::PostUndo(bool bSuccess)
+{
+	TArray<UObject*> OuterObjects;
+	PropertyHandle->GetOuterObjects(OuterObjects);
+
+	TArray<void*> StructPtrs;
+	PropertyHandle->AccessRawData(StructPtrs);
+	if (StructPtrs.Num() == 1)
+	{
+		FloatRamp = static_cast<FPVFloatRamp*>(StructPtrs[0]);
+		if (FloatRamp && FloatRampOwner)
+		{
+			if (OuterObjects.Num() == 1)
+			{
+				FloatRampOwner->SetOwner(OuterObjects[0]);
+			}
+			FloatRampOwner->SetCurve(&FloatRamp->EditorCurveData);
+			CurveEditor->SetCurveOwner(FloatRampOwner.Get(), PropertyHandle->IsEditable());
+		}
+	}
+}
+
+void FPVFloatRampCustomization::PostRedo(bool bSuccess)
+{
+	PostUndo(bSuccess);
+}
+
+SCurveEditor::FArguments FPVFloatRampCustomization::GetCurveEditorArgs(const TSharedRef<IPropertyHandle>& InPropertyHandle)
+{
+	SCurveEditor::FArguments Args;
+
+	static const FName XAxisName(TEXT("XAxisName"));
+	static const FName YAxisName(TEXT("YAxisName"));
+
+	static const FName XAxisMin(TEXT("XAxisMin"));
+	static const FName XAxisMax(TEXT("XAxisMax"));
+	static const FName YAxisMin(TEXT("YAxisMin"));
+	static const FName YAxisMax(TEXT("YAxisMax"));
+
+	if (InPropertyHandle->HasMetaData(XAxisName))
+	{
+		Args.XAxisName(InPropertyHandle->GetMetaData(XAxisName));
+	}
+	if (InPropertyHandle->HasMetaData(YAxisName))
+	{
+		Args.YAxisName(InPropertyHandle->GetMetaData(YAxisName));
+	}
+
+	if (InPropertyHandle->HasMetaData(XAxisMin))
+	{
+		Args.ViewMinInput_Lambda(
+			[Value = InPropertyHandle->GetFloatMetaData(XAxisMin)]
+				{
+					return Value - 0.05f;
+				}
+		);
+	}
+	if (InPropertyHandle->HasMetaData(XAxisMax))
+	{
+		Args.ViewMaxInput_Lambda(
+			[Value = InPropertyHandle->GetFloatMetaData(XAxisMax)]
+				{
+					return Value + 0.05f;
+				}
+		);
+	}
+	if (InPropertyHandle->HasMetaData(YAxisMin))
+	{
+		Args.ViewMinOutput_Lambda(
+			[Value = InPropertyHandle->GetFloatMetaData(YAxisMin)]
+				{
+					return Value - 0.05f;
+				}
+		);
+	}
+	if (InPropertyHandle->HasMetaData(YAxisMax))
+	{
+		Args.ViewMaxOutput_Lambda(
+			[Value = InPropertyHandle->GetFloatMetaData(YAxisMax)]
+				{
+					return Value + 0.05f;
+				}
+		);
+	}
+
+	Args.TimelineLength(0.0f)
+	    .HideUI(false)
+	    .AllowZoomOutput(true)
+	    .ShowZoomButtons(false)
+	    .ZoomToFitHorizontal(false)
+	    .ZoomToFitVertical(false);
+
+	return Args;
+}
+
+TSharedRef<SCurveEditor> FPVFloatRampCustomization::CreateCurveEditor(const SCurveEditor::FArguments& Args) const
+{
+	return SArgumentNew(Args, SCurveEditor);
+}
+
+TSharedRef<SWidget> FPVFloatRampCustomization::CreateRampValueWidget(const TSharedRef<SWidget>& ChildContent, bool bIsWindowWidget)
+{
+	SVerticalBox::FArguments Slots;
+
+	Slots
+		+ SVerticalBox::Slot()
+		  .AutoHeight()
+		  .HAlign(HAlign_Fill)
+		  .VAlign(VAlign_Top)
+		[
+			SNew(SOverlay)
+			+ SOverlay::Slot(0)
+			[
+				SNew(SColorBlock).Color(FLinearColor(0.018f, 0.018f, 0.018f))
+			]
+			+ SOverlay::Slot(1)
+			.Padding(2.0f)
+			[
+				CreateRampControlsWidget(bIsWindowWidget)
+			]
+		];
+
+	if (!bIsWindowWidget)
+	{
+		Slots
+			+ SVerticalBox::Slot()
+			[
+				SNew(SSeparator)
+				.Thickness(1.0f)
+			];
+	}
+
+	Slots
+		+ SVerticalBox::Slot()
+		  .HAlign(HAlign_Fill)
+		  .VAlign(VAlign_Fill)
+		  .SizeParam(bIsWindowWidget
+			  ? TOptional<FSizeParam>(FStretch(1.0f))
+			  : TOptional<FSizeParam>(FAuto()))
+		  .MinHeight(bIsWindowWidget
+			  ? 400
+			  : 150)
+		[
+			ChildContent
+		];
+
+	return SArgumentNew(Slots, SVerticalBox);
+}
+
+TSharedRef<SWidget> FPVFloatRampCustomization::CreateRampControlsWidget(bool bIsWindowWidget)
+{
+	SHorizontalBox::FArguments ToolsArgs;
+	CommandList->CanExecuteAction(FCurveEditorCommands::Get().ZoomToFitHorizontal.ToSharedRef());
+	if (CAN_EXECUTE_ACTION(ZoomToFitVertical) && !bHasHorizontalRange)
+	{
+		ToolsArgs
+			+ SHorizontalBox::Slot()
+			.HAlign(HAlign_Left)
+			[
+				SNew(SButton)
+				.ToolTipText(FText::FromString("Zoom To Fit Horizontal"))
+				.OnClicked_Lambda([this]
+					{
+						EXECUTE_ACTION(ZoomToFitHorizontal);
+						return FReply::Handled();
+					}
+				)
+				.ContentPadding(1)
+				[
+					SNew(SImage)
+					.Image(FAppStyle::GetBrush("CurveEd.FitHorizontal"))
+					.ColorAndOpacity(FSlateColor::UseForeground())
+				]
+			];
+	}
+
+	if (CAN_EXECUTE_ACTION(ZoomToFitVertical) && !bHasVerticalRange)
+	{
+		ToolsArgs
+			+ SHorizontalBox::Slot()
+			.HAlign(HAlign_Left)
+			[
+				SNew(SButton)
+				.ToolTipText(FText::FromString("Zoom To Fit Vertical"))
+				.OnClicked_Lambda([this]
+					{
+						EXECUTE_ACTION(ZoomToFitVertical);
+						return FReply::Handled();
+					}
+				)
+				.ContentPadding(1)
+				[
+					SNew(SImage)
+					.Image(FAppStyle::GetBrush("CurveEd.FitVertical"))
+					.ColorAndOpacity(FSlateColor::UseForeground())
+				]
+			];
+	}
+
+	static const TArray<TSharedPtr<FString>> Presets{
+		MakeShared<FString>("Flat"),
+		MakeShared<FString>("Hill"),
+		MakeShared<FString>("Linear"),
+		MakeShared<FString>("Smooth"),
+		MakeShared<FString>("Steps"),
+		MakeShared<FString>("Valley"),
+	};
+
+	SHorizontalBox::FArguments SettingsArgs;
+
+	// TODO: Hide this until we have Icons for each curve preset. Only text looks ugly
+	if (bIsWindowWidget && false)
+	{
+		SHorizontalBox::FArguments PresetsArgs;
+		for (const TSharedPtr<FString>& String : Presets)
+		{
+			PresetsArgs
+				+ SHorizontalBox::Slot()
+				  .HAlign(HAlign_Right)
+				  .VAlign(VAlign_Top)
+				  .AutoWidth()
+				[
+					SNew(SButton)
+					.Text(FText::FromString(*String))
+					.OnClicked_Lambda([=, this]
+						{
+							OnRampPresetSelected(*String);
+							return FReply::Handled();
+						}
+					)
+				];
+		}
+		SettingsArgs
+			+ SHorizontalBox::Slot()
+			  .HAlign(HAlign_Right)
+			  .VAlign(VAlign_Top)
+			  .AutoWidth()
+			[
+				SArgumentNew(PresetsArgs, SHorizontalBox)
+			];
+	}
+	else
+	{
+		SettingsArgs
+			+ SHorizontalBox::Slot()
+			.HAlign(HAlign_Right)
+			[
+				SAssignNew(PresetsComboBox, SComboBox<TSharedPtr<FString>>)
+				.HasDownArrow(false)
+				.ToolTipText(FText::FromString("Presets"))
+				.OptionsSource(&Presets)
+				.ButtonStyle(&FCoreStyle::Get().GetWidgetStyle<FButtonStyle>("Button"))
+				.OnSelectionChanged_Lambda(
+					[this](const TSharedPtr<FString>& SelectedPreset, const ESelectInfo::Type SelectType)
+						{
+							if (SelectType != ESelectInfo::Direct)
+							{
+								OnRampPresetSelected(*SelectedPreset);
+								PresetsComboBox->SetSelectedItem(nullptr);
+							}
+						}
+				)
+				.OnGenerateWidget_Lambda(
+					[](const TSharedPtr<FString>& Item)
+						{
+							return SNew(STextBlock)
+								.Text(FText::FromString(*Item));
+						}
+				)
+				.ContentPadding(1)
+				[
+					SNew(SImage)
+					.Image(FAppStyle::GetBrush("Icons.Settings"))
+					.ColorAndOpacity(FSlateColor::UseForeground())
+				]
+			];
+	}
+
+	return SNew(SHorizontalBox)
+		+ SHorizontalBox::Slot()
+		.HAlign(HAlign_Left)
+		[
+			SArgumentNew(ToolsArgs, SHorizontalBox)
+		]
+		+ SHorizontalBox::Slot()
+		.HAlign(HAlign_Right)
+		[
+			SArgumentNew(SettingsArgs, SHorizontalBox)
+		];
+}
+
+void FPVFloatRampCustomization::OnRampPresetSelected(const FString& Selection) const
+{
+	check(FloatRamp && FloatRampOwner);
+	FRichCurve& Curve = FloatRamp->EditorCurveData;
+
+	FScopedTransaction PresetTransaction(FText::FromString("Apply Curve Preset"));
+
+	FloatRampOwner->ModifyOwner();
+
+	if (Selection == "Flat")
+	{
+		Curve.Reset();
+		Curve.SetKeyInterpMode(Curve.AddKey(0.5f, 1.0f), RCIM_Constant);
+	}
+	else if (Selection == "Hill")
+	{
+		Curve.Reset();
+		Curve.SetKeyInterpMode(Curve.AddKey(0.0f, 0.0f), RCIM_Cubic);
+		Curve.SetKeyInterpMode(Curve.AddKey(0.5f, 1.0f), RCIM_Cubic);
+		Curve.SetKeyInterpMode(Curve.AddKey(1.0f, 0.0f), RCIM_Cubic);
+	}
+	else if (Selection == "Linear")
+	{
+		Curve.Reset();
+		Curve.SetKeyInterpMode(Curve.AddKey(0.0f, 0.0f), RCIM_Linear);
+		Curve.SetKeyInterpMode(Curve.AddKey(1.0f, 1.0f), RCIM_Linear);
+	}
+	else if (Selection == "Smooth")
+	{
+		Curve.Reset();
+		Curve.SetKeyInterpMode(Curve.AddKey(0.0f, 0.0f), RCIM_Cubic);
+		Curve.SetKeyInterpMode(Curve.AddKey(1.0f, 1.0f), RCIM_Cubic);
+	}
+	else if (Selection == "Steps")
+	{
+		Curve.Reset();
+		Curve.SetKeyInterpMode(Curve.AddKey(0.0f, 0.0f), RCIM_Constant);
+		Curve.SetKeyInterpMode(Curve.AddKey(0.25f, 1.0f / 3.0f), RCIM_Constant);
+		Curve.SetKeyInterpMode(Curve.AddKey(0.5f, 2.0f / 3.0f), RCIM_Constant);
+		Curve.SetKeyInterpMode(Curve.AddKey(0.75f, 1.0f), RCIM_Constant);
+	}
+	else if (Selection == "Valley")
+	{
+		Curve.Reset();
+		Curve.SetKeyInterpMode(Curve.AddKey(0.0f, 1.0f), RCIM_Cubic);
+		Curve.SetKeyInterpMode(Curve.AddKey(0.5f, 0.0f), RCIM_Cubic);
+		Curve.SetKeyInterpMode(Curve.AddKey(1.0f, 1.0f), RCIM_Cubic);
+	}
+}
+
+FReply FPVFloatRampCustomization::OnRampPreviewDoubleClick(const FGeometry& InMyGeometry, const FPointerEvent& InMouseEvent)
+{
+	if (InMouseEvent.GetEffectingButton() == EKeys::LeftMouseButton)
+	{
+		DestroyPopOutWindow();
+
+		// Determine the position of the window so that it will spawn near the mouse, but not go off the screen.
+		const FVector2f CursorPos = FSlateApplication::Get().GetCursorPos();
+		FSlateRect Anchor(CursorPos.X, CursorPos.Y, CursorPos.X, CursorPos.Y);
+
+		FVector2f AdjustedSummonLocation = FSlateApplication::Get().CalculatePopupWindowPosition(Anchor,
+			FPVFloatRampCustomization::DEFAULT_WINDOW_SIZE, true, FVector2f::ZeroVector, Orient_Horizontal);
+
+		TSharedPtr<SWindow> Window = SNew(SWindow)
+			.Title(FText::Format(FText::FromString("{0} - Internal Curve Editor"), PropertyHandle->GetPropertyDisplayName()))
+			.ClientSize(FPVFloatRampCustomization::DEFAULT_WINDOW_SIZE)
+			.ScreenPosition(AdjustedSummonLocation)
+			.AutoCenter(EAutoCenter::None)
+			.SupportsMaximize(false)
+			.SupportsMinimize(false)
+			.SizingRule(ESizingRule::FixedSize);
+
+		check(CurveEditor);
+		const TSharedRef<SWidget> CurveWindowWidget = CreateRampValueWidget(CurveEditor.ToSharedRef(), true);
+		Window->SetContent(CurveWindowWidget);
+
+		FWidgetPath WidgetPath;
+		FSlateApplication::Get().GeneratePathToWidgetChecked(CurveEditor.ToSharedRef(), WidgetPath);
+		Window = FSlateApplication::Get().AddWindowAsNativeChild(Window.ToSharedRef(), WidgetPath.GetWindow());
+
+		CurveEditorWindow = Window;
+	}
+	return FReply::Handled();
+}
+
+void FPVFloatRampCustomization::DestroyPopOutWindow()
+{
+	if (CurveEditorWindow.IsValid())
+	{
+		CurveEditorWindow.Pin()->RequestDestroyWindow();
+		CurveEditorWindow.Reset();
+	}
+}
+
+#undef LOCTEXT_NAMESPACE
+
+#undef EXECUTE_ACTION
+#undef CAN_EXECUTE_ACTION
diff --git a/Engine/Plugins/Experimental/ProceduralVegetationEditor/Source/ProceduralVegetationEditor/Private/Customizations/PVFloatRampCustomization.h b/Engine/Plugins/Experimental/ProceduralVegetationEditor/Source/ProceduralVegetationEditor/Private/Customizations/PVFloatRampCustomization.h
new file mode 100644
--- /dev/null
+++ b/Engine/Plugins/Experimental/ProceduralVegetationEditor/Source/ProceduralVegetationEditor/Private/Customizations/PVFloatRampCustomization.h
@@ -0,0 +1,114 @@
+// Copyright Epic Games, Inc. All Rights Reserved.
+
+#pragma once
+
+#include "CurveEditorCommands.h"
+#include "IPropertyTypeCustomization.h"
+#include "SCurveEditor.h"
+
+#include "Widgets/Input/SComboBox.h"
+
+struct FPVFloatRamp;
+class SCurveEditor;
+
+class FPVFloatRampOwner : public FCurveOwnerInterface
+{
+public:
+	DECLARE_MULTICAST_DELEGATE_OneParam(FOnFloatRampChanged, FRichCurve* /* ChangedCurve */);
+
+public:
+	FPVFloatRampOwner(FRichCurve* InCurve, const FName InCurveName = NAME_None);
+	FPVFloatRampOwner(const FPVFloatRampOwner& Other);
+	FPVFloatRampOwner& operator=(const FPVFloatRampOwner& Other);
+
+	FPVFloatRampOwner(FPVFloatRampOwner&& Other) = delete;
+	FPVFloatRampOwner& operator=(FPVFloatRampOwner&& Other) = delete;
+
+	UE_DEPRECATED(5.6, "Use version taking a TAdderReserverRef")
+	virtual TArray<FRichCurveEditInfoConst> GetCurves() const override;
+	virtual void GetCurves(TAdderReserverRef<FRichCurveEditInfoConst> Curves) const override;
+	virtual TArray<FRichCurveEditInfo> GetCurves() override;
+
+	const FName& GetCurveName() const;
+	void SetCurveName(const FName InName);
+	void SetCurve(FRichCurve* InCurve);
+	FRichCurve& GetRichCurve();
+	const FRichCurve& GetRichCurve() const;
+
+	void SetOwner(UObject* InOwner);
+	virtual TArray<const UObject*> GetOwners() const override;
+	virtual void ModifyOwner() override;
+	virtual void MakeTransactional() override;
+
+	virtual void OnCurveChanged(const TArray<FRichCurveEditInfo>& ChangedCurveEditInfos) override;
+	virtual bool IsValidCurve(FRichCurveEditInfo CurveInfo) override;
+
+	FOnFloatRampChanged& OnFloatRampChanged();
+
+private:
+	FRichCurve* FloatCurve;
+
+	TObjectPtr<UObject> CurveOwner;
+
+	FRichCurveEditInfo FloatCurveEditInfo;
+	FRichCurveEditInfoConst ConstFloatCurveEditInfo;
+
+	FOnFloatRampChanged OnFloatRampChangedDelegate;
+};
+
+class FPVFloatRampCustomization : public IPropertyTypeCustomization, FEditorUndoClient
+{
+public:
+	FPVFloatRampCustomization();
+	virtual ~FPVFloatRampCustomization() override;
+
+	static TSharedRef<IPropertyTypeCustomization> MakeInstance();
+
+	virtual void CustomizeHeader(
+		TSharedRef<IPropertyHandle> InPropertyHandle,
+		FDetailWidgetRow& HeaderRow,
+		IPropertyTypeCustomizationUtils& InCustomizationUtils
+	) override;
+
+	virtual void CustomizeChildren(
+		TSharedRef<IPropertyHandle> InPropertyHandle,
+		IDetailChildrenBuilder& ChildBuilder,
+		IPropertyTypeCustomizationUtils& InCustomizationUtils
+	)
+	override;
+
+	virtual void PostUndo(bool bSuccess) override;
+	virtual void PostRedo(bool bSuccess) override;
+
+private:
+	SCurveEditor::FArguments GetCurveEditorArgs(const TSharedRef<IPropertyHandle>& InPropertyHandle);
+
+	TSharedRef<SCurveEditor> CreateCurveEditor(const SCurveEditor::FArguments& Args) const;
+	TSharedRef<SWidget> CreateRampValueWidget(const TSharedRef<SWidget>& ChildContent, bool bIsWindowWidget = false);
+	TSharedRef<SWidget> CreateRampControlsWidget(bool bIsWindowWidget = false);
+
+	void OnRampPresetSelected(const FString& Selection) const;
+
+	FReply OnRampPreviewDoubleClick(const FGeometry& InMyGeometry, const FPointerEvent& InMouseEvent);
+	void DestroyPopOutWindow();
+
+private:
+	TSharedPtr<IPropertyHandle> PropertyHandle;
+
+	FPVFloatRamp* FloatRamp = nullptr;
+	TUniquePtr<FPVFloatRampOwner> FloatRampOwner;
+
+	bool bHasHorizontalRange = false;
+	bool bHasVerticalRange = false;
+
+	TSharedPtr<SCurveEditor> CurveEditor;
+
+	TWeakPtr<SWindow> CurveEditorWindow;
+
+	TSharedPtr<FUICommandList> CommandList;
+	FCurveEditorCommands* CurveEditorCommands;
+
+	TSharedPtr<SComboBox<TSharedPtr<FString>>> PresetsComboBox;
+
+	static const FVector2D DEFAULT_WINDOW_SIZE;
+};
diff --git a/Engine/Plugins/Experimental/ProceduralVegetationEditor/Source/ProceduralVegetationEditor/Private/Importer/MegaSeedImportPipeline.cpp b/Engine/Plugins/Experimental/ProceduralVegetationEditor/Source/ProceduralVegetationEditor/Private/Importer/MegaSeedImportPipeline.cpp
deleted file mode 100644
diff --git a/Engine/Plugins/Experimental/ProceduralVegetationEditor/Source/ProceduralVegetationEditor/Private/Importer/MegaSeedImportPipeline.h b/Engine/Plugins/Experimental/ProceduralVegetationEditor/Source/ProceduralVegetationEditor/Private/Importer/MegaSeedImportPipeline.h
deleted file mode 100644
diff --git a/Engine/Plugins/Experimental/ProceduralVegetationEditor/Source/ProceduralVegetationEditor/Private/Importer/MegaSeedImportWorkflow.cpp b/Engine/Plugins/Experimental/ProceduralVegetationEditor/Source/ProceduralVegetationEditor/Private/Importer/MegaSeedImportWorkflow.cpp
deleted file mode 100644
diff --git a/Engine/Plugins/Experimental/ProceduralVegetationEditor/Source/ProceduralVegetationEditor/Private/Importer/MegaSeedImportWorkflow.h b/Engine/Plugins/Experimental/ProceduralVegetationEditor/Source/ProceduralVegetationEditor/Private/Importer/MegaSeedImportWorkflow.h
deleted file mode 100644
diff --git a/Engine/Plugins/Experimental/ProceduralVegetationEditor/Source/ProceduralVegetationEditor/Private/Importer/MegaSeedImportWorkflowFactory.cpp b/Engine/Plugins/Experimental/ProceduralVegetationEditor/Source/ProceduralVegetationEditor/Private/Importer/MegaSeedImportWorkflowFactory.cpp
deleted file mode 100644
diff --git a/Engine/Plugins/Experimental/ProceduralVegetationEditor/Source/ProceduralVegetationEditor/Private/Importer/MegaSeedImportWorkflowFactory.h b/Engine/Plugins/Experimental/ProceduralVegetationEditor/Source/ProceduralVegetationEditor/Private/Importer/MegaSeedImportWorkflowFactory.h
deleted file mode 100644
diff --git a/Engine/Plugins/Experimental/ProceduralVegetationEditor/Source/ProceduralVegetationEditor/Private/Importer/MegaSeedImporter.cpp b/Engine/Plugins/Experimental/ProceduralVegetationEditor/Source/ProceduralVegetationEditor/Private/Importer/MegaSeedImporter.cpp
deleted file mode 100644
diff --git a/Engine/Plugins/Experimental/ProceduralVegetationEditor/Source/ProceduralVegetationEditor/Private/Importer/MegaSeedImporter.h b/Engine/Plugins/Experimental/ProceduralVegetationEditor/Source/ProceduralVegetationEditor/Private/Importer/MegaSeedImporter.h
deleted file mode 100644
diff --git a/Engine/Plugins/Experimental/ProceduralVegetationEditor/Source/ProceduralVegetationEditor/Private/ProceduralVegetationEditorModule.cpp b/Engine/Plugins/Experimental/ProceduralVegetationEditor/Source/ProceduralVegetationEditor/Private/ProceduralVegetationEditorModule.cpp
--- a/Engine/Plugins/Experimental/ProceduralVegetationEditor/Source/ProceduralVegetationEditor/Private/ProceduralVegetationEditorModule.cpp
+++ b/Engine/Plugins/Experimental/ProceduralVegetationEditor/Source/ProceduralVegetationEditor/Private/ProceduralVegetationEditorModule.cpp
@@ -1,37 +1,46 @@
 // Copyright Epic Games, Inc. All Rights Reserved.
 
 #include "ProceduralVegetationEditorModule.h"
+
 #include "ContentBrowserMenuContexts.h"
-#include "FabWorkflowFactoryRegistry.h"
 #include "LevelEditor.h"
-#include "ProceduralVegetationLink.h"
 #include "PCGModule.h"
-#include "Importer/MegaSeedImportWorkflowFactory.h"
-#include "Modules/ModuleManager.h"
+#include "ProceduralVegetationLink.h"
+#include "PropertyEditorModule.h"
 #include "PVEditorCommands.h"
 #include "PVEditorSettings.h"
-#include "Styling/SlateStyleRegistry.h"
-#include "DataTypes/PVData.h"
-#include "DataTypes/PVGrowthData.h"
-#include "DataTypes/PVMeshData.h"
-#include "DataTypes/PVFoliageMeshData.h"
+#include "PVFloatRamp.h"
 #include "Selection.h"
-#include "Animation/SkeletalMeshActor.h"
-#include "Visualizations/PVDataVisualization.h"
-#include "Engine/StaticMeshActor.h"
-#include "ToolMenus.h"
 #include "ToolMenuContext.h"
 #include "ToolMenuDelegates.h"
+#include "ToolMenus.h"
 #include "ToolMenuSection.h"
-#include "GameFramework/Actor.h"
-#include "Framework/MultiBox/MultiBox.h"
-#include "Engine/StaticMesh.h"
-#include "Engine/SkeletalMesh.h"
+
+#include "Animation/SkeletalMeshActor.h"
 #include "Animation/Skeleton.h"
-#include "Components/StaticMeshComponent.h"
+
 #include "Components/SkeletalMeshComponent.h"
+#include "Components/StaticMeshComponent.h"
+
+#include "Customizations/PVFloatRampCustomization.h"
+
+#include "DataTypes/PVData.h"
+#include "DataTypes/PVFoliageMeshData.h"
+#include "DataTypes/PVGrowthData.h"
+#include "DataTypes/PVMeshData.h"
+
+#include "Engine/SkeletalMesh.h"
+#include "Engine/StaticMesh.h"
+#include "Engine/StaticMeshActor.h"
+
+#include "GameFramework/Actor.h"
+
+#include "Modules/ModuleManager.h"
+
 #include "Templates/SharedPointer.h"
 
+#include "Visualizations/PVDataVisualization.h"
+
 #define LOCTEXT_NAMESPACE "ProceduralVegetationEditorModule"
 
 DEFINE_LOG_CATEGORY(LogProceduralVegetationEditor);
@@ -48,9 +57,9 @@ void FProceduralVegetationEditorModule::StartupModule()
 	FPCGDataVisualizationRegistry& DataVisRegistry = FPCGModule::GetMutablePCGDataVisualizationRegistry();
 	DataVisRegistry.RegisterPCGDataVisualization(UPVData::StaticClass(), MakeUnique<const FPVDataVisualization>());
 
-	// FAB
-	MegaSeedImportWorkflowFactory = MakeShared<FMegaSeedImportWorkflowFactory>();
-	FFabWorkflowFactoryRegistry::RegisterFactory(MegaSeedImportWorkflowFactory);
+	FPropertyEditorModule& PropertyEditor = FModuleManager::LoadModuleChecked<FPropertyEditorModule>("PropertyEditor");
+
+	PropertyEditor.RegisterCustomPropertyTypeLayout(FPVFloatRamp::StaticStruct()->GetFName(), FOnGetPropertyTypeCustomizationInstance::CreateStatic(&FPVFloatRampCustomization::MakeInstance));
 }
 
 void FProceduralVegetationEditorModule::ShutdownModule()
@@ -58,7 +67,6 @@ void FProceduralVegetationEditorModule::ShutdownModule()
 	UnregisterLevelEditorMenuExtension();
 	UnregisterPinColorAndIcons();
 	FPVEditorCommands::Unregister();
-	FFabWorkflowFactoryRegistry::UnregisterFactory(MegaSeedImportWorkflowFactory);
 }
 
 void FProceduralVegetationEditorModule::RegisterMenus()
diff --git a/Engine/Plugins/Experimental/ProceduralVegetationEditor/Source/ProceduralVegetationEditor/ProceduralVegetationEditor.Build.cs b/Engine/Plugins/Experimental/ProceduralVegetationEditor/Source/ProceduralVegetationEditor/ProceduralVegetationEditor.Build.cs
--- a/Engine/Plugins/Experimental/ProceduralVegetationEditor/Source/ProceduralVegetationEditor/ProceduralVegetationEditor.Build.cs
+++ b/Engine/Plugins/Experimental/ProceduralVegetationEditor/Source/ProceduralVegetationEditor/ProceduralVegetationEditor.Build.cs
@@ -26,25 +26,20 @@ public class ProceduralVegetationEditor : ModuleRules
 				"ContentBrowserData",
 				"Core",
 				"CoreUObject",
+				"CurveEditor",
 				"DataflowEditor",
 				"Engine",
-				"Fab",
 				"GeometryCollectionEngine",
 				"GeometryCore",
 				"GeometryFramework",
 				"InputCore",
-				"InterchangeCore",
-				"InterchangeEngine",
-				"InterchangeFactoryNodes",
-				"InterchangeNodes",
-				"InterchangePipelines",
 				"Json",
 				"JsonUtilities",
 				"LevelEditor",
-				"ProceduralVegetation",
 				"MeshConversionEngineTypes",
 				"PCG",
 				"PCGEditor",
+				"ProceduralVegetation",
 				"Projects",
 				"RenderCore",
 				"Slate",
