Fix static analysis error in FPVBranchDebugVisualization::GetPivot

diff --git a/Engine/Plugins/Experimental/ProceduralVegetationEditor/Source/ProceduralVegetationEditor/Private/Visualizations/DebugVisualization/PVDebugVisualizations.cpp b/Engine/Plugins/Experimental/ProceduralVegetationEditor/Source/ProceduralVegetationEditor/Private/Visualizations/DebugVisualization/PVDebugVisualizations.cpp
--- a/Engine/Plugins/Experimental/ProceduralVegetationEditor/Source/ProceduralVegetationEditor/Private/Visualizations/DebugVisualization/PVDebugVisualizations.cpp
+++ b/Engine/Plugins/Experimental/ProceduralVegetationEditor/Source/ProceduralVegetationEditor/Private/Visualizations/DebugVisualization/PVDebugVisualizations.cpp
@@ -95,21 +95,18 @@ void FPVBranchDebugVisualization::GetPivot(const FManagedArrayCollection& InColl
 	const PV::Facades::FBranchFacade BranchFacade(InCollection);
 	const PV::Facades::FPointFacade PointFacade(InCollection);
 
-	TArray<FVector3f> Positions;
+	OutPos = FVector3f::ZeroVector;
+	OutScale = 1.f;
 
-	if (BranchFacade.GetElementCount() > InIndex)
+	const TArray<int32>& BranchPoints = BranchFacade.GetPoints(InIndex);
+	if (BranchPoints.Num() > 0)
 	{
-		const TArray<int32>& BranchPoints = BranchFacade.GetPoints(InIndex);
+		const int RootPointIndex = BranchPoints[0];
 
-		if (BranchPoints.Num() > 0)
+		if (PointFacade.GetElementCount() > RootPointIndex)
 		{
-			const int RootPointIndex = BranchPoints[0];
-
-			if(PointFacade.GetElementCount() > RootPointIndex)
-			{
-				OutPos = PointFacade.GetPosition(RootPointIndex);
-				OutScale = PointFacade.GetPointScale(RootPointIndex);
-			}
+			OutPos = PointFacade.GetPosition(RootPointIndex);
+			OutScale = PointFacade.GetPointScale(RootPointIndex);
 		}
 	}
 }
