[PVE] Added minor UI changes after PCG updates:
- Removed extra graph settings from the panel.
- Added Find and Force regenerate buttons in the toolbar.

diff --git a/Engine/Plugins/Experimental/ProceduralVegetationEditor/Source/ProceduralVegetation/Public/ProceduralVegetation.h b/Engine/Plugins/Experimental/ProceduralVegetationEditor/Source/ProceduralVegetation/Public/ProceduralVegetation.h
--- a/Engine/Plugins/Experimental/ProceduralVegetationEditor/Source/ProceduralVegetation/Public/ProceduralVegetation.h
+++ b/Engine/Plugins/Experimental/ProceduralVegetationEditor/Source/ProceduralVegetation/Public/ProceduralVegetation.h
@@ -16,12 +16,25 @@ public:
 	UProceduralVegetationGraph(const FObjectInitializer& ObjectInitializer)
 		: UPCGGraph(ObjectInitializer)
 	{
+#if WITH_EDITORONLY_DATA
+		bIsStandaloneGraph = true;
+		bExposeGenerationInAssetExplorer = false;
+#endif
+		
 #if WITH_EDITOR
 		// Hides input and output nodes
 		SetHiddenFlagInputNode(/*bHidden=*/true);
 		SetHiddenFlagOutputNode(/*bHidden=*/true);
 #endif
 	}
+
+#if WITH_EDITOR
+	virtual bool ShouldDisplayDebuggingProperties() const override { return false; }
+	virtual bool CanToggleStandaloneGraph() const override { return false; };
+	virtual bool IsExportToLibraryEnabled() const override { return false; }
+	virtual bool ShowGraphCustomization() const override { return false; }
+	virtual bool IsTemplatePropertyEnabled() const override { return false; }
+#endif
 };
 
 /**
diff --git a/Engine/Plugins/Experimental/ProceduralVegetationEditor/Source/ProceduralVegetationEditor/Private/PVEditor.cpp b/Engine/Plugins/Experimental/ProceduralVegetationEditor/Source/ProceduralVegetationEditor/Private/PVEditor.cpp
--- a/Engine/Plugins/Experimental/ProceduralVegetationEditor/Source/ProceduralVegetationEditor/Private/PVEditor.cpp
+++ b/Engine/Plugins/Experimental/ProceduralVegetationEditor/Source/ProceduralVegetationEditor/Private/PVEditor.cpp
@@ -76,7 +76,7 @@ void FPVEditor::Initialize(const EToolkitMode::Type InMode, const TSharedPtr<cla
 
 	if (FEngineAnalytics::IsAvailable())
 	{
-		FEngineAnalytics::GetProvider().RecordEvent(TEXT("Editor.Usage.ProceduralVegetation.SessionStarted"));
+		FEngineAnalytics::GetProvider().RecordEvent(TEXT("Editor.Usage.PVE.SessionStarted"));
 	}
 }
 
@@ -129,7 +129,7 @@ void FPVEditor::OnClose()
 		TArray<FAnalyticsEventAttribute> Attributes;
 		Attributes.Add(FAnalyticsEventAttribute(TEXT("TimeActive.Seconds"),  (FDateTime::Now() - SessionStartTime).GetTotalSeconds()));
 				
-		FEngineAnalytics::GetProvider().RecordEvent(TEXT("Editor.Usage.ProceduralVegetation.SessionEnded"), Attributes);
+		FEngineAnalytics::GetProvider().RecordEvent(TEXT("Editor.Usage.PVE.SessionEnded"), Attributes);
 	}
 	
 	FPCGEditor::OnClose();
@@ -232,6 +232,9 @@ TSharedRef<FTabManager::FLayout> FPVEditor::GetDefaultLayout() const
 
 void FPVEditor::RegisterToolbarInternal(FToolMenuSection& PCGSection) const
 {
+	RegisterToolbarButton(PCGSection, EPCGToolbarButtons::Find);
+	RegisterToolbarButton(PCGSection, EPCGToolbarButtons::ForceRegen);
+	
 	const FPVEditorCommands& Commands = FPVEditorCommands::Get();
 	PCGSection.AddEntry(FToolMenuEntry::InitToolBarButton(
 			Commands.Export,
