[PVE] BugHawk Fix for CL 45721388

diff --git a/Engine/Plugins/Experimental/ProceduralVegetationEditor/Source/ProceduralVegetationEditor/Private/Customizations/PVFloatRampCustomization.cpp b/Engine/Plugins/Experimental/ProceduralVegetationEditor/Source/ProceduralVegetationEditor/Private/Customizations/PVFloatRampCustomization.cpp
--- a/Engine/Plugins/Experimental/ProceduralVegetationEditor/Source/ProceduralVegetationEditor/Private/Customizations/PVFloatRampCustomization.cpp
+++ b/Engine/Plugins/Experimental/ProceduralVegetationEditor/Source/ProceduralVegetationEditor/Private/Customizations/PVFloatRampCustomization.cpp
@@ -293,6 +293,9 @@ void FPVFloatRampCustomization::CustomizeChildren(
 		PropertyHandle = InPropertyHandle;
 	}
 
+	check(FloatRampOwner);
+	check(CurveEditor);
+
 	ChildBuilder
 		.AddCustomRow(FText::FromName(FloatRampOwner->GetCurveName()))
 		.WholeRowContent()
diff --git a/Engine/Plugins/Experimental/ProceduralVegetationEditor/Source/ProceduralVegetationEditor/Private/ProceduralVegetationEditorModule.cpp b/Engine/Plugins/Experimental/ProceduralVegetationEditor/Source/ProceduralVegetationEditor/Private/ProceduralVegetationEditorModule.cpp
--- a/Engine/Plugins/Experimental/ProceduralVegetationEditor/Source/ProceduralVegetationEditor/Private/ProceduralVegetationEditorModule.cpp
+++ b/Engine/Plugins/Experimental/ProceduralVegetationEditor/Source/ProceduralVegetationEditor/Private/ProceduralVegetationEditorModule.cpp
@@ -67,6 +67,9 @@ void FProceduralVegetationEditorModule::ShutdownModule()
 	UnregisterLevelEditorMenuExtension();
 	UnregisterPinColorAndIcons();
 	FPVEditorCommands::Unregister();
+
+	FPropertyEditorModule& PropertyEditor = FModuleManager::LoadModuleChecked<FPropertyEditorModule>("PropertyEditor");
+	PropertyEditor.UnregisterCustomPropertyTypeLayout(FPVFloatRamp::StaticStruct()->GetFName());
 }
 
 void FProceduralVegetationEditorModule::RegisterMenus()
