[PVE] Data asset changes
Applying the Material U range
Creating profile data asset

diff --git a/Engine/Plugins/Experimental/ProceduralVegetationEditor/Source/ProceduralVegetation/Private/Facades/PVBranchFacade.cpp b/Engine/Plugins/Experimental/ProceduralVegetationEditor/Source/ProceduralVegetation/Private/Facades/PVBranchFacade.cpp
--- a/Engine/Plugins/Experimental/ProceduralVegetationEditor/Source/ProceduralVegetation/Private/Facades/PVBranchFacade.cpp
+++ b/Engine/Plugins/Experimental/ProceduralVegetationEditor/Source/ProceduralVegetation/Private/Facades/PVBranchFacade.cpp
@@ -17,6 +17,7 @@ namespace PV::Facades
 		, BranchHierarchyNumber(InCollection, AttributeNames::BranchHierarchyNumber, GroupNames::BranchGroup)
 		, BranchSimulationGroupIndex(InCollection, AttributeNames::BranchSimulationGroupIndex, GroupNames::BranchGroup)
 		, TrunkMaterialPathAttribute(InCollection, AttributeNames::TrunkMaterialPath, GroupNames::DetailsGroup)
+		, TrunkURangeAttribute(InCollection, AttributeNames::TrunkURange, GroupNames::DetailsGroup)
 	{
 	}
 
@@ -32,6 +33,7 @@ namespace PV::Facades
 		, BranchHierarchyNumber(InCollection, AttributeNames::BranchHierarchyNumber, GroupNames::BranchGroup)
 		, BranchSimulationGroupIndex(InCollection, AttributeNames::BranchSimulationGroupIndex, GroupNames::BranchGroup)
 		, TrunkMaterialPathAttribute(InCollection, AttributeNames::TrunkMaterialPath, GroupNames::DetailsGroup)
+		, TrunkURangeAttribute(InCollection, AttributeNames::TrunkURange, GroupNames::DetailsGroup)
 	{
 	}
 
@@ -252,6 +254,35 @@ namespace PV::Facades
 		return FString();
 	}
 
+	void FBranchFacade::SetTrunkURange(const TArray<FVector2f>& InURange)
+	{
+		check(!IsConst());
+
+		int32 NumElements = Collection->NumElements(GroupNames::DetailsGroup);
+
+		if (NumElements == 0)
+		{
+			Collection->AddElements(1, GroupNames::DetailsGroup);
+		}
+		if (!TrunkURangeAttribute.IsValid())
+		{
+			TrunkURangeAttribute.Add();
+		}
+
+		TrunkURangeAttribute.Modify()[0] = InURange;
+	}
+
+	const TArray<FVector2f>& FBranchFacade::GetTrunkURange() const
+	{
+		if (TrunkURangeAttribute.IsValid() && TrunkURangeAttribute.IsValidIndex(0))
+		{
+			return TrunkURangeAttribute[0];
+		}
+
+		static TArray<FVector2f> TrunkURange;
+		return TrunkURange;
+	}
+
 	void FBranchFacade::CopyEntry(const int32 FromIndex, const int32 ToIndex)
 	{
 		if (IsValid() && BranchNumbers.IsValidIndex(FromIndex) && BranchNumbers.IsValidIndex(ToIndex))
diff --git a/Engine/Plugins/Experimental/ProceduralVegetationEditor/Source/ProceduralVegetation/Private/Facades/PVFoliageFacade.cpp b/Engine/Plugins/Experimental/ProceduralVegetationEditor/Source/ProceduralVegetation/Private/Facades/PVFoliageFacade.cpp
--- a/Engine/Plugins/Experimental/ProceduralVegetationEditor/Source/ProceduralVegetation/Private/Facades/PVFoliageFacade.cpp
+++ b/Engine/Plugins/Experimental/ProceduralVegetationEditor/Source/ProceduralVegetation/Private/Facades/PVFoliageFacade.cpp
@@ -3,6 +3,8 @@
 #include "Facades/PVFoliageFacade.h"
 
 #include "ProceduralVegetationModule.h"
+#include "VisualizeTexture.h"
+#include "Chaos/Deformable/MuscleActivationConstraints.h"
 #include "Facades/PVAttributesNames.h"
 
 namespace PV::Facades
@@ -242,6 +244,20 @@ namespace PV::Facades
 		return FString();
 	}
 
+	TArray<FString> FFoliageFacade::GetFoliageNames() const
+	{
+		TArray<FString> FoliageNames;
+		if (NamesAttribute.IsValid())
+		{
+			for (int32 i = 0; i < NamesAttribute.Num(); i++)
+			{
+				FoliageNames.AddUnique(NamesAttribute[i]);
+			}
+		}
+
+		return FoliageNames;
+	}
+
 	const FVector3f& FFoliageFacade::GetPivotPoint(const int32 Index) const
 	{
 		if (PivotPointsAttribute.IsValid() && PivotPointsAttribute.IsValidIndex(Index))
diff --git a/Engine/Plugins/Experimental/ProceduralVegetationEditor/Source/ProceduralVegetation/Private/Implementations/PVMaterialSettings.cpp b/Engine/Plugins/Experimental/ProceduralVegetationEditor/Source/ProceduralVegetation/Private/Implementations/PVMaterialSettings.cpp
--- a/Engine/Plugins/Experimental/ProceduralVegetationEditor/Source/ProceduralVegetation/Private/Implementations/PVMaterialSettings.cpp
+++ b/Engine/Plugins/Experimental/ProceduralVegetationEditor/Source/ProceduralVegetation/Private/Implementations/PVMaterialSettings.cpp
@@ -1,8 +1,6 @@
 // Copyright Epic Games, Inc. All Rights Reserved.
 
 #include "PVMaterialSettings.h"
-
-#include "ProceduralVegetationModule.h"
 #include "Algo/MaxElement.h"
 #include "Facades/PVBranchFacade.h"
 #include "Facades/PVPointFacade.h"
@@ -32,10 +30,11 @@ void FPVMaterialSettings::ApplyMaterialSettings(FManagedArrayCollection& Collect
 	SetOutputUVs(Collection, PointUVData, BranchUVData);
 }
 
-void FPVMaterialSettings::SetMaterial(FString TrunkMaterial, int32 Index)
+void FPVMaterialSettings::SetMaterial(FString TrunkMaterial, FVector2f UMinMax, int32 Index)
 {
 	FTrunkGenerationMaterialSetup MaterialSetup;
 	MaterialSetup.Material = LoadObject<UMaterialInterface>(nullptr, TrunkMaterial);
+	MaterialSetup.URange = FFloatRange(UMinMax.X, UMinMax.Y);
 
 	if (!MaterialSetups.IsValidIndex(Index))
 	{
diff --git a/Engine/Plugins/Experimental/ProceduralVegetationEditor/Source/ProceduralVegetation/Private/Implementations/PVMaterialSettings.h b/Engine/Plugins/Experimental/ProceduralVegetationEditor/Source/ProceduralVegetation/Private/Implementations/PVMaterialSettings.h
--- a/Engine/Plugins/Experimental/ProceduralVegetationEditor/Source/ProceduralVegetation/Private/Implementations/PVMaterialSettings.h
+++ b/Engine/Plugins/Experimental/ProceduralVegetationEditor/Source/ProceduralVegetation/Private/Implementations/PVMaterialSettings.h
@@ -100,7 +100,7 @@ struct FPVMaterialSettings
 	 */
 	void ApplyMaterialSettings(FManagedArrayCollection& Collection) const;
 
-	void SetMaterial(FString TrunkMaterial, int32 Index);
+	void SetMaterial(FString TrunkMaterial, FVector2f UMinMax, int32 Index);
 
 private:
 
diff --git a/Engine/Plugins/Experimental/ProceduralVegetationEditor/Source/ProceduralVegetation/Private/Nodes/PVImporterSettings.cpp b/Engine/Plugins/Experimental/ProceduralVegetationEditor/Source/ProceduralVegetation/Private/Nodes/PVImporterSettings.cpp
--- a/Engine/Plugins/Experimental/ProceduralVegetationEditor/Source/ProceduralVegetation/Private/Nodes/PVImporterSettings.cpp
+++ b/Engine/Plugins/Experimental/ProceduralVegetationEditor/Source/ProceduralVegetation/Private/Nodes/PVImporterSettings.cpp
@@ -52,6 +52,9 @@ bool FPVImporterElement::ExecuteInternal(FPCGContext* InContext) const
 		PCGLog::InputOutput::LogInvalidInputDataError(InContext);
 		return true;
 	}
+
+	PV::SetFoliagePaths(Collection, SkeletonFilePath);
+	
 	UPVGrowthData* OutManagedArrayCollectionData = FPCGContext::NewObject_AnyThread<UPVGrowthData>(InContext);
 	OutManagedArrayCollectionData->Initialize(MoveTemp(Collection));
 		
diff --git a/Engine/Plugins/Experimental/ProceduralVegetationEditor/Source/ProceduralVegetation/Private/Nodes/PVMeshBuilderSettings.cpp b/Engine/Plugins/Experimental/ProceduralVegetationEditor/Source/ProceduralVegetation/Private/Nodes/PVMeshBuilderSettings.cpp
--- a/Engine/Plugins/Experimental/ProceduralVegetationEditor/Source/ProceduralVegetation/Private/Nodes/PVMeshBuilderSettings.cpp
+++ b/Engine/Plugins/Experimental/ProceduralVegetationEditor/Source/ProceduralVegetation/Private/Nodes/PVMeshBuilderSettings.cpp
@@ -90,13 +90,20 @@ bool FPVMeshBuilderElement::ExecuteInternal(FPCGContext* InContext) const
 			{
 				const PV::Facades::FBranchFacade BranchFacade(SkeletonCollection);
 				const FString TrunkMaterial = BranchFacade.GetTrunkMaterialPath();
-				if (!TrunkMaterial.IsEmpty())
-				{
-					MaterialSettings.SetMaterial(TrunkMaterial, 0);
-				}
-				else
+				const TArray<FVector2f> TrunkURanges = BranchFacade.GetTrunkURange();
+
+				int32 Generation = 0;
+				for (const FVector2f& URange : TrunkURanges)
 				{
-					MutableSettings->MesherSettings.bOverrideMaterial = true;
+					if (!TrunkMaterial.IsEmpty())
+					{
+						MaterialSettings.SetMaterial(TrunkMaterial, URange, Generation);
+						Generation++;
+					}
+					else
+					{
+						MutableSettings->MesherSettings.bOverrideMaterial = true;
+					}
 				}
 			}
 
diff --git a/Engine/Plugins/Experimental/ProceduralVegetationEditor/Source/ProceduralVegetation/Private/ProceduralVegetationPreset.cpp b/Engine/Plugins/Experimental/ProceduralVegetationEditor/Source/ProceduralVegetation/Private/ProceduralVegetationPreset.cpp
--- a/Engine/Plugins/Experimental/ProceduralVegetationEditor/Source/ProceduralVegetation/Private/ProceduralVegetationPreset.cpp
+++ b/Engine/Plugins/Experimental/ProceduralVegetationEditor/Source/ProceduralVegetation/Private/ProceduralVegetationPreset.cpp
@@ -1,9 +1,20 @@
 ﻿// Copyright Epic Games, Inc. All Rights Reserved.
 
 #include "ProceduralVegetationPreset.h"
-
+#include "ProceduralVegetationModule.h"
+#include "Misc/PackageName.h"
+#include "UObject/Package.h"
+#include "AssetRegistry/IAssetRegistry.h"
+#include "Facades/PVBranchFacade.h"
 #include "Helpers/PVJSONHelper.h"
 
+UProceduralVegetationPreset::UProceduralVegetationPreset()
+{
+	FString AssetPath = FPackageName::GetLongPackagePath(GetPackage()->GetName());
+	FoliageFolder.Path = FPaths::Combine(AssetPath, TEXT("Instances"));
+	MaterialsFolder.Path = FPaths::Combine(AssetPath, TEXT("Materials"));
+}
+
 void UProceduralVegetationPreset::LoadFromVariantFiles(const TMap<FString, FString>& InVariantFiles)
 {
 	for (const auto& [VariantName, VariantFile] : InVariantFiles)
@@ -15,3 +26,106 @@ void UProceduralVegetationPreset::LoadFromVariantFiles(const TMap<FString, FStri
 		}
 	}
 }
+
+void UProceduralVegetationPreset::UpdateDataAsset()
+{
+	if (JsonDirectoryPath.Path.IsEmpty())
+	{
+		UE_LOG(LogProceduralVegetation, Warning, TEXT("Please specify a JSON path"));
+		return;
+	}
+	
+	IFileManager& FileManager = IFileManager::Get();
+	
+	TArray<FString> FileNames;
+	FileManager.FindFiles(FileNames, *JsonDirectoryPath.Path, TEXT("*.json"));
+
+	Variants.Empty();
+	
+	TArray<TSharedPtr<FJsonObject>> MetaFiles;
+	for (const auto& FileName : FileNames)
+	{
+		FString FullPath = JsonDirectoryPath.Path / FileName;
+		FManagedArrayCollection Collection;
+		FString OutError;
+		if ( PV::LoadMegaPlantsJsonToCollection(Collection, FullPath, OutError))
+		{
+			Variants.Add(FPaths::GetBaseFilename(FullPath), Collection);
+		}
+		else if (TSharedPtr<FJsonObject> MetaJson = PV::LoadMetaFileIntoJsonObject(FullPath, OutError))
+		{
+			MetaFiles.Add(MetaJson);
+		}
+		if (!OutError.IsEmpty())
+		{
+			UE_LOG(LogProceduralVegetation, Warning, TEXT("%s"), *OutError);
+		}
+	}
+	
+	int VariantIndex = 0;
+	for (auto& Pair : Variants)
+	{
+		if (MetaFiles.IsValidIndex(VariantIndex))
+		{
+			PV::LoadMetaJsonToCollection(Pair.Value, MetaFiles[VariantIndex]);
+		}
+		else if (MetaFiles.Num() > 0 && VariantIndex > MetaFiles.Num() - 1)
+		{
+			PV::LoadMetaJsonToCollection(Pair.Value, MetaFiles[0]);
+		}
+		
+		VariantIndex++;
+	}
+	
+	if (bCreateProfileDataAsset)
+	{
+		CreateProfileDataAsset();
+	}
+	
+	UpdateFoliageAndMaterialPath();
+	GetPackage()->SetDirtyFlag(true);
+}
+
+void UProceduralVegetationPreset::CreateProfileDataAsset()
+{
+	FString ProfilePackageName(PlantProfileName);
+	FString PackagePath = FPaths::Combine(FPackageName::GetLongPackagePath(GetPackage()->GetName()), ProfilePackageName);
+	UPackage* PlantProfilePackage = CreatePackage(*PackagePath);
+	
+	UPlantProfileAsset* NewPlantProfileAsset = NewObject<UPlantProfileAsset>(PlantProfilePackage, *ProfilePackageName, RF_Public | RF_Standalone);
+
+	for (auto& Pair : Variants)
+	{
+		PV::Facades::FPlantProfileFacade Facade = PV::Facades::FPlantProfileFacade(Pair.Value);
+		for (int32 Index = 0; Index < Facade.NumProfileEntries(); Index++)
+		{
+			FPlantProfile Profile;
+			Profile.ProfileName = FString::Format(TEXT("Profile_{0}"), {Index});
+			auto Points = Facade.GetProfilePoints(Index);
+			Profile.ProfilePoints = Points;
+			NewPlantProfileAsset->Profiles.Add(Profile);
+		}
+	}
+	
+	IAssetRegistry::Get()->AssetCreated(NewPlantProfileAsset);
+	PlantProfilePackage->SetDirtyFlag(true);
+}
+
+void UProceduralVegetationPreset::UpdateFoliageAndMaterialPath()
+{
+	for (auto& Pair : Variants)
+	{
+		PV::Facades::FFoliageFacade FoliageFacade(Pair.Value);
+		PV::Facades::FBranchFacade BranchFacade(Pair.Value);
+			
+		for (int32 FoliageIndex = 0; FoliageIndex < FoliageFacade.NumFoliageNames(); FoliageIndex++)
+		{
+			FString ShortName = FoliageFacade.GetFoliageName(FoliageIndex);
+			//FString ObjectName = FString::Format(TEXT("{0}.{1}"),{ ShortName, ShortName});
+			FString FullObjectName = FPaths::Combine(*FoliageFolder.Path, *ShortName);
+			FoliageFacade.SetFoliageName(FoliageIndex, FullObjectName);
+			FString FullMaterialName = FPaths::Combine(*MaterialsFolder.Path, TrunkMaterialName + '.' + TrunkMaterialName);
+			BranchFacade.SetTrunkMaterialPath( FullMaterialName );
+		}
+	}
+}
diff --git a/Engine/Plugins/Experimental/ProceduralVegetationEditor/Source/ProceduralVegetation/Public/Facades/PVAttributesNames.h b/Engine/Plugins/Experimental/ProceduralVegetationEditor/Source/ProceduralVegetation/Public/Facades/PVAttributesNames.h
--- a/Engine/Plugins/Experimental/ProceduralVegetationEditor/Source/ProceduralVegetation/Public/Facades/PVAttributesNames.h
+++ b/Engine/Plugins/Experimental/ProceduralVegetationEditor/Source/ProceduralVegetation/Public/Facades/PVAttributesNames.h
@@ -27,6 +27,7 @@ namespace PV::AttributeNames
 	inline static const FName BranchFoliageIDs = FName(TEXT("FoliageID"));
 	inline static const FName BranchUVMaterial = FName(TEXT("branchUVMat"));
 	inline static const FName TrunkMaterialPath = FName(TEXT("TrunkMaterialPath"));
+	inline static const FName TrunkURange = FName(TEXT("TrunkURange"));
 	inline static const FName BranchHierarchyNumber = FName(TEXT("branchHierarchyNumber"));
 	inline static const FName BranchSimulationGroupIndex = FName(TEXT("branchSimulationGroupIndex"));
 	inline static const FName PlantNumber = FName(TEXT("plantNumber"));
diff --git a/Engine/Plugins/Experimental/ProceduralVegetationEditor/Source/ProceduralVegetation/Public/Facades/PVBranchFacade.h b/Engine/Plugins/Experimental/ProceduralVegetationEditor/Source/ProceduralVegetation/Public/Facades/PVBranchFacade.h
--- a/Engine/Plugins/Experimental/ProceduralVegetationEditor/Source/ProceduralVegetation/Public/Facades/PVBranchFacade.h
+++ b/Engine/Plugins/Experimental/ProceduralVegetationEditor/Source/ProceduralVegetation/Public/Facades/PVBranchFacade.h
@@ -63,6 +63,10 @@ namespace PV::Facades
 
 		FString GetTrunkMaterialPath() const;
 
+		void SetTrunkURange(const TArray<FVector2f>& InURange);
+
+		const TArray<FVector2f>& GetTrunkURange() const;
+
 		virtual void CopyEntry(const int32 FromIndex, const int32 ToIndex) override;
 
 		virtual void RemoveEntries(const int32 NumEntries, const int32 StartIndex) override;
@@ -84,5 +88,6 @@ namespace PV::Facades
 		TManagedArrayAccessor<int32> BranchHierarchyNumber;
 		TManagedArrayAccessor<int32> BranchSimulationGroupIndex;
 		TManagedArrayAccessor<FString> TrunkMaterialPathAttribute;
+		TManagedArrayAccessor<TArray<FVector2f>> TrunkURangeAttribute;
 	};
 }
diff --git a/Engine/Plugins/Experimental/ProceduralVegetationEditor/Source/ProceduralVegetation/Public/Facades/PVFoliageFacade.h b/Engine/Plugins/Experimental/ProceduralVegetationEditor/Source/ProceduralVegetation/Public/Facades/PVFoliageFacade.h
--- a/Engine/Plugins/Experimental/ProceduralVegetationEditor/Source/ProceduralVegetation/Public/Facades/PVFoliageFacade.h
+++ b/Engine/Plugins/Experimental/ProceduralVegetationEditor/Source/ProceduralVegetation/Public/Facades/PVFoliageFacade.h
@@ -62,6 +62,8 @@ namespace PV::Facades
 
 		FString GetFoliageName(const int32 Index) const;
 
+		TArray<FString> GetFoliageNames() const;
+
 		const FVector3f& GetPivotPoint(const int32 Index) const;
 
 		void SetPivotPoint(const int32 Index, const FVector3f& Input);
diff --git a/Engine/Plugins/Experimental/ProceduralVegetationEditor/Source/ProceduralVegetation/Public/Helpers/PVJSONHelper.h b/Engine/Plugins/Experimental/ProceduralVegetationEditor/Source/ProceduralVegetation/Public/Helpers/PVJSONHelper.h
--- a/Engine/Plugins/Experimental/ProceduralVegetationEditor/Source/ProceduralVegetation/Public/Helpers/PVJSONHelper.h
+++ b/Engine/Plugins/Experimental/ProceduralVegetationEditor/Source/ProceduralVegetation/Public/Helpers/PVJSONHelper.h
@@ -7,6 +7,7 @@
 #include "Facades/PVAttributesNames.h"
 #include "Dom/JsonObject.h"
 #include "Dom/JsonValue.h"
+#include "Facades/PVBranchFacade.h"
 #include "Facades/PVMetaInfoFacade.h"
 #include "Facades/PVFoliageFacade.h"
 #include "Facades/PVProfileFacade.h"
@@ -275,18 +276,11 @@ namespace PV
 
 		Facades::FFoliageFacade Facade(Collection, TotalValues);
 
-		FString AbsolutePath = FPaths::ConvertRelativePathToFull(InPath);
-		FString PackagePath;
-
-		if (FPackageName::TryConvertFilenameToLongPackageName(AbsolutePath, PackagePath))
+		for (FString& FoliageName : FoliageNames)
 		{
-			FString PackageFolder = FPackageName::GetLongPackagePath(PackagePath);
-			Facade.SetFoliagePath(PackageFolder);
-			for (FString& FoliageName : FoliageNames)
-			{
-				FoliageName = PackageFolder / FoliageName + '.' + FoliageName;
-			}
+			FoliageName = FoliageName + '.' + FoliageName;
 		}
+		
 		Facade.SetFoliageNames(FoliageNames);
 
 		int32 FoliageEntryIndex = 0;
@@ -336,6 +330,27 @@ namespace PV
 		}
 	}
 
+	static void SetFoliagePaths(FManagedArrayCollection& Collection,const FString& FilePath)
+	{
+		Facades::FFoliageFacade Facade(Collection);
+
+		FString AbsolutePath = FPaths::ConvertRelativePathToFull(FilePath);
+		FString PackagePath;
+
+		if (FPackageName::TryConvertFilenameToLongPackageName(AbsolutePath, PackagePath))
+		{
+			TArray<FString> FoliageNames = Facade.GetFoliageNames();
+			FString PackageFolder = FPackageName::GetLongPackagePath(PackagePath);
+			Facade.SetFoliagePath(PackageFolder);
+			for (FString& FoliageName : FoliageNames)
+			{
+				FoliageName = PackageFolder / FoliageName;
+			}
+			Facade.SetFoliageNames(FoliageNames);
+		}
+		
+	}
+
 	static bool HasJsonFieldPath(const TSharedPtr<FJsonObject>& JsonObject, const FString& Path)
 	{
 		if (!JsonObject.IsValid())
@@ -480,6 +495,60 @@ namespace PV
 		Facades::FMetaInfoFacade MetaInfoFacade(Collection);
 		MetaInfoFacade.CreateGuid(FilePath);
 
+		OutErrorMessage = "";
+		return true;
+	}
+
+	static TSharedPtr<FJsonObject> LoadMetaFileIntoJsonObject(const FString& FilePath, FString& OutErrorMessage)
+	{
+		FString MetaData;
+		if (!FFileHelper::LoadFileToString(MetaData, *FilePath))
+		{
+			OutErrorMessage = TEXT("Failed to open skeleton file");
+			return nullptr;
+		}
+
+		TSharedPtr<FJsonObject> LoadedData = MakeShareable(new FJsonObject);
+		if (!FJsonSerializer::Deserialize(TJsonReaderFactory<>::Create(MetaData), LoadedData))
+		{
+			OutErrorMessage = TEXT("Failed to deserialize JSON file for skeleton");
+			return nullptr;
+		}
+
+		TArray<FString> RequiredJSONPaths = {
+			TEXT("apix_uv"),
+			TEXT("apix_texture_channels")
+		};
+
+		for (const FString& Path : RequiredJSONPaths)
+		{
+			if (!HasJsonFieldPath(LoadedData, Path))
+			{
+				OutErrorMessage = FString::Printf(TEXT("Failed to find required path in JSON file: %s"), *Path);
+				return nullptr;
+			}
+		}
+
+		OutErrorMessage = "";
+		return LoadedData;
+	}
+
+	static bool LoadMetaJsonToCollection(FManagedArrayCollection& Collection, TSharedPtr<FJsonObject> LoadedData)
+	{
+		
+		const TArray<TSharedPtr<FJsonValue>>& MaterialURanges = LoadedData->GetArrayField(TEXT("apix_uv"));
+
+		TArray<FVector2f> URanges;
+		URanges.Reserve(MaterialURanges.Num());
+		
+		for (int i=0; i < MaterialURanges.Num(); i+= 2)
+		{
+			URanges.Add(FVector2f(MaterialURanges[i]->AsNumber() , MaterialURanges[i+1]->AsNumber()));
+		}
+
+		Facades::FBranchFacade Facade = Facades::FBranchFacade(Collection);
+		Facade.SetTrunkURange(URanges);
+		
 		return true;
 	}
 }
diff --git a/Engine/Plugins/Experimental/ProceduralVegetationEditor/Source/ProceduralVegetation/Public/ProceduralVegetationPreset.h b/Engine/Plugins/Experimental/ProceduralVegetationEditor/Source/ProceduralVegetation/Public/ProceduralVegetationPreset.h
--- a/Engine/Plugins/Experimental/ProceduralVegetationEditor/Source/ProceduralVegetation/Public/ProceduralVegetationPreset.h
+++ b/Engine/Plugins/Experimental/ProceduralVegetationEditor/Source/ProceduralVegetation/Public/ProceduralVegetationPreset.h
@@ -14,9 +14,58 @@ class PROCEDURALVEGETATION_API UProceduralVegetationPreset : public UDataAsset
 	GENERATED_BODY()
 
 public:
+	UProceduralVegetationPreset();
+	
+	UPROPERTY(EditAnywhere, Category="Preset Data", meta = (DevelopmentOnly))
+	FDirectoryPath JsonDirectoryPath;
+	
+	UPROPERTY(EditAnywhere, Category = "Preset Data", Meta = (ContentDir, DevelopmentOnly))
+	bool bOverrideFolderPaths = false;
+	
+	UPROPERTY(EditAnywhere, Category = "Preset Data", Meta = (ContentDir, DevelopmentOnly , EditCondition = "bOverrideFolderPaths"))
+	FDirectoryPath FoliageFolder;
+	
+	UPROPERTY(EditAnywhere, Category = "Preset Data", Meta = (ContentDir, DevelopmentOnly , EditCondition = "bOverrideFolderPaths"))
+	FDirectoryPath MaterialsFolder;
+	
+	UPROPERTY(EditAnywhere, Category = "Preset Data", meta = (DevelopmentOnly))
+	FString TrunkMaterialName;
+	
+	UPROPERTY(EditAnywhere, Category = "Preset Data", meta = (DevelopmentOnly))
+	bool bCreateProfileDataAsset = false;
+	
+	UPROPERTY(EditAnywhere, Category = "Preset Data", meta = (DevelopmentOnly, EditCondition = "bCreateProfileDataAsset"))
+	FString PlantProfileName;
+	
 	UPROPERTY(EditAnywhere, Category="Preset Data")
 	TMap<FString, FManagedArrayCollection> Variants;
 
 public:
-	void LoadFromVariantFiles(const TMap<FString, FString>& InVariantFiles);
+	void LoadFromVariantFiles(const TMap<FString, FString>& InVariantFiles);UFUNCTION(CallInEditor, Category="Preset Data", meta = (DevelopmentOnly))
+
+	void UpdateDataAsset();
+private:
+	void CreateProfileDataAsset();
+	
+	void UpdateFoliageAndMaterialPath();
+};
+
+USTRUCT(BlueprintType)
+struct PROCEDURALVEGETATION_API FPlantProfile
+{
+	GENERATED_BODY()
+public:
+	UPROPERTY(VisibleAnywhere, Category="PlantProfile")
+	FString ProfileName;
+	UPROPERTY(VisibleAnywhere, Category="PlantProfile")
+	TArray<float> ProfilePoints;
 };
+
+UCLASS(BlueprintType)
+class PROCEDURALVEGETATION_API UPlantProfileAsset : public UDataAsset
+{
+	GENERATED_BODY()
+public:
+	UPROPERTY(VisibleAnywhere, Category="PlantProfile")
+	TArray<FPlantProfile> Profiles;
+};
\ No newline at end of file
diff --git a/Engine/Plugins/Experimental/ProceduralVegetationEditor/Source/ProceduralVegetationEditor/Private/PVEditor.cpp b/Engine/Plugins/Experimental/ProceduralVegetationEditor/Source/ProceduralVegetationEditor/Private/PVEditor.cpp
--- a/Engine/Plugins/Experimental/ProceduralVegetationEditor/Source/ProceduralVegetationEditor/Private/PVEditor.cpp
+++ b/Engine/Plugins/Experimental/ProceduralVegetationEditor/Source/ProceduralVegetationEditor/Private/PVEditor.cpp
@@ -333,9 +333,7 @@ bool FPVEditor::GetCollectionFromNode(const UPCGNode* InNode, FName& OutPinName,
 	{
 		Pin = Pins[0];
 	}
-		
-	check(Pin)
-		
+	
 	const FPCGStack* PCGStack = GetStackBeingInspected();
 	check(PCGStack)
 		
@@ -344,7 +342,10 @@ bool FPVEditor::GetCollectionFromNode(const UPCGNode* InNode, FName& OutPinName,
 	TArray<FPCGStackFrame>& StackFrames = Stack.GetStackFramesMutable();
 	StackFrames.Reserve(StackFrames.Num() + 2);
 	StackFrames.Emplace(InNode);
-	StackFrames.Emplace(Pin);
+	if (Pin)
+	{
+		StackFrames.Emplace(Pin);
+	}
 		
 	const FPCGDataCollection* Data = ExecutionSource->GetExecutionState().GetInspection().GetInspectionData(Stack);
 		
@@ -356,7 +357,7 @@ bool FPVEditor::GetCollectionFromNode(const UPCGNode* InNode, FName& OutPinName,
 		{
 			const UPVData* ProceduralVegetationData = Cast<UPVData>(TaggedData->Data);
 		
-			if(ProceduralVegetationData != nullptr)
+			if(ProceduralVegetationData != nullptr && Pin)
 			{
 				OutPinName = Pin->Properties.Label;
 				OutCollection = ProceduralVegetationData->GetCollection();
