[PVE] Bug hawk fixes

diff --git a/Engine/Plugins/Experimental/ProceduralVegetationEditor/Source/ProceduralVegetation/Private/Implementations/PVMaterialSettings.cpp b/Engine/Plugins/Experimental/ProceduralVegetationEditor/Source/ProceduralVegetation/Private/Implementations/PVMaterialSettings.cpp
--- a/Engine/Plugins/Experimental/ProceduralVegetationEditor/Source/ProceduralVegetation/Private/Implementations/PVMaterialSettings.cpp
+++ b/Engine/Plugins/Experimental/ProceduralVegetationEditor/Source/ProceduralVegetation/Private/Implementations/PVMaterialSettings.cpp
@@ -43,6 +43,7 @@ void FPVMaterialSettings::SetMaterial(FString TrunkMaterial, FVector2f UMinMax,
 	else
 	{
 		MaterialSetups[Index].Material = MaterialSetup.Material;
+		MaterialSetups[Index].URange = MaterialSetup.URange;
 	}
 }
 
diff --git a/Engine/Plugins/Experimental/ProceduralVegetationEditor/Source/ProceduralVegetation/Private/ProceduralVegetationPreset.cpp b/Engine/Plugins/Experimental/ProceduralVegetationEditor/Source/ProceduralVegetation/Private/ProceduralVegetationPreset.cpp
--- a/Engine/Plugins/Experimental/ProceduralVegetationEditor/Source/ProceduralVegetation/Private/ProceduralVegetationPreset.cpp
+++ b/Engine/Plugins/Experimental/ProceduralVegetationEditor/Source/ProceduralVegetation/Private/ProceduralVegetationPreset.cpp
@@ -124,8 +124,16 @@ void UProceduralVegetationPreset::UpdateFoliageAndMaterialPath()
 			//FString ObjectName = FString::Format(TEXT("{0}.{1}"),{ ShortName, ShortName});
 			FString FullObjectName = FPaths::Combine(*FoliageFolder.Path, *ShortName);
 			FoliageFacade.SetFoliageName(FoliageIndex, FullObjectName);
-			FString FullMaterialName = FPaths::Combine(*MaterialsFolder.Path, TrunkMaterialName + '.' + TrunkMaterialName);
-			BranchFacade.SetTrunkMaterialPath( FullMaterialName );
+
+			if (!TrunkMaterialName.IsEmpty())
+			{
+				FString FullMaterialName = FPaths::Combine(*MaterialsFolder.Path, TrunkMaterialName + '.' + TrunkMaterialName);
+				BranchFacade.SetTrunkMaterialPath( FullMaterialName );
+			}
+			else
+			{
+				UE_LOG(LogProceduralVegetation, Warning, TEXT("Material name is empty"));
+			}
 		}
 	}
 }
diff --git a/Engine/Plugins/Experimental/ProceduralVegetationEditor/Source/ProceduralVegetation/Public/Helpers/PVJSONHelper.h b/Engine/Plugins/Experimental/ProceduralVegetationEditor/Source/ProceduralVegetation/Public/Helpers/PVJSONHelper.h
--- a/Engine/Plugins/Experimental/ProceduralVegetationEditor/Source/ProceduralVegetation/Public/Helpers/PVJSONHelper.h
+++ b/Engine/Plugins/Experimental/ProceduralVegetationEditor/Source/ProceduralVegetation/Public/Helpers/PVJSONHelper.h
@@ -541,7 +541,7 @@ namespace PV
 		TArray<FVector2f> URanges;
 		URanges.Reserve(MaterialURanges.Num());
 		
-		for (int i=0; i < MaterialURanges.Num(); i+= 2)
+		for (int i=0; (i + 1) < MaterialURanges.Num(); i+= 2)
 		{
 			URanges.Add(FVector2f(MaterialURanges[i]->AsNumber() , MaterialURanges[i+1]->AsNumber()));
 		}
diff --git a/Engine/Plugins/Experimental/ProceduralVegetationEditor/Source/ProceduralVegetationEditor/Private/PVEditor.cpp b/Engine/Plugins/Experimental/ProceduralVegetationEditor/Source/ProceduralVegetationEditor/Private/PVEditor.cpp
--- a/Engine/Plugins/Experimental/ProceduralVegetationEditor/Source/ProceduralVegetationEditor/Private/PVEditor.cpp
+++ b/Engine/Plugins/Experimental/ProceduralVegetationEditor/Source/ProceduralVegetationEditor/Private/PVEditor.cpp
@@ -337,7 +337,7 @@ bool FPVEditor::GetCollectionFromNode(const UPCGNode* InNode, FName& OutPinName,
 	const FPCGStack* PCGStack = GetStackBeingInspected();
 	check(PCGStack)
 		
-	// Create a temporary stack with Node+Pin to query the exact DataCollection
+	// Create a temporary stack with Node + valid Pin to query the exact DataCollection
 	FPCGStack Stack = *PCGStack;
 	TArray<FPCGStackFrame>& StackFrames = Stack.GetStackFramesMutable();
 	StackFrames.Reserve(StackFrames.Num() + 2);
