[PVE] Fixed issue: Randomly getting memory corruption crash. Setting bIsStandaloneGraph

diff --git a/Engine/Plugins/Experimental/ProceduralVegetationEditor/Source/ProceduralVegetation/Public/ProceduralVegetation.h b/Engine/Plugins/Experimental/ProceduralVegetationEditor/Source/ProceduralVegetation/Public/ProceduralVegetation.h
--- a/Engine/Plugins/Experimental/ProceduralVegetationEditor/Source/ProceduralVegetation/Public/ProceduralVegetation.h
+++ b/Engine/Plugins/Experimental/ProceduralVegetationEditor/Source/ProceduralVegetation/Public/ProceduralVegetation.h
@@ -17,7 +17,7 @@ public:
 		: UPCGGraph(ObjectInitializer)
 	{
 #if WITH_EDITORONLY_DATA
-		bIsStandaloneGraph = true;
+		bIsStandaloneGraph = false;
 		bExposeGenerationInAssetExplorer = false;
 #endif
 		
