Add PVDebugMode CVar. Lock PVDebugVisualization and PVImporterSettings behind PVDebugMode.

diff --git a/Engine/Plugins/Experimental/ProceduralVegetationEditor/Source/ProceduralVegetation/Private/DataTypes/PVData.cpp b/Engine/Plugins/Experimental/ProceduralVegetationEditor/Source/ProceduralVegetation/Private/DataTypes/PVData.cpp
--- a/Engine/Plugins/Experimental/ProceduralVegetationEditor/Source/ProceduralVegetation/Private/DataTypes/PVData.cpp
+++ b/Engine/Plugins/Experimental/ProceduralVegetationEditor/Source/ProceduralVegetation/Private/DataTypes/PVData.cpp
@@ -33,11 +33,7 @@ void UPVData::AddToCrc(FArchiveCrc32& Ar, bool bFullDataCrc) const
 		Ar << Collection;
 #if WITH_EDITOR
 		Ar << RenderType;
-
-#if PV_SHOW_DEBUG_VISUALIZATION 
 		Ar << DebugSettings;
-#endif
-		
 #endif
 	}
 	else
@@ -75,11 +71,7 @@ UPCGSpatialData* UPVData::CopyInternal(FPCGContext* Context) const
 
 #if WITH_EDITOR
 	NewProceduralVegetationData->RenderType = RenderType;
-
-#if PV_SHOW_DEBUG_VISUALIZATION 
 	NewProceduralVegetationData->DebugSettings = DebugSettings;
-#endif
-	
 #endif
 	
 	return NewProceduralVegetationData;
diff --git a/Engine/Plugins/Experimental/ProceduralVegetationEditor/Source/ProceduralVegetation/Private/Helpers/PVUtilities.cpp b/Engine/Plugins/Experimental/ProceduralVegetationEditor/Source/ProceduralVegetation/Private/Helpers/PVUtilities.cpp
--- a/Engine/Plugins/Experimental/ProceduralVegetationEditor/Source/ProceduralVegetation/Private/Helpers/PVUtilities.cpp
+++ b/Engine/Plugins/Experimental/ProceduralVegetationEditor/Source/ProceduralVegetation/Private/Helpers/PVUtilities.cpp
@@ -9,9 +9,31 @@
 #include "UObject/Package.h"
 #include "Components/InstancedStaticMeshComponent.h"
 #include "Components/MeshComponent.h"
+#include "Nodes/PVImporterSettings.h"
 
 namespace PV::Utilities
 {
+#if WITH_EDITOR
+	static TAutoConsoleVariable<bool> CVarEnablePVDebugMode(
+		TEXT("PV.DebugMode.Enabled"),
+		false,
+		TEXT("Enables debug mode for the Procedural Vegetation Editor"), 
+		FConsoleVariableDelegate::CreateLambda([](IConsoleVariable* ConsoleVariable)
+			{
+				UPVImporterSettings::StaticClass()->GetDefaultObject<UPVImporterSettings>()->bExposeToLibrary = ConsoleVariable->GetBool();
+			})
+	);
+#endif
+
+	bool DebugModeEnabled()
+	{
+#if WITH_EDITOR
+		return CVarEnablePVDebugMode.GetValueOnAnyThread();
+#else
+		return false;
+#endif
+	}
+
 	void AddFoliageInstancesToActor(const Facades::FFoliageFacade& FoliageFacade, AActor* Actor,
 	                                TMap<FString, TObjectPtr<UMeshComponent>>& OutInstancedComponentMap)
 	{
diff --git a/Engine/Plugins/Experimental/ProceduralVegetationEditor/Source/ProceduralVegetation/Private/Nodes/PVBaseSettings.cpp b/Engine/Plugins/Experimental/ProceduralVegetationEditor/Source/ProceduralVegetation/Private/Nodes/PVBaseSettings.cpp
--- a/Engine/Plugins/Experimental/ProceduralVegetationEditor/Source/ProceduralVegetation/Private/Nodes/PVBaseSettings.cpp
+++ b/Engine/Plugins/Experimental/ProceduralVegetationEditor/Source/ProceduralVegetation/Private/Nodes/PVBaseSettings.cpp
@@ -59,10 +59,7 @@ void FPVBaseElement::PostExecuteInternal(FPCGContext* InContext) const
 			if (const UPVData* ProceduralVegetationData = Cast<UPVData>(OutputData.Data))
 			{
 				ProceduralVegetationData->SetRenderType(Settings->GetRenderType());
-
-#if PV_SHOW_DEBUG_VISUALIZATION
 				ProceduralVegetationData->SetDebugSettings(Settings->GetDebugVisualizationSettings());
-#endif
 			}
 		}
 	}
diff --git a/Engine/Plugins/Experimental/ProceduralVegetationEditor/Source/ProceduralVegetation/Private/Nodes/PVImporterSettings.cpp b/Engine/Plugins/Experimental/ProceduralVegetationEditor/Source/ProceduralVegetation/Private/Nodes/PVImporterSettings.cpp
--- a/Engine/Plugins/Experimental/ProceduralVegetationEditor/Source/ProceduralVegetation/Private/Nodes/PVImporterSettings.cpp
+++ b/Engine/Plugins/Experimental/ProceduralVegetationEditor/Source/ProceduralVegetation/Private/Nodes/PVImporterSettings.cpp
@@ -13,6 +13,16 @@
 #include "Helpers/PCGSettingsHelpers.h"
 #include "Helpers/PVJSONHelper.h"
 
+UPVImporterSettings::UPVImporterSettings()
+{
+#if WITH_EDITOR
+	if (HasAnyFlags(RF_ClassDefaultObject))
+	{
+		bExposeToLibrary = PV::Utilities::DebugModeEnabled();
+	}
+#endif
+}
+
 TArray<FPCGPinProperties> UPVImporterSettings::InputPinProperties() const
 {
 	return TArray<FPCGPinProperties>();
diff --git a/Engine/Plugins/Experimental/ProceduralVegetationEditor/Source/ProceduralVegetation/Private/Nodes/PVImporterSettings.h b/Engine/Plugins/Experimental/ProceduralVegetationEditor/Source/ProceduralVegetation/Private/Nodes/PVImporterSettings.h
--- a/Engine/Plugins/Experimental/ProceduralVegetationEditor/Source/ProceduralVegetation/Private/Nodes/PVImporterSettings.h
+++ b/Engine/Plugins/Experimental/ProceduralVegetationEditor/Source/ProceduralVegetation/Private/Nodes/PVImporterSettings.h
@@ -13,6 +13,8 @@ class UPVImporterSettings : public UPVBaseSettings
 	GENERATED_BODY()
 
 public:
+	UPVImporterSettings();
+
 #if WITH_EDITOR
 	virtual FName GetDefaultNodeName() const override { return FName(TEXT("ProceduralVegetationImporter")); }
 	virtual FText GetDefaultNodeTitle() const override { return NSLOCTEXT("PVImporterSettings", "NodeTitle", "Procedural Vegetation Importer"); }
diff --git a/Engine/Plugins/Experimental/ProceduralVegetationEditor/Source/ProceduralVegetation/Public/DataTypes/PVData.h b/Engine/Plugins/Experimental/ProceduralVegetationEditor/Source/ProceduralVegetation/Public/DataTypes/PVData.h
--- a/Engine/Plugins/Experimental/ProceduralVegetationEditor/Source/ProceduralVegetation/Public/DataTypes/PVData.h
+++ b/Engine/Plugins/Experimental/ProceduralVegetationEditor/Source/ProceduralVegetation/Public/DataTypes/PVData.h
@@ -117,9 +117,6 @@ public:
 #if WITH_EDITOR
 	const TArray<EPVRenderType>& GetRenderType() const { return RenderType; }
 	void SetRenderType(TArray<EPVRenderType> InRenderType) const { RenderType = InRenderType; }
-#endif
-
-#if PV_SHOW_DEBUG_VISUALIZATION 
 	void SetDebugSettings(const FPVDebugSettings& InSettings) const { DebugSettings = InSettings; }
 	const FPVDebugSettings& GetDebugSettings() const { return DebugSettings; }
 #endif
@@ -130,11 +127,7 @@ protected:
 
 #if WITH_EDITOR
 	mutable TArray<EPVRenderType> RenderType;
-
-#if PV_SHOW_DEBUG_VISUALIZATION 
 	mutable FPVDebugSettings DebugSettings;
-#endif
-	
 #endif
 	
 	// ~Begin UPCGData interface
diff --git a/Engine/Plugins/Experimental/ProceduralVegetationEditor/Source/ProceduralVegetation/Public/Helpers/PVUtilities.h b/Engine/Plugins/Experimental/ProceduralVegetationEditor/Source/ProceduralVegetation/Public/Helpers/PVUtilities.h
--- a/Engine/Plugins/Experimental/ProceduralVegetationEditor/Source/ProceduralVegetation/Public/Helpers/PVUtilities.h
+++ b/Engine/Plugins/Experimental/ProceduralVegetationEditor/Source/ProceduralVegetation/Public/Helpers/PVUtilities.h
@@ -7,12 +7,10 @@ class UMeshComponent;
 class FFoliageFacade;
 class AActor;
 
-#ifndef PV_SHOW_DEBUG_VISUALIZATION
-	#define PV_SHOW_DEBUG_VISUALIZATION WITH_EDITOR && UE_BUILD_DEBUG 
-#endif
-
 namespace PV::Utilities
 {
+	bool PROCEDURALVEGETATION_API DebugModeEnabled();
+
 	DECLARE_DELEGATE_OneParam(FFoliageComponentCreatedCallback, UMeshComponent*);
 	DECLARE_DELEGATE_OneParam(FFoliageInstanceAddedCallback, int32);
 	
diff --git a/Engine/Plugins/Experimental/ProceduralVegetationEditor/Source/ProceduralVegetation/Public/Nodes/PVBaseSettings.h b/Engine/Plugins/Experimental/ProceduralVegetationEditor/Source/ProceduralVegetation/Public/Nodes/PVBaseSettings.h
--- a/Engine/Plugins/Experimental/ProceduralVegetationEditor/Source/ProceduralVegetation/Public/Nodes/PVBaseSettings.h
+++ b/Engine/Plugins/Experimental/ProceduralVegetationEditor/Source/ProceduralVegetation/Public/Nodes/PVBaseSettings.h
@@ -16,9 +16,7 @@ public:
 #if WITH_EDITOR
 	virtual EPCGSettingsType GetType() const override { return EPCGSettingsType::Spatial; }
 	virtual TArray<EPVRenderType> GetRenderType() const { return TArray<EPVRenderType>(); }
-#endif
 
-#if PV_SHOW_DEBUG_VISUALIZATION
 	FPVDebugSettings GetDebugVisualizationSettings() const { return DebugVisualizationSettings; }
 #endif
 
@@ -31,8 +29,10 @@ protected:
 	virtual FPCGDataTypeIdentifier GetOutputPinTypeIdentifier() const;
 private:
 	
+#if WITH_EDITORONLY_DATA
 	UPROPERTY(EditAnywhere, Category = DebugSettings, meta=(EditCondition = "bDebug", EditConditionHides))
 	FPVDebugSettings DebugVisualizationSettings;
+#endif
 };
 
 class FPVBaseElement : public IPCGElement
diff --git a/Engine/Plugins/Experimental/ProceduralVegetationEditor/Source/ProceduralVegetationEditor/Private/Visualizations/DebugVisualization/PVDataAttributeValue.cpp b/Engine/Plugins/Experimental/ProceduralVegetationEditor/Source/ProceduralVegetationEditor/Private/Visualizations/DebugVisualization/PVDataAttributeValue.cpp
--- a/Engine/Plugins/Experimental/ProceduralVegetationEditor/Source/ProceduralVegetationEditor/Private/Visualizations/DebugVisualization/PVDataAttributeValue.cpp
+++ b/Engine/Plugins/Experimental/ProceduralVegetationEditor/Source/ProceduralVegetationEditor/Private/Visualizations/DebugVisualization/PVDataAttributeValue.cpp
@@ -2,8 +2,6 @@
 
 #include "PVDataAttributeValue.h"
 
-#if PV_SHOW_DEBUG_VISUALIZATION 
-
 #include "GeometryCollection/ManagedArrayCollection.h"
 
 FPVAttributeScalarValue::FPVAttributeScalarValue(const FName InGroupName, const FName InAttributeName, EManagedArrayType InAttributeType,
@@ -100,5 +98,3 @@ TArray<FVector> FPVAttributeVectorValue::ToVectors()
 
 	return Vectors;
 }
-
-#endif
\ No newline at end of file
diff --git a/Engine/Plugins/Experimental/ProceduralVegetationEditor/Source/ProceduralVegetationEditor/Private/Visualizations/DebugVisualization/PVDataAttributeValue.h b/Engine/Plugins/Experimental/ProceduralVegetationEditor/Source/ProceduralVegetationEditor/Private/Visualizations/DebugVisualization/PVDataAttributeValue.h
--- a/Engine/Plugins/Experimental/ProceduralVegetationEditor/Source/ProceduralVegetationEditor/Private/Visualizations/DebugVisualization/PVDataAttributeValue.h
+++ b/Engine/Plugins/Experimental/ProceduralVegetationEditor/Source/ProceduralVegetationEditor/Private/Visualizations/DebugVisualization/PVDataAttributeValue.h
@@ -4,8 +4,6 @@
 
 #include "Helpers/PVUtilities.h"
 
-#if PV_SHOW_DEBUG_VISUALIZATION 
-
 #include "PVDataAttributeValueInterface.h"
 
 class FPVAttributeScalarValue : public FPVAttributeValueInterface
@@ -30,5 +28,3 @@ public:
 private:
 	TArray<FVector> ValueArray;
 };
-
-#endif
\ No newline at end of file
diff --git a/Engine/Plugins/Experimental/ProceduralVegetationEditor/Source/ProceduralVegetationEditor/Private/Visualizations/DebugVisualization/PVDataAttributeValueInterface.cpp b/Engine/Plugins/Experimental/ProceduralVegetationEditor/Source/ProceduralVegetationEditor/Private/Visualizations/DebugVisualization/PVDataAttributeValueInterface.cpp
--- a/Engine/Plugins/Experimental/ProceduralVegetationEditor/Source/ProceduralVegetationEditor/Private/Visualizations/DebugVisualization/PVDataAttributeValueInterface.cpp
+++ b/Engine/Plugins/Experimental/ProceduralVegetationEditor/Source/ProceduralVegetationEditor/Private/Visualizations/DebugVisualization/PVDataAttributeValueInterface.cpp
@@ -2,8 +2,6 @@
 
 #include "PVDataAttributeValueInterface.h"
 
-#if PV_SHOW_DEBUG_VISUALIZATION 
-
 #include "GeometryCollection/ManagedArrayCollection.h"
 #include "PVDataAttributeValue.h"
 
@@ -20,5 +18,4 @@ FPVAttributeValuePtr FPVAttributeValueInterface::Create(const FName InGroupName,
 	default:
 		return nullptr;
 	}
-}
-#endif
\ No newline at end of file
+}
\ No newline at end of file
diff --git a/Engine/Plugins/Experimental/ProceduralVegetationEditor/Source/ProceduralVegetationEditor/Private/Visualizations/DebugVisualization/PVDataAttributeValueInterface.h b/Engine/Plugins/Experimental/ProceduralVegetationEditor/Source/ProceduralVegetationEditor/Private/Visualizations/DebugVisualization/PVDataAttributeValueInterface.h
--- a/Engine/Plugins/Experimental/ProceduralVegetationEditor/Source/ProceduralVegetationEditor/Private/Visualizations/DebugVisualization/PVDataAttributeValueInterface.h
+++ b/Engine/Plugins/Experimental/ProceduralVegetationEditor/Source/ProceduralVegetationEditor/Private/Visualizations/DebugVisualization/PVDataAttributeValueInterface.h
@@ -4,8 +4,6 @@
 
 #include "Helpers/PVUtilities.h"
 
-#if PV_SHOW_DEBUG_VISUALIZATION
-
 #include "GeometryCollection/ManagedArray.h"
 #include "GeometryCollection/ManagedArrayTypes.h"
 
@@ -28,5 +26,3 @@ public:
 
 	static FPVAttributeValuePtr Create(const FName InGroupName, const FName InAttributeName, EManagedArrayType InAttributeType,const FManagedArrayCollection& InCollection);
 };
-
-#endif
\ No newline at end of file
diff --git a/Engine/Plugins/Experimental/ProceduralVegetationEditor/Source/ProceduralVegetationEditor/Private/Visualizations/DebugVisualization/PVDebugVisualizationBase.cpp b/Engine/Plugins/Experimental/ProceduralVegetationEditor/Source/ProceduralVegetationEditor/Private/Visualizations/DebugVisualization/PVDebugVisualizationBase.cpp
--- a/Engine/Plugins/Experimental/ProceduralVegetationEditor/Source/ProceduralVegetationEditor/Private/Visualizations/DebugVisualization/PVDebugVisualizationBase.cpp
+++ b/Engine/Plugins/Experimental/ProceduralVegetationEditor/Source/ProceduralVegetationEditor/Private/Visualizations/DebugVisualization/PVDebugVisualizationBase.cpp
@@ -2,8 +2,6 @@
 
 #include "PVDebugVisualizationBase.h"
 
-#if PV_SHOW_DEBUG_VISUALIZATION 
-
 #include "UObject/Package.h"
 #include "Engine/StaticMesh.h"
 #include "AdvancedPreviewScene.h"
@@ -207,5 +205,4 @@ void FPVDebugVisualizationBase::AddToScene(FPCGSceneSetupParams& InOutParams, UP
 	InOutParams.Scene->AddComponent(InComponent, FTransform::Identity);
 }
 
-#undef LOCTEXT_NAMESPACE
-#endif
\ No newline at end of file
+#undef LOCTEXT_NAMESPACE
\ No newline at end of file
diff --git a/Engine/Plugins/Experimental/ProceduralVegetationEditor/Source/ProceduralVegetationEditor/Private/Visualizations/DebugVisualization/PVDebugVisualizationBase.h b/Engine/Plugins/Experimental/ProceduralVegetationEditor/Source/ProceduralVegetationEditor/Private/Visualizations/DebugVisualization/PVDebugVisualizationBase.h
--- a/Engine/Plugins/Experimental/ProceduralVegetationEditor/Source/ProceduralVegetationEditor/Private/Visualizations/DebugVisualization/PVDebugVisualizationBase.h
+++ b/Engine/Plugins/Experimental/ProceduralVegetationEditor/Source/ProceduralVegetationEditor/Private/Visualizations/DebugVisualization/PVDebugVisualizationBase.h
@@ -4,8 +4,6 @@
 
 #include "Helpers/PVUtilities.h"
 
-#if PV_SHOW_DEBUG_VISUALIZATION 
-
 #include "GeometryCollection/ManagedArray.h"
 #include "Templates/SharedPointer.h"
 
@@ -47,6 +45,4 @@ protected:
 };
 typedef TSharedPtr<FPVDebugVisualizationBase> FPVDebugVisualizationPtr;
 
-DECLARE_LOG_CATEGORY_EXTERN(LogPVDebugVisualization, Log, All);
-
-#endif
\ No newline at end of file
+DECLARE_LOG_CATEGORY_EXTERN(LogPVDebugVisualization, Log, All);
\ No newline at end of file
diff --git a/Engine/Plugins/Experimental/ProceduralVegetationEditor/Source/ProceduralVegetationEditor/Private/Visualizations/DebugVisualization/PVDebugVisualizations.cpp b/Engine/Plugins/Experimental/ProceduralVegetationEditor/Source/ProceduralVegetationEditor/Private/Visualizations/DebugVisualization/PVDebugVisualizations.cpp
--- a/Engine/Plugins/Experimental/ProceduralVegetationEditor/Source/ProceduralVegetationEditor/Private/Visualizations/DebugVisualization/PVDebugVisualizations.cpp
+++ b/Engine/Plugins/Experimental/ProceduralVegetationEditor/Source/ProceduralVegetationEditor/Private/Visualizations/DebugVisualization/PVDebugVisualizations.cpp
@@ -2,8 +2,6 @@
 
 #include "PVDebugVisualizations.h"
 
-#if PV_SHOW_DEBUG_VISUALIZATION 
-
 #include "Facades/PVBranchFacade.h"
 #include "Facades/PVFoliageFacade.h"
 #include "Facades/PVPointFacade.h"
@@ -109,5 +107,4 @@ void FPVBranchDebugVisualization::GetPivot(const FManagedArrayCollection& InColl
 			OutScale = PointFacade.GetPointScale(RootPointIndex);
 		}
 	}
-}
-#endif
\ No newline at end of file
+}
\ No newline at end of file
diff --git a/Engine/Plugins/Experimental/ProceduralVegetationEditor/Source/ProceduralVegetationEditor/Private/Visualizations/DebugVisualization/PVDebugVisualizations.h b/Engine/Plugins/Experimental/ProceduralVegetationEditor/Source/ProceduralVegetationEditor/Private/Visualizations/DebugVisualization/PVDebugVisualizations.h
--- a/Engine/Plugins/Experimental/ProceduralVegetationEditor/Source/ProceduralVegetationEditor/Private/Visualizations/DebugVisualization/PVDebugVisualizations.h
+++ b/Engine/Plugins/Experimental/ProceduralVegetationEditor/Source/ProceduralVegetationEditor/Private/Visualizations/DebugVisualization/PVDebugVisualizations.h
@@ -4,8 +4,6 @@
 
 #include "Helpers/PVUtilities.h"
 
-#if PV_SHOW_DEBUG_VISUALIZATION
-
 #include "PVDebugVisualizationBase.h"
 #include "GeometryCollection/ManagedArray.h"
 
@@ -37,5 +35,3 @@ protected:
 	virtual TArray<FVector3f> GetPivotPositions(const FManagedArrayCollection& InCollection) override;
 	virtual void GetPivot(const FManagedArrayCollection& InCollection, const int InIndex, FVector3f& OutPos, float& OutScale) override;
 };
-
-#endif
\ No newline at end of file
diff --git a/Engine/Plugins/Experimental/ProceduralVegetationEditor/Source/ProceduralVegetationEditor/Private/Visualizations/DebugVisualization/PVDebugVisualizer.cpp b/Engine/Plugins/Experimental/ProceduralVegetationEditor/Source/ProceduralVegetationEditor/Private/Visualizations/DebugVisualization/PVDebugVisualizer.cpp
--- a/Engine/Plugins/Experimental/ProceduralVegetationEditor/Source/ProceduralVegetationEditor/Private/Visualizations/DebugVisualization/PVDebugVisualizer.cpp
+++ b/Engine/Plugins/Experimental/ProceduralVegetationEditor/Source/ProceduralVegetationEditor/Private/Visualizations/DebugVisualization/PVDebugVisualizer.cpp
@@ -2,8 +2,6 @@
 
 #include "PVDebugVisualizer.h"
 
-#if PV_SHOW_DEBUG_VISUALIZATION 
-
 #include "PVDebugVisualizationBase.h"
 #include "PVDebugVisualizations.h"
 #include "DataTypes/PVData.h"
@@ -52,5 +50,3 @@ FPVDebugVisualizationPtr FPVDebugVisualizer::CreateVisualizer(EPVDebugType InDeb
 			return nullptr;
 	}
 }
-
-#endif
\ No newline at end of file
diff --git a/Engine/Plugins/Experimental/ProceduralVegetationEditor/Source/ProceduralVegetationEditor/Private/Visualizations/DebugVisualization/PVDebugVisualizer.h b/Engine/Plugins/Experimental/ProceduralVegetationEditor/Source/ProceduralVegetationEditor/Private/Visualizations/DebugVisualization/PVDebugVisualizer.h
--- a/Engine/Plugins/Experimental/ProceduralVegetationEditor/Source/ProceduralVegetationEditor/Private/Visualizations/DebugVisualization/PVDebugVisualizer.h
+++ b/Engine/Plugins/Experimental/ProceduralVegetationEditor/Source/ProceduralVegetationEditor/Private/Visualizations/DebugVisualization/PVDebugVisualizer.h
@@ -3,8 +3,6 @@
 
 #include "Helpers/PVUtilities.h"
 
-#if PV_SHOW_DEBUG_VISUALIZATION 
-
 #include "PVDebugVisualizationBase.h"
 
 struct FManagedArrayCollection;
@@ -28,6 +26,4 @@ public:
 	static void Draw(const UPVData* InData,FPCGSceneSetupParams& InOutParams);
 private:
 	static FPVDebugVisualizationPtr CreateVisualizer(EPVDebugType InDebugType);
-};
-
-#endif
\ No newline at end of file
+};
\ No newline at end of file
diff --git a/Engine/Plugins/Experimental/ProceduralVegetationEditor/Source/ProceduralVegetationEditor/Private/Visualizations/PVDataVisualization.cpp b/Engine/Plugins/Experimental/ProceduralVegetationEditor/Source/ProceduralVegetationEditor/Private/Visualizations/PVDataVisualization.cpp
--- a/Engine/Plugins/Experimental/ProceduralVegetationEditor/Source/ProceduralVegetationEditor/Private/Visualizations/PVDataVisualization.cpp
+++ b/Engine/Plugins/Experimental/ProceduralVegetationEditor/Source/ProceduralVegetationEditor/Private/Visualizations/PVDataVisualization.cpp
@@ -103,12 +103,13 @@ FPCGSetupSceneFunc FPVDataVisualization::GetViewportSetupFunc(const UPCGSettings
 			InOutParams.Scene->AddComponent(MannequinComponent, FTransform(FVector(-Bounds.BoxExtent.X * 2.0, 0, 0)));
 		}
 
-#if PV_SHOW_DEBUG_VISUALIZATION
+		if (PV::Utilities::DebugModeEnabled())
+		{
 			if (WeakSettings.IsValid() && WeakSettings->bDebug)
 			{
 				FPVDebugVisualizer::Draw(WeakData.Get(), InOutParams);
 			}
-#endif
+		}
 
 		InOutParams.Scene->SetFloorVisibility(true);
 		InOutParams.Scene->SetFloorOffset(1.0f);
