[PVE] Mesh path and name validataions
Fixed crash when Mesh is added to the level
Fixed crash when aleady saved and read only asset is updated

diff --git a/Engine/Plugins/Experimental/ProceduralVegetationEditor/Source/ProceduralVegetation/Private/Helpers/PVExportHelper.cpp b/Engine/Plugins/Experimental/ProceduralVegetationEditor/Source/ProceduralVegetation/Private/Helpers/PVExportHelper.cpp
--- a/Engine/Plugins/Experimental/ProceduralVegetationEditor/Source/ProceduralVegetation/Private/Helpers/PVExportHelper.cpp
+++ b/Engine/Plugins/Experimental/ProceduralVegetationEditor/Source/ProceduralVegetation/Private/Helpers/PVExportHelper.cpp
@@ -32,6 +32,8 @@
 #include "Misc/Paths.h"
 #include "UObject/Package.h"
 #include "ProceduralVegetation.h"
+#include "SkinnedAssetCompiler.h"
+#include "StaticMeshCompiler.h"
 #include "GeometryCollection/Facades/CollectionConstraintOverrideFacade.h"
 
 namespace PV::Export::Internal
@@ -634,6 +636,7 @@ namespace PV::Export::Internal
 				//Build the nanite assemblies
             	BuildNaniteAssemblyData(ExportedCollection, ExportMesh);
             	ExportMesh->Build(true);
+				FStaticMeshCompilingManager::Get().FinishAllCompilation();
 			}
 			else
 			{
@@ -813,6 +816,7 @@ namespace PV::Export::Internal
 		}
 				
 		ExportMesh->Build();
+		FSkinnedAssetCompilingManager::Get().FinishAllCompilation();
 #endif
 	}
 
@@ -989,10 +993,30 @@ namespace PV::Export
 			UPackage* MeshPackage = CreatePackage(*FullPath);
 			UPackage* SkeletonPackage = CreatePackage(*(FullPath + "_Skeleton"));
 
-			// Create a skeletal mesh to save in content browser
-			const TObjectPtr<USkeletalMesh> ExportMesh = NewObject<USkeletalMesh>(MeshPackage, FName(MeshName), RF_Standalone | RF_Public);
-			const TObjectPtr<USkeleton> ExportSkeleton = NewObject<USkeleton>(SkeletonPackage, FName(MeshName + "_Skeleton"),
-			                                                                  RF_Standalone | RF_Public);
+			FAssetData AssetData;
+			FString ErrorMessage;
+
+			TObjectPtr<USkeletalMesh> ExportMesh = nullptr;
+			
+			if (Utilities::CanOverwriteAsset(FullPath, USkeletalMesh::StaticClass(), AssetData, ErrorMessage))
+			{
+				ExportMesh = CastChecked<USkeletalMesh>(AssetData.GetAsset());
+			}
+			else
+			{
+				ExportMesh = NewObject<USkeletalMesh>(MeshPackage, FName(MeshName), RF_Standalone | RF_Public);
+			}
+
+			TObjectPtr<USkeleton> ExportSkeleton = nullptr;
+			if (Utilities::CanOverwriteAsset(*(FullPath + "_Skeleton"), USkeletalMesh::StaticClass(), AssetData, ErrorMessage))
+			{
+				ExportSkeleton = CastChecked<USkeleton>(AssetData.GetAsset());
+			}
+			else
+			{
+				ExportSkeleton = NewObject<USkeleton>(SkeletonPackage, FName(MeshName + "_Skeleton"),RF_Standalone | RF_Public);
+			}
+			
 			Internal::AttachProceduralVegetationLink(ExportMesh, InProceduralVegetation);
 			Internal::AttachProceduralVegetationLink(ExportSkeleton, InProceduralVegetation);
 			
diff --git a/Engine/Plugins/Experimental/ProceduralVegetationEditor/Source/ProceduralVegetation/Private/Helpers/PVUtilities.cpp b/Engine/Plugins/Experimental/ProceduralVegetationEditor/Source/ProceduralVegetation/Private/Helpers/PVUtilities.cpp
--- a/Engine/Plugins/Experimental/ProceduralVegetationEditor/Source/ProceduralVegetation/Private/Helpers/PVUtilities.cpp
+++ b/Engine/Plugins/Experimental/ProceduralVegetationEditor/Source/ProceduralVegetation/Private/Helpers/PVUtilities.cpp
@@ -1,8 +1,8 @@
 ﻿// Copyright Epic Games, Inc. All Rights Reserved.
 
 #include "Helpers/PVUtilities.h"
-
 #include "ProceduralVegetationModule.h"
+#include "AssetRegistry/AssetRegistryModule.h"
 #include "Components/InstancedSkinnedMeshComponent.h"
 #include "Engine/StaticMesh.h"
 #include "Engine/SkeletalMesh.h"
@@ -10,6 +10,10 @@
 #include "Components/InstancedStaticMeshComponent.h"
 #include "Components/MeshComponent.h"
 #include "Nodes/PVImporterSettings.h"
+#include "Helpers/PVExportHelper.h"
+#include "Misc/PackageName.h"
+#include "AssetRegistry/AssetData.h"
+#include "Misc/FileHelper.h"
 
 namespace PV::Utilities
 {
@@ -127,4 +131,116 @@ namespace PV::Utilities
 			}
 		}
 	}
+
+	FString FormatLongPackageNameErrorCode(const FPackageName::EErrorCode ErrorCode)
+	{
+		switch (ErrorCode)
+		{
+		case FPackageName::EErrorCode::PackageNameUnknown:
+			return "Unknown Export Path error";
+		case FPackageName::EErrorCode::PackageNameEmptyPath:
+			return "Empty Export Path";
+		case FPackageName::EErrorCode::PackageNamePathNotMounted:
+			return "Mount Point for Export Path is not valid";
+		case FPackageName::EErrorCode::PackageNamePathIsMemoryOnly:
+			return "Export Path exists in memory only. ";
+		case FPackageName::EErrorCode::PackageNameSpacesNotAllowed:
+			return "Spaces are not allowed in Export Path";	
+		case FPackageName::EErrorCode::PackageNameContainsInvalidCharacters:
+			return FString::Printf(TEXT("Export Path contains one of \n%s\n Invalid Characters"), TEXT(R"(\\:*?\"<>|' ,.&!~\n\r\t@#)"));
+		case FPackageName::EErrorCode::LongPackageNames_PathTooShort:
+			return "Export Path is too small to be a valid path";
+		case FPackageName::EErrorCode::LongPackageNames_PathWithNoStartingSlash:
+			return "Export Path has to start with a slash [/]";
+		case FPackageName::EErrorCode::LongPackageNames_PathWithTrailingSlash:
+			return "Export Path cannot end with a slash [/]";
+		case FPackageName::EErrorCode::LongPackageNames_PathWithDoubleSlash:
+			return "Export Path contains a double slash [/]";
+		}
+		return FString();
+	}
+
+	TArray<FAssetData> GetExistingAssets(FString PackageName)
+	{
+		// Get the Asset Registry Module
+		FAssetRegistryModule& AssetRegistryModule = FModuleManager::LoadModuleChecked<FAssetRegistryModule>("AssetRegistry");
+		// Get all asset data in the Content Browser
+		TArray<FAssetData> AssetData;
+		AssetRegistryModule.Get().GetAssetsByPackageName(FName(*PackageName), AssetData);
+		return AssetData;
+	}
+
+	bool ValidateMeshPathAndName(FString MeshName, FString Path, UClass* InClass, FString& OutError)
+	{
+		if (FPackageName::EErrorCode ErrorCode; !FPackageName::IsValidLongPackageName(Path, false, &ErrorCode))
+        {
+			OutError =  FormatLongPackageNameErrorCode(ErrorCode);
+        	return false;
+        }
+
+       if (FText Reason; !IsFileNameValid(FName(MeshName), Reason))
+        {
+			OutError = Reason.ToString();
+        	return false;
+        }
+
+        if (IsInGameThread())
+        {
+        	FString PackageName = FPaths::Combine(Path, MeshName);
+        	bool bPackageExists = GetExistingAssets(PackageName).Num() > 0;
+        	
+        	FAssetData Asset;
+        	if (bPackageExists && !CanOverwriteAsset(PackageName, InClass, Asset, OutError))
+        	{
+        		 return false;
+        	}
+        }
+		else
+		{
+			OutError = FString(TEXT("ValidateMeshPathAndName needs to be called from Game Thread. Returning validation failed result."));
+			return false;
+		}
+
+		return true;
+	}
+
+	bool IsFileNameValid(FName FileName, FText& Reason)
+	{
+		FileName.IsValidObjectName(Reason);
+		FFileHelper::IsFilenameValidForSaving(FileName.ToString(), Reason);
+
+		return Reason.IsEmpty();
+	}
+
+	bool CanOverwriteAsset(FString PackageName, UClass* InClass, FAssetData& OutAsset, FString& OutError)
+	{
+		bool bCanOverwrite = false;
+		
+		TArray<FAssetData> AssetData = GetExistingAssets(PackageName);
+
+		if (AssetData.Num() > 0)
+		{
+			FString AssetType = "";
+			// Iterate through the asset data
+			for (FAssetData& Asset : AssetData)
+			{
+				AssetType = Asset.GetClass()->GetName();
+				if (Asset.IsInstanceOf(InClass))
+				{
+					OutAsset = Asset;
+					bCanOverwrite = true;
+					break;
+				}
+			}
+			// It is an error if we are trying to replace an object of a different class
+			if (!bCanOverwrite)
+			{
+				bCanOverwrite = false;
+				OutError = FString::Format(TEXT("Create a Unique name for the output as another asset of different type {0} exists with same name {1}"), {*AssetType, *PackageName});
+				return false;
+			}
+		}
+
+		return bCanOverwrite;
+	}
 }
diff --git a/Engine/Plugins/Experimental/ProceduralVegetationEditor/Source/ProceduralVegetation/Private/Nodes/PVOutputSettings.cpp b/Engine/Plugins/Experimental/ProceduralVegetationEditor/Source/ProceduralVegetation/Private/Nodes/PVOutputSettings.cpp
--- a/Engine/Plugins/Experimental/ProceduralVegetationEditor/Source/ProceduralVegetation/Private/Nodes/PVOutputSettings.cpp
+++ b/Engine/Plugins/Experimental/ProceduralVegetationEditor/Source/ProceduralVegetation/Private/Nodes/PVOutputSettings.cpp
@@ -9,34 +9,6 @@
 
 #define LOCTEXT_NAMESPACE "PVOutputSettings"
 
-FString FormatLongPackageNameErrorCode(const FPackageName::EErrorCode ErrorCode)
-{
-	switch (ErrorCode)
-	{
-	case FPackageName::EErrorCode::PackageNameUnknown:
-		return "Unknown Export Path error";
-	case FPackageName::EErrorCode::PackageNameEmptyPath:
-		return "Empty Export Path";
-	case FPackageName::EErrorCode::PackageNamePathNotMounted:
-		return "Mount Point for Export Path is not valid";
-	case FPackageName::EErrorCode::PackageNamePathIsMemoryOnly:
-		return "Export Path exists in memory only. ";
-	case FPackageName::EErrorCode::PackageNameSpacesNotAllowed:
-		return "Spaces are not allowed in Export Path";	
-	case FPackageName::EErrorCode::PackageNameContainsInvalidCharacters:
-		return FString::Printf(TEXT("Export Path contains one of \n%s\n Invalid Characters"), TEXT(R"(\\:*?\"<>|' ,.&!~\n\r\t@#)"));
-	case FPackageName::EErrorCode::LongPackageNames_PathTooShort:
-		return "Export Path is too small to be a valid path";
-	case FPackageName::EErrorCode::LongPackageNames_PathWithNoStartingSlash:
-		return "Export Path has to start with a slash [/]";
-	case FPackageName::EErrorCode::LongPackageNames_PathWithTrailingSlash:
-		return "Export Path cannot end with a slash [/]";
-	case FPackageName::EErrorCode::LongPackageNames_PathWithDoubleSlash:
-		return "Export Path contains a double slash [/]";
-	}
-	return FString();
-}
-
 UPVOutputSettings::UPVOutputSettings()
 {
 	if (PCGHelpers::IsNewObjectAndNotDefault(this))
@@ -79,9 +51,12 @@ bool FPVOutputElement::ExecuteInternal(FPCGContext* InContext) const
 		if (const UPVData* InputData = Cast<UPVData>(Input.Data))
 		{
 			const FString OutputPath = Settings->ExportSettings.ContentBrowserFolder.Path;
-			if (FPackageName::EErrorCode ErrorCode; !FPackageName::IsValidLongPackageName(OutputPath, false, &ErrorCode))
+			const FName MeshName = Settings->ExportSettings.MeshName;
+
+			UClass* MeshClass = Settings->ExportSettings.ExportMeshType == EPVExportMeshType::StaticMesh ? UStaticMesh::StaticClass() : USkeletalMesh::StaticClass();
+			if (FString Reason; !PV::Utilities::ValidateMeshPathAndName(MeshName.ToString(), OutputPath, MeshClass, Reason))
 			{
-				PCGLog::LogErrorOnGraph(FText::FromString(FormatLongPackageNameErrorCode(ErrorCode)), InContext);
+				PCGLog::LogErrorOnGraph(FText::FromString(Reason), InContext);
 				return true;
 			}
 
diff --git a/Engine/Plugins/Experimental/ProceduralVegetationEditor/Source/ProceduralVegetation/Public/Helpers/PVUtilities.h b/Engine/Plugins/Experimental/ProceduralVegetationEditor/Source/ProceduralVegetation/Public/Helpers/PVUtilities.h
--- a/Engine/Plugins/Experimental/ProceduralVegetationEditor/Source/ProceduralVegetation/Public/Helpers/PVUtilities.h
+++ b/Engine/Plugins/Experimental/ProceduralVegetationEditor/Source/ProceduralVegetation/Public/Helpers/PVUtilities.h
@@ -16,4 +16,10 @@ namespace PV::Utilities
 	
 	void PROCEDURALVEGETATION_API AddFoliageInstancesToActor(const Facades::FFoliageFacade& InFoliageFacade, AActor* InActor, TMap<FString, TObjectPtr<UMeshComponent>>& OutInstancedComponentMap);
 	void PROCEDURALVEGETATION_API AddFoliageInstances(const Facades::FFoliageFacade& InFoliageFacade, const UObject* InParent, FFoliageComponentCreatedCallback InCallback, TMap<FString, TObjectPtr<UMeshComponent>>& OutInstancedComponentMap);
+
+	bool PROCEDURALVEGETATION_API ValidateMeshPathAndName(FString MeshName, FString Path, UClass* InClass, FString& OutError);
+	
+	bool IsFileNameValid(FName FileName, FText& Reason);
+
+	bool CanOverwriteAsset(FString PackageName, UClass* InClass, FAssetData& OutAsset, FString& OutError);
 }
diff --git a/Engine/Plugins/Experimental/ProceduralVegetationEditor/Source/ProceduralVegetationEditor/Private/PVEditor.cpp b/Engine/Plugins/Experimental/ProceduralVegetationEditor/Source/ProceduralVegetationEditor/Private/PVEditor.cpp
--- a/Engine/Plugins/Experimental/ProceduralVegetationEditor/Source/ProceduralVegetationEditor/Private/PVEditor.cpp
+++ b/Engine/Plugins/Experimental/ProceduralVegetationEditor/Source/ProceduralVegetationEditor/Private/PVEditor.cpp
@@ -304,7 +304,9 @@ FReply FPVEditor::ExportInternal()
 
 		FScopedSlowTask Progress(NodesToExport.Num(), NSLOCTEXT("ExportAsset", "ExportAsset_SlowTask", "Exporting Asset(s) ..."));
 		Progress.MakeDialog();
-		
+
+		FString ErrorMessage;
+					
 		for(auto Node : NodesToExport)
 		{
 			if(const UPVOutputSettings* OutputSettings = Cast<UPVOutputSettings>(Node->GetSettings()))
@@ -313,11 +315,33 @@ FReply FPVEditor::ExportInternal()
 				FName PinName = NAME_None;
 				if (GetCollectionFromNode(Node, PinName, OutCollection))
 				{
-					PV::Export::ExportCollectionAsMesh(ProceduralVegetationBeingEdited, OutCollection, OutCollection, OutputSettings->ExportSettings);
+					FString Reason;
+					UClass* MeshClass = OutputSettings->ExportSettings.ExportMeshType == EPVExportMeshType::StaticMesh ? UStaticMesh::StaticClass() : USkeletalMesh::StaticClass();
+					if (PV::Utilities::ValidateMeshPathAndName(OutputSettings->ExportSettings.MeshName.ToString(), OutputSettings->ExportSettings.ContentBrowserFolder.Path,
+						MeshClass, Reason))
+					{
+						PV::Export::ExportCollectionAsMesh(ProceduralVegetationBeingEdited, OutCollection, OutCollection, OutputSettings->ExportSettings);
+					}
+					else
+					{
+						ErrorMessage += FString::Format( TEXT("Export failed for mesh {0} with error {1}"), {OutputSettings->ExportSettings.MeshName.ToString(), Reason});
+					}
 					Progress.EnterProgressFrame();
 				}
+				else
+				{
+					ErrorMessage += FString::Format( TEXT("Export failed for output node {0}, Invalid output please fix the errors on the node."), { Node->GetAuthoredTitleName().ToString()});
+				}
 			}
 		}
+
+		if (!ErrorMessage.IsEmpty())
+		{
+			FMessageLog EditorErrors("EditorErrors");
+			EditorErrors.Error(LOCTEXT("PVE Mesh Export Failed", "PVE Mesh Export Failed"));
+			EditorErrors.Info(FText::FromString(ErrorMessage));
+			EditorErrors.Notify(LOCTEXT("PVE Mesh Export Failed", "PVE Mesh Export Failed"));
+		}
 	}
 	
 	return FReply::Handled();
