Fix crash when modifying displacement texture on Mesh Builder node

diff --git a/Engine/Plugins/Experimental/ProceduralVegetationEditor/Source/ProceduralVegetation/Private/Implementations/PVMeshBuilder.cpp b/Engine/Plugins/Experimental/ProceduralVegetationEditor/Source/ProceduralVegetation/Private/Implementations/PVMeshBuilder.cpp
--- a/Engine/Plugins/Experimental/ProceduralVegetationEditor/Source/ProceduralVegetation/Private/Implementations/PVMeshBuilder.cpp
+++ b/Engine/Plugins/Experimental/ProceduralVegetationEditor/Source/ProceduralVegetation/Private/Implementations/PVMeshBuilder.cpp
@@ -955,7 +955,7 @@ void FPVMeshBuilder::ExtractDisplacementData(const TObjectPtr<UTexture2D>& Textu
 	const int32 NumPix = Width * Height;
 	OutValues.SetNumZeroed(NumPix);
 
-	const uint8* RawData = Source.LockMip(0);
+	const uint8* RawData = Source.LockMipReadOnly(0);
 
 	switch (Source.GetFormat())
 	{
