[PVE] Warning on export if Nanite Assemblies not enabled

diff --git a/Engine/Plugins/Experimental/ProceduralVegetationEditor/Config/BaseProceduralVegetationEditor.ini b/Engine/Plugins/Experimental/ProceduralVegetationEditor/Config/BaseProceduralVegetationEditor.ini
--- a/Engine/Plugins/Experimental/ProceduralVegetationEditor/Config/BaseProceduralVegetationEditor.ini
+++ b/Engine/Plugins/Experimental/ProceduralVegetationEditor/Config/BaseProceduralVegetationEditor.ini
@@ -56,4 +56,5 @@
 
 +PropertyRedirects=(OldName="PVPCGMegaSeedLoadSettings.RecipeDataAsset",NewName="PVPresetLoaderSettings.Preset")
 
-+StructRedirects=(OldName="/Script/Engine.RuntimeFloatCurve",NewName="/Script/ProceduralVegetation.PVFloatRamp")
\ No newline at end of file
++StructRedirects=(OldName="/Script/Engine.RuntimeFloatCurve",NewName="/Script/ProceduralVegetation.PVFloatRamp")
++PropertyRedirects=(OldName="/Script/ProceduralVegetation.PVExportParams.bCreateNaniteAssemblies",NewName="/Script/ProceduralVegetation.PVExportParams.bCreateNaniteFoliage")
\ No newline at end of file
diff --git a/Engine/Plugins/Experimental/ProceduralVegetationEditor/Source/ProceduralVegetation/Private/Helpers/PVExportHelper.cpp b/Engine/Plugins/Experimental/ProceduralVegetationEditor/Source/ProceduralVegetation/Private/Helpers/PVExportHelper.cpp
--- a/Engine/Plugins/Experimental/ProceduralVegetationEditor/Source/ProceduralVegetation/Private/Helpers/PVExportHelper.cpp
+++ b/Engine/Plugins/Experimental/ProceduralVegetationEditor/Source/ProceduralVegetation/Private/Helpers/PVExportHelper.cpp
@@ -980,7 +980,7 @@ namespace PV::Export
 			PV::Export::Internal::ExportCollectionToStaticMesh(
 				ExportMesh,
 				Collection,
-				ExportParams.bCreateNaniteAssemblies,
+				ExportParams.bCreateNaniteFoliage,
 				ExportParams.NaniteShapePreservation
 			);
 
@@ -1026,7 +1026,7 @@ namespace PV::Export
 				ExportMesh,
 				Collection,
 				ExportParams.NaniteShapePreservation,
-				ExportParams.bCreateNaniteAssemblies
+				ExportParams.bCreateNaniteFoliage
 			);
 
 			MeshPackage->SetDirtyFlag(true);
diff --git a/Engine/Plugins/Experimental/ProceduralVegetationEditor/Source/ProceduralVegetation/Public/Helpers/PVExportHelper.h b/Engine/Plugins/Experimental/ProceduralVegetationEditor/Source/ProceduralVegetation/Public/Helpers/PVExportHelper.h
--- a/Engine/Plugins/Experimental/ProceduralVegetationEditor/Source/ProceduralVegetation/Public/Helpers/PVExportHelper.h
+++ b/Engine/Plugins/Experimental/ProceduralVegetationEditor/Source/ProceduralVegetation/Public/Helpers/PVExportHelper.h
@@ -39,8 +39,8 @@ struct PROCEDURALVEGETATION_API FPVExportParams
 	UPROPERTY(EditAnywhere, Category = "Export", Meta = (DisplayName = "Export Mesh Type"))
 	EPVExportMeshType ExportMeshType = EPVExportMeshType::StaticMesh;
 
-	UPROPERTY(EditAnywhere, Category = "Export", Meta = (DisplayName = "Create Nanite Assemblies"))
-	bool bCreateNaniteAssemblies = true;
+	UPROPERTY(EditAnywhere, Category = "Export", Meta = (DisplayName = "Create Nanite Foliage"))
+	bool bCreateNaniteFoliage = true;
 	
 	UPROPERTY(EditAnywhere, Category = "Export", Meta = (DisplayName = "Nanite Shape Preservation Method"))
 	ENaniteShapePreservation NaniteShapePreservation = ENaniteShapePreservation::Voxelize;
diff --git a/Engine/Plugins/Experimental/ProceduralVegetationEditor/Source/ProceduralVegetationEditor/Private/Customizations/PVOutputSettingsCustomizations.cpp b/Engine/Plugins/Experimental/ProceduralVegetationEditor/Source/ProceduralVegetationEditor/Private/Customizations/PVOutputSettingsCustomizations.cpp
new file mode 100644
--- /dev/null
+++ b/Engine/Plugins/Experimental/ProceduralVegetationEditor/Source/ProceduralVegetationEditor/Private/Customizations/PVOutputSettingsCustomizations.cpp
@@ -0,0 +1,219 @@
+// Copyright Epic Games, Inc. All Rights Reserved.
+
+#include "PVOutputSettingsCustomizations.h"
+
+#include "DetailLayoutBuilder.h"
+#include "DetailWidgetRow.h"
+#include "IDetailChildrenBuilder.h"
+#include "ISettingsModule.h"
+#include "RenderUtils.h"
+
+#include "Framework/Docking/TabManager.h"
+
+#include "Helpers/PVExportHelper.h"
+
+#include "Nodes/PVOutputSettings.h"
+
+#include "Widgets/Input/SButton.h"
+#include "Widgets/Layout/SBox.h"
+#include "Widgets/Text/STextBlock.h"
+
+FReply OnOpenProjectSettings()
+{
+	FModuleManager::GetModulePtr<ISettingsModule>("Settings")->ShowViewer("Project", "Engine", "Rendering");
+	return FReply::Handled();
+}
+
+FString GetInternalName(const TSharedPtr<IPropertyHandle>& ChildHandle)
+{
+	return ChildHandle->GetProperty()->GetName();
+}
+
+TSharedRef<IPropertyTypeCustomization> FPVOutputSettingsCustomizations::MakeInstance()
+{
+	return MakeShareable(new FPVOutputSettingsCustomizations());
+}
+
+void FPVOutputSettingsCustomizations::CustomizeHeader(
+	TSharedRef<IPropertyHandle> PropertyHandle,
+	FDetailWidgetRow& HeaderRow,
+	IPropertyTypeCustomizationUtils& CustomizationUtils
+)
+{}
+
+void FPVOutputSettingsCustomizations::CustomizeChildren(
+	TSharedRef<IPropertyHandle> PropertyHandle,
+	IDetailChildrenBuilder& ChildBuilder,
+	IPropertyTypeCustomizationUtils& CustomizationUtils
+)
+{
+	uint32 NumChild;
+	PropertyHandle->GetNumChildren(NumChild);
+
+	for (uint32 ChildIndex = 0; ChildIndex < NumChild; ++ChildIndex)
+	{
+		const TSharedPtr<IPropertyHandle> ChildHandle = PropertyHandle->GetChildHandle(ChildIndex);
+		IDetailPropertyRow& RowBuilder = ChildBuilder.AddProperty(ChildHandle.ToSharedRef());
+
+		if (GetInternalName(ChildHandle) == GET_MEMBER_NAME_STRING_VIEW_CHECKED(FPVExportParams, bCreateNaniteFoliage))
+		{
+			CustomizeNaniteAssembliesWidget(ChildHandle, RowBuilder);
+		}
+		else if (GetInternalName(ChildHandle) == GET_MEMBER_NAME_STRING_VIEW_CHECKED(FPVExportParams, NaniteShapePreservation))
+		{
+			CustomizeNaniteShapePreservationWidget(ChildHandle, RowBuilder);
+		}
+	}
+}
+
+void FPVOutputSettingsCustomizations::CustomizeNaniteAssembliesWidget(
+	const TSharedPtr<IPropertyHandle>& ChildHandle,
+	IDetailPropertyRow& RowBuilder
+)
+{
+	const static FText WarningText = FText::FromString(
+		"WARNING!! Enable Nanite Foliage in Project Settings to enable support for Nanite Assemblies."
+	);
+
+	const TSharedRef<SWidget> NameWidget = ChildHandle->CreatePropertyNameWidget();
+	const TSharedRef<SWidget> ValueWidget =
+		SNew(SHorizontalBox)
+		+ SHorizontalBox::Slot()
+		.HAlign(HAlign_Left)
+		[
+			ChildHandle->CreatePropertyValueWidget()
+		]
+		+ SHorizontalBox::Slot()
+		.HAlign(HAlign_Right)
+		[
+			SNew(SBox)
+			.Visibility_Lambda([=]()
+				{
+					if (bool bCreateNaniteAssemblies; ChildHandle->GetValue(bCreateNaniteAssemblies) == FPropertyAccess::Success)
+					{
+						return (bCreateNaniteAssemblies && !NaniteAssembliesSupported())
+							? EVisibility::Visible
+							: EVisibility::Collapsed;
+					}
+					return EVisibility::Collapsed;
+				})
+			[
+				SNew(SVerticalBox)
+				+ SVerticalBox::Slot()
+				.AutoHeight()
+				.VAlign(VAlign_Top)
+				[
+					SNew(STextBlock)
+					.AutoWrapText(true)
+					.Font(IDetailLayoutBuilder::GetDetailFont())
+					.ColorAndOpacity(FLinearColor::Yellow)
+					.ShadowColorAndOpacity(FLinearColor::Black)
+					.ShadowOffset(FVector2D(3.0f, 3.0f))
+					.Text(WarningText)
+					.ToolTipText(WarningText)
+				]
+				+ SVerticalBox::Slot()
+				.VAlign(VAlign_Bottom)
+				.HAlign(HAlign_Right)
+				.AutoHeight()
+				[
+					SNew(SButton)
+					.Text(FText::FromString("Open Project Settings"))
+					.TextStyle(&FAppStyle::Get().GetWidgetStyle<FTextBlockStyle>("TinyText"))
+					.OnClicked_Static(&OnOpenProjectSettings)
+				]
+			]
+		];
+
+	RowBuilder
+		.CustomWidget(true)
+		.NameContent()
+		[
+			NameWidget
+		]
+		.ValueContent()
+		.HAlign(HAlign_Fill)
+		[
+			ValueWidget
+		];
+}
+
+void FPVOutputSettingsCustomizations::CustomizeNaniteShapePreservationWidget(
+	const TSharedPtr<IPropertyHandle>& ChildHandle,
+	IDetailPropertyRow& RowBuilder
+)
+{
+	const static FText WarningText = FText::FromString(
+		"WARNING!! Voxels are not enabled for the project and therefore the exported mesh will not render as voxels. "
+		"Enable Nanite Foliage in the project settings to enable support."
+	);
+
+	const TSharedRef<SWidget> NameWidget = ChildHandle->CreatePropertyNameWidget();
+	const TSharedRef<SWidget> ValueWidget =
+		SNew(SHorizontalBox)
+		+ SHorizontalBox::Slot()
+		.HAlign(HAlign_Left)
+		[
+			SNew(SBox)
+			.MinDesiredWidth(125)
+			[
+				ChildHandle->CreatePropertyValueWidget()
+			]
+		]
+		+ SHorizontalBox::Slot()
+		.HAlign(HAlign_Right)
+		[
+			SNew(SBox)
+			.Visibility_Lambda([=]()
+				{
+					if (
+						ENaniteShapePreservation NaniteShapePreservation;
+						ChildHandle->GetValue(reinterpret_cast<uint8&>(NaniteShapePreservation)) == FPropertyAccess::Success
+					)
+					{
+						return (NaniteShapePreservation == ENaniteShapePreservation::Voxelize && !NaniteVoxelsSupported())
+							? EVisibility::Visible
+							: EVisibility::Collapsed;
+					}
+					return EVisibility::Collapsed;
+				})
+			[
+				SNew(SVerticalBox)
+				+ SVerticalBox::Slot()
+				.AutoHeight()
+				.VAlign(VAlign_Top)
+				[
+					SNew(STextBlock)
+					.AutoWrapText(true)
+					.Font(IDetailLayoutBuilder::GetDetailFont())
+					.ColorAndOpacity(FLinearColor::Yellow)
+					.ShadowColorAndOpacity(FLinearColor::Black)
+					.ShadowOffset(FVector2D(3.0f, 3.0f))
+					.Text(WarningText)
+					.ToolTipText(WarningText)
+				]
+				+ SVerticalBox::Slot()
+				.VAlign(VAlign_Bottom)
+				.HAlign(HAlign_Right)
+				.AutoHeight()
+				[
+					SNew(SButton)
+					.Text(FText::FromString("Open Project Settings"))
+					.TextStyle(&FAppStyle::Get().GetWidgetStyle<FTextBlockStyle>("TinyText"))
+					.OnClicked_Static(&OnOpenProjectSettings)
+				]
+			]
+		];
+
+	RowBuilder
+		.CustomWidget()
+		.NameContent()
+		[
+			NameWidget
+		]
+		.ValueContent()
+		.HAlign(HAlign_Fill)
+		[
+			ValueWidget
+		];
+}
diff --git a/Engine/Plugins/Experimental/ProceduralVegetationEditor/Source/ProceduralVegetationEditor/Private/Customizations/PVOutputSettingsCustomizations.h b/Engine/Plugins/Experimental/ProceduralVegetationEditor/Source/ProceduralVegetationEditor/Private/Customizations/PVOutputSettingsCustomizations.h
new file mode 100644
--- /dev/null
+++ b/Engine/Plugins/Experimental/ProceduralVegetationEditor/Source/ProceduralVegetationEditor/Private/Customizations/PVOutputSettingsCustomizations.h
@@ -0,0 +1,20 @@
+// Copyright Epic Games, Inc. All Rights Reserved.
+
+#pragma once
+
+#include "IPropertyTypeCustomization.h"
+
+class IDetailPropertyRow;
+
+class FPVOutputSettingsCustomizations : public IPropertyTypeCustomization
+{
+public:
+	static TSharedRef<IPropertyTypeCustomization> MakeInstance();
+
+	virtual void CustomizeHeader(TSharedRef<IPropertyHandle> PropertyHandle, FDetailWidgetRow& HeaderRow, IPropertyTypeCustomizationUtils& CustomizationUtils) override;
+	virtual void CustomizeChildren(TSharedRef<IPropertyHandle> PropertyHandle, IDetailChildrenBuilder& ChildBuilder, IPropertyTypeCustomizationUtils& CustomizationUtils) override;
+
+private:
+	void CustomizeNaniteAssembliesWidget(const TSharedPtr<IPropertyHandle>& ChildHandle, IDetailPropertyRow& RowBuilder);
+	void CustomizeNaniteShapePreservationWidget(const TSharedPtr<IPropertyHandle>& ChildHandle, IDetailPropertyRow& RowBuilder);
+};
diff --git a/Engine/Plugins/Experimental/ProceduralVegetationEditor/Source/ProceduralVegetationEditor/Private/ProceduralVegetationEditorModule.cpp b/Engine/Plugins/Experimental/ProceduralVegetationEditor/Source/ProceduralVegetationEditor/Private/ProceduralVegetationEditorModule.cpp
--- a/Engine/Plugins/Experimental/ProceduralVegetationEditor/Source/ProceduralVegetationEditor/Private/ProceduralVegetationEditorModule.cpp
+++ b/Engine/Plugins/Experimental/ProceduralVegetationEditor/Source/ProceduralVegetationEditor/Private/ProceduralVegetationEditorModule.cpp
@@ -23,6 +23,7 @@
 #include "Components/StaticMeshComponent.h"
 
 #include "Customizations/PVFloatRampCustomization.h"
+#include "Customizations/PVOutputSettingsCustomizations.h"
 
 #include "DataTypes/PVData.h"
 #include "DataTypes/PVFoliageMeshData.h"
@@ -35,6 +36,8 @@
 
 #include "GameFramework/Actor.h"
 
+#include "Helpers/PVExportHelper.h"
+
 #include "Modules/ModuleManager.h"
 
 #include "Templates/SharedPointer.h"
@@ -58,7 +61,7 @@ void FProceduralVegetationEditorModule::StartupModule()
 	DataVisRegistry.RegisterPCGDataVisualization(UPVData::StaticClass(), MakeUnique<const FPVDataVisualization>());
 
 	FPropertyEditorModule& PropertyEditor = FModuleManager::LoadModuleChecked<FPropertyEditorModule>("PropertyEditor");
-
+	PropertyEditor.RegisterCustomPropertyTypeLayout(FPVExportParams::StaticStruct()->GetFName(), FOnGetPropertyTypeCustomizationInstance::CreateStatic(&FPVOutputSettingsCustomizations::MakeInstance));
 	PropertyEditor.RegisterCustomPropertyTypeLayout(FPVFloatRamp::StaticStruct()->GetFName(), FOnGetPropertyTypeCustomizationInstance::CreateStatic(&FPVFloatRampCustomization::MakeInstance));
 }
 
@@ -68,8 +71,14 @@ void FProceduralVegetationEditorModule::ShutdownModule()
 	UnregisterPinColorAndIcons();
 	FPVEditorCommands::Unregister();
 
-	FPropertyEditorModule& PropertyEditor = FModuleManager::LoadModuleChecked<FPropertyEditorModule>("PropertyEditor");
-	PropertyEditor.UnregisterCustomPropertyTypeLayout(FPVFloatRamp::StaticStruct()->GetFName());
+	if (UObjectInitialized() && !IsEngineExitRequested())
+	{
+		if (FPropertyEditorModule* PropertyEditor = FModuleManager::GetModulePtr<FPropertyEditorModule>("PropertyEditor"))
+		{
+			PropertyEditor->UnregisterCustomPropertyTypeLayout(FPVExportParams::StaticStruct()->GetFName());
+			PropertyEditor->UnregisterCustomPropertyTypeLayout(FPVFloatRamp::StaticStruct()->GetFName());
+		}
+	}
 }
 
 void FProceduralVegetationEditorModule::RegisterMenus()
diff --git a/Engine/Plugins/Experimental/ProceduralVegetationEditor/Source/ProceduralVegetationEditor/ProceduralVegetationEditor.Build.cs b/Engine/Plugins/Experimental/ProceduralVegetationEditor/Source/ProceduralVegetationEditor/ProceduralVegetationEditor.Build.cs
--- a/Engine/Plugins/Experimental/ProceduralVegetationEditor/Source/ProceduralVegetationEditor/ProceduralVegetationEditor.Build.cs
+++ b/Engine/Plugins/Experimental/ProceduralVegetationEditor/Source/ProceduralVegetationEditor/ProceduralVegetationEditor.Build.cs
@@ -41,6 +41,7 @@ public class ProceduralVegetationEditor : ModuleRules
 				"PCGEditor",
 				"ProceduralVegetation",
 				"Projects",
+				"PropertyEditor",
 				"RenderCore",
 				"Slate",
 				"SlateCore",
