[PVE] Simplify Mesh simplification settings in the Mesher node and settings in Foliage distributor node and fix issue with Gravity node where the range is too wide

diff --git a/Engine/Plugins/Experimental/ProceduralVegetationEditor/Config/BaseProceduralVegetationEditor.ini b/Engine/Plugins/Experimental/ProceduralVegetationEditor/Config/BaseProceduralVegetationEditor.ini
--- a/Engine/Plugins/Experimental/ProceduralVegetationEditor/Config/BaseProceduralVegetationEditor.ini
+++ b/Engine/Plugins/Experimental/ProceduralVegetationEditor/Config/BaseProceduralVegetationEditor.ini
@@ -57,4 +57,12 @@
 +PropertyRedirects=(OldName="PVPCGMegaSeedLoadSettings.RecipeDataAsset",NewName="PVPresetLoaderSettings.Preset")
 
 +StructRedirects=(OldName="/Script/Engine.RuntimeFloatCurve",NewName="/Script/ProceduralVegetation.PVFloatRamp")
-+PropertyRedirects=(OldName="/Script/ProceduralVegetation.PVExportParams.bCreateNaniteAssemblies",NewName="/Script/ProceduralVegetation.PVExportParams.bCreateNaniteFoliage")
\ No newline at end of file
++PropertyRedirects=(OldName="/Script/ProceduralVegetation.PVExportParams.bCreateNaniteAssemblies",NewName="/Script/ProceduralVegetation.PVExportParams.bCreateNaniteFoliage")
++PropertyRedirects=(OldName="/Script/ProceduralVegetation.PVMeshBuilderParams.PointHullRetention",NewName="/Script/ProceduralVegetation.PVMeshBuilderParams.HullRetention")
++PropertyRedirects=(OldName="/Script/ProceduralVegetation.PVMeshBuilderParams.PointMainTrunkRetention",NewName="/Script/ProceduralVegetation.PVMeshBuilderParams.MainTrunkRetention")
++PropertyRedirects=(OldName="/Script/ProceduralVegetation.PVMeshBuilderParams.PointGroundRetention",NewName="/Script/ProceduralVegetation.PVMeshBuilderParams.GroundRetention")
++PropertyRedirects=(OldName="/Script/ProceduralVegetation.PVMeshBuilderParams.PointHullRetentionGradient",NewName="/Script/ProceduralVegetation.PVMeshBuilderParams.HullRetentionGradient")
++PropertyRedirects=(OldName="/Script/ProceduralVegetation.PVMeshBuilderParams.PointMainTrunkRetentionGradient",NewName="/Script/ProceduralVegetation.PVMeshBuilderParams.MainTrunkRetentionGradient")
++PropertyRedirects=(OldName="/Script/ProceduralVegetation.PVMeshBuilderParams.PointGroundRetentionGradient",NewName="/Script/ProceduralVegetation.PVMeshBuilderParams.GroundRetentionGradient")
++PropertyRedirects=(OldName="/Script/ProceduralVegetation.PVMeshBuilderParams.SegmentScaleRetention",NewName="/Script/ProceduralVegetation.PVMeshBuilderParams.ScaleRetention")
++PropertyRedirects=(OldName="/Script/ProceduralVegetation.PVMeshBuilderParams.SegmentScaleRetentionGradient",NewName="/Script/ProceduralVegetation.PVMeshBuilderParams.ScaleRetentionGradient")
\ No newline at end of file
diff --git a/Engine/Plugins/Experimental/ProceduralVegetationEditor/Source/ProceduralVegetation/Private/Implementations/PVGravity.cpp b/Engine/Plugins/Experimental/ProceduralVegetationEditor/Source/ProceduralVegetation/Private/Implementations/PVGravity.cpp
--- a/Engine/Plugins/Experimental/ProceduralVegetationEditor/Source/ProceduralVegetation/Private/Implementations/PVGravity.cpp
+++ b/Engine/Plugins/Experimental/ProceduralVegetationEditor/Source/ProceduralVegetation/Private/Implementations/PVGravity.cpp
@@ -106,7 +106,7 @@ void FPVGravity::ApplyGravity(
 				FVector2f(FMath::Max(0.1f, 1 - InGravityParams.AngleCorrection), 1),
 				YDot);
 
-			const float GravityImpact = InGravityParams.Gravity * (1 - InPScaleGradient) * YDot;
+			const float GravityImpact = (InGravityParams.Gravity / 10.0f) * (1 - InPScaleGradient) * YDot;
 			const FVector3f GravityForce = FMath::Lerp(InRelativeDirection, GravityDirection, GravityImpact).GetUnsafeNormal();
 			return FQuat4f::FindBetween(InRelativeDirection, GravityForce);
 		};
diff --git a/Engine/Plugins/Experimental/ProceduralVegetationEditor/Source/ProceduralVegetation/Private/Implementations/PVGravity.h b/Engine/Plugins/Experimental/ProceduralVegetationEditor/Source/ProceduralVegetation/Private/Implementations/PVGravity.h
--- a/Engine/Plugins/Experimental/ProceduralVegetationEditor/Source/ProceduralVegetation/Private/Implementations/PVGravity.h
+++ b/Engine/Plugins/Experimental/ProceduralVegetationEditor/Source/ProceduralVegetation/Private/Implementations/PVGravity.h
@@ -27,7 +27,7 @@ struct FPVGravityParams
 	UPROPERTY(EditAnywhere, Category="Gravity")
 	EGravityMode Mode = EGravityMode::Gravity;
 
-	UPROPERTY(EditAnywhere, Category="Gravity", meta=(ClampMin=0.0f, ClampMax=1.0f, UIMin=0.0f, UIMax=1.0f))
+	UPROPERTY(EditAnywhere, Category="Gravity", meta=(ClampMin=0.0f, ClampMax=10.0f, UIMin=0.0f, UIMax=1.0f))
 	float Gravity = 0;
 
 	UPROPERTY(EditAnywhere, Category="Gravity", meta=(EditCondition = "Mode == EGravityMode::Gravity"))
diff --git a/Engine/Plugins/Experimental/ProceduralVegetationEditor/Source/ProceduralVegetation/Private/Implementations/PVMeshBuilder.cpp b/Engine/Plugins/Experimental/ProceduralVegetationEditor/Source/ProceduralVegetation/Private/Implementations/PVMeshBuilder.cpp
--- a/Engine/Plugins/Experimental/ProceduralVegetationEditor/Source/ProceduralVegetation/Private/Implementations/PVMeshBuilder.cpp
+++ b/Engine/Plugins/Experimental/ProceduralVegetationEditor/Source/ProceduralVegetation/Private/Implementations/PVMeshBuilder.cpp
@@ -175,10 +175,10 @@ void FPVMeshBuilder::SetNjordPixelID(FGeometryCollection& OutGeometryCollection)
 	for (const auto& KeyValuePair : BudNumberPointIndicesMap)
 	{
 		int ValueIndex = 0;
-		for (auto Value : KeyValuePair.Value)
+		for (const auto Value : KeyValuePair.Value)
 		{
 			//Modify NjordPixelID
-			float NjordPixelID = KeyValuePair.Key + (float)ValueIndex / KeyValuePair.Value.Num();
+			float NjordPixelID = KeyValuePair.Key + static_cast<float>(ValueIndex) / KeyValuePair.Value.Num();
 			NjordPixelIds[Value] = NjordPixelID;
 			ValueIndex++;
 		}
@@ -302,51 +302,51 @@ TSet<int32> FPVMeshBuilder::ComputePointGradients(const PV::Facades::FPointFacad
 
 		// Compute scale removal gradient
 		const float PointHullGradient =
-			MeshBuilderParams.PointHullRetentionGradient.GetRichCurveConst()->Eval(1.0f - HullGradient) * MeshBuilderParams.PointHullRetention;
+			MeshBuilderParams.HullRetentionGradient.GetRichCurveConst()->Eval(1.0f - HullGradient) * MeshBuilderParams.HullRetention;
 
 		const float PointMainTrunkGradient =
-			MeshBuilderParams.PointMainTrunkRetentionGradient.GetRichCurveConst()->Eval(1.0f - MainTrunkGradient) * MeshBuilderParams.PointMainTrunkRetention;
+			MeshBuilderParams.MainTrunkRetentionGradient.GetRichCurveConst()->Eval(1.0f - MainTrunkGradient) * MeshBuilderParams.MainTrunkRetention;
 
 		const float PointGroundGradient =
-			MeshBuilderParams.PointGroundRetentionGradient.GetRichCurveConst()->Eval(1.0f - GroundGradient) * MeshBuilderParams.PointGroundRetention;
+			MeshBuilderParams.GroundRetentionGradient.GetRichCurveConst()->Eval(1.0f - GroundGradient) * MeshBuilderParams.GroundRetention;
 
 		const float ScaleRemovalGradient =
 			(ScaleGradient + PointHullGradient + PointMainTrunkGradient + PointGroundGradient) * MaxPointScale / 100.0f;
 
 		// Compute mesh divisions gradient
-		const float DivisionsScaleGradient = MeshBuilderParams.DivisionScaleRetentionGradient.GetRichCurveConst()->Eval(1.0f - ScaleGradient);
-		const float DivisionsHullGradient = MeshBuilderParams.DivisionHullRetentionGradient.GetRichCurveConst()->Eval(1.0f - HullGradient);
-		const float DivisionsMainTrunkGradient = MeshBuilderParams.DivisionMainTrunkRetentionGradient.GetRichCurveConst()->Eval(1.0f - MainTrunkGradient);
-		const float DivisionsGroundGradient = MeshBuilderParams.DivisionGroundRetentionGradient.GetRichCurveConst()->Eval(1.0f - GroundGradient);
+		const float DivisionsScaleGradient = MeshBuilderParams.ScaleRetentionGradient.GetRichCurveConst()->Eval(1.0f - ScaleGradient);
+		const float DivisionsHullGradient = MeshBuilderParams.HullRetentionGradient.GetRichCurveConst()->Eval(1.0f - HullGradient);
+		const float DivisionsMainTrunkGradient = MeshBuilderParams.MainTrunkRetentionGradient.GetRichCurveConst()->Eval(1.0f - MainTrunkGradient);
+		const float DivisionsGroundGradient = MeshBuilderParams.GroundRetentionGradient.GetRichCurveConst()->Eval(1.0f - GroundGradient);
 
 		const float InterpolatedHullGradient = FMath::GetMappedRangeValueClamped(
 			FVector2f(0.0f, 1.0f),
 			FVector2f(DivisionsHullGradient, 1.0f),
-			1.0f - MeshBuilderParams.DivisionHullRetention);
+			1.0f - MeshBuilderParams.HullRetention);
 		const float InterpolatedMainTrunkGradient = FMath::GetMappedRangeValueClamped(
 			FVector2f(0.0f, 1.0f),
 			FVector2f(DivisionsMainTrunkGradient, 1.0f),
-			1.0f - MeshBuilderParams.DivisionMainTrunkRetention);
+			1.0f - MeshBuilderParams.MainTrunkRetention);
 		const float InterpolatedGroundGradient = FMath::GetMappedRangeValueClamped(
 			FVector2f(0.0f, 1.0f),
 			FVector2f(DivisionsGroundGradient, 1.0f),
-			1.0f - MeshBuilderParams.DivisionGroundRetention);
+			1.0f - MeshBuilderParams.GroundRetention);
 
 		const float MeshDivisionsGradient = DivisionsScaleGradient * InterpolatedHullGradient * InterpolatedMainTrunkGradient *
 			InterpolatedGroundGradient;
 		OutMeshDivisionsGradients.Add(MeshDivisionsGradient);
 
 		// Compute delta modifier gradient
-		const float SegmentScaleGradient = MeshBuilderParams.SegmentScaleRetentionGradient.GetRichCurveConst()->Eval(1.0f - ScaleGradient);
-		const float SegmentHullGradient = MeshBuilderParams.SegmentHullRetentionGradient.GetRichCurveConst()->Eval(1.0f - HullGradient);
-		const float SegmentMainTrunkGradient = MeshBuilderParams.SegmentMainTrunkRetentionGradient.GetRichCurveConst()->Eval(1.0f - MainTrunkGradient);
-		const float SegmentGroundGradient = MeshBuilderParams.SegmentGroundRetentionGradient.GetRichCurveConst()->Eval(1.0f - GroundGradient);
+		const float SegmentScaleGradient = MeshBuilderParams.ScaleRetentionGradient.GetRichCurveConst()->Eval(1.0f - ScaleGradient);
+		const float SegmentHullGradient = MeshBuilderParams.HullRetentionGradient.GetRichCurveConst()->Eval(1.0f - HullGradient);
+		const float SegmentMainTrunkGradient = MeshBuilderParams.MainTrunkRetentionGradient.GetRichCurveConst()->Eval(1.0f - MainTrunkGradient);
+		const float SegmentGroundGradient = MeshBuilderParams.GroundRetentionGradient.GetRichCurveConst()->Eval(1.0f - GroundGradient);
 
 		const float DeltaModifierGradient =
-			(SegmentScaleGradient * MeshBuilderParams.SegmentScaleRetention) +
-			(SegmentHullGradient * MeshBuilderParams.SegmentHullRetention) +
-			(SegmentMainTrunkGradient * MeshBuilderParams.SegmentMainTrunkRetention) +
-			(SegmentGroundGradient * MeshBuilderParams.SegmentGroundRetention);
+			(SegmentScaleGradient * MeshBuilderParams.ScaleRetention) +
+			(SegmentHullGradient * MeshBuilderParams.HullRetention) +
+			(SegmentMainTrunkGradient * MeshBuilderParams.MainTrunkRetention) +
+			(SegmentGroundGradient * MeshBuilderParams.GroundRetention);
 
 		// Should the point be removed per scale gradient
 		if (!HardPoints.Contains(PointIndex) && MeshBuilderParams.PointRemoval > 0.0f && ScaleRemovalGradient < MeshBuilderParams.PointRemoval)
diff --git a/Engine/Plugins/Experimental/ProceduralVegetationEditor/Source/ProceduralVegetation/Private/Implementations/PVMeshBuilder.h b/Engine/Plugins/Experimental/ProceduralVegetationEditor/Source/ProceduralVegetation/Private/Implementations/PVMeshBuilder.h
--- a/Engine/Plugins/Experimental/ProceduralVegetationEditor/Source/ProceduralVegetation/Private/Implementations/PVMeshBuilder.h
+++ b/Engine/Plugins/Experimental/ProceduralVegetationEditor/Source/ProceduralVegetation/Private/Implementations/PVMeshBuilder.h
@@ -49,17 +49,13 @@ struct FPVMeshBuilderParams
 	float SegmentRetentionImpact = 1.0f;
 
 	UPROPERTY(EditAnywhere, Category="Mesh Settings", DisplayName="Longest Segment Allowed",
-		meta=(ClampMin=0.0f, ClampMax=1.0f, UIMin=0.0f, UIMax=1.0f))
+		meta=(ClampMin=0.0f, ClampMax=50.0f, UIMin=0.0f, UIMax=10.0f))
 	float LongestSegmentLength = 0.1f;
 
 	UPROPERTY(EditAnywhere, Category="Mesh Settings", DisplayName="Shortest Segment Allowed",
 		meta=(ClampMin=0.0f, ClampMax=1.0f, UIMin=0.0f, UIMax=1.0f))
 	float ShortestSegmentLength = 0.0f;
 
-	UPROPERTY(EditAnywhere, Category="Mesh Settings", DisplayName="Reduction Accuracy",
-		meta=(ClampMin=3, ClampMax=15, UIMin=1, UIMax=15))
-	int32 Accuracy = 5;
-
 	UPROPERTY(EditAnywhere, Category="Mesh Settings", DisplayName="Min Divisions",
 		meta=(ClampMin=3, ClampMax=1024, UIMin=3, UIMax=36))
 	int32 MinMeshDivisions = 6;
@@ -68,78 +64,37 @@ struct FPVMeshBuilderParams
 		meta=(ClampMin=3, ClampMax=1024, UIMin=3, UIMax=36))
 	int32 MaxMeshDivisions = 12;
 
-	UPROPERTY(EditAnywhere, Category="Point Removal Settings", DisplayName="Hull Retention",
-		meta=(ClampMin=0.0f, ClampMax=1.0f, UIMin=0.0f, UIMax=1.0f))
-	float PointHullRetention = 0.0f;
-
-	UPROPERTY(EditAnywhere, Category="Point Removal Settings", DisplayName="Main Trunk Retention",
-		meta=(ClampMin=0.0f, ClampMax=1.0f, UIMin=0.0f, UIMax=1.0f))
-	float PointMainTrunkRetention = 0.0f;
-
-	UPROPERTY(EditAnywhere, Category="Point Removal Settings", DisplayName="Ground Retention",
-		meta=(ClampMin=0.0f, ClampMax=1.0f, UIMin=0.0f, UIMax=1.0f))
-	float PointGroundRetention = 0.0f;
-
-	UPROPERTY(EditAnywhere, Category="Point Removal Settings", DisplayName="Hull Retention Gradient", meta=(XAxisMin=0.0f, XAxisMax=1.0f, YAxisMin=0.0f, YAxisMax=1.0f))
-	FPVFloatRamp PointHullRetentionGradient;
-
-	UPROPERTY(EditAnywhere, Category="Point Removal Settings", DisplayName="Main Trunk Retention Gradient", meta=(XAxisMin=0.0f, XAxisMax=1.0f, YAxisMin=0.0f, YAxisMax=1.0f))
-	FPVFloatRamp PointMainTrunkRetentionGradient;
-
-	UPROPERTY(EditAnywhere, Category="Point Removal Settings", DisplayName="Ground Retention Gradient", meta=(XAxisMin=0.0f, XAxisMax=1.0f, YAxisMin=0.0f, YAxisMax=1.0f))
-	FPVFloatRamp PointGroundRetentionGradient;
-
-	UPROPERTY(EditAnywhere, Category="Segment Removal Settings", DisplayName="Hull Retention",
-		meta=(ClampMin=0.0f, ClampMax=1.0f, UIMin=0.0f, UIMax=1.0f))
-	float SegmentHullRetention = 0.0f;
+	UPROPERTY(EditAnywhere, Category="Mesh Settings|Advanced", DisplayName="Reduction Accuracy",
+		meta=(ClampMin=3, ClampMax=15, UIMin=1, UIMax=15))
+	int32 Accuracy = 5;
 
-	UPROPERTY(EditAnywhere, Category="Segment Removal Settings", DisplayName="Main Trunk Retention",
+	UPROPERTY(EditAnywhere, Category="Mesh Settings|Advanced", DisplayName="Hull Retention",
 		meta=(ClampMin=0.0f, ClampMax=1.0f, UIMin=0.0f, UIMax=1.0f))
-	float SegmentMainTrunkRetention = 0.0f;
+	float HullRetention = 0.0f;
 
-	UPROPERTY(EditAnywhere, Category="Segment Removal Settings", DisplayName="Ground Retention",
+	UPROPERTY(EditAnywhere, Category="Mesh Settings|Advanced", DisplayName="Main Trunk Retention",
 		meta=(ClampMin=0.0f, ClampMax=1.0f, UIMin=0.0f, UIMax=1.0f))
-	float SegmentGroundRetention = 0.0f;
+	float MainTrunkRetention = 0.0f;
 
-	UPROPERTY(EditAnywhere, Category="Segment Removal Settings", DisplayName="Scale Retention",
+	UPROPERTY(EditAnywhere, Category="Mesh Settings|Advanced", DisplayName="Ground Retention",
 		meta=(ClampMin=0.0f, ClampMax=1.0f, UIMin=0.0f, UIMax=1.0f))
-	float SegmentScaleRetention = 0.0f;
-
-	UPROPERTY(EditAnywhere, Category="Segment Removal Settings", DisplayName="Hull Retention Gradient", meta=(XAxisMin=0.0f, XAxisMax=1.0f, YAxisMin=0.0f, YAxisMax=1.0f))
-	FPVFloatRamp SegmentHullRetentionGradient;
-
-	UPROPERTY(EditAnywhere, Category="Segment Removal Settings", DisplayName="Main Trunk Retention Gradient", meta=(XAxisMin=0.0f, XAxisMax=1.0f, YAxisMin=0.0f, YAxisMax=1.0f))
-	FPVFloatRamp SegmentMainTrunkRetentionGradient;
+	float GroundRetention = 0.0f;
 
-	UPROPERTY(EditAnywhere, Category="Segment Removal Settings", DisplayName="Ground Retention Gradient", meta=(XAxisMin=0.0f, XAxisMax=1.0f, YAxisMin=0.0f, YAxisMax=1.0f))
-	FPVFloatRamp SegmentGroundRetentionGradient;
+	UPROPERTY(EditAnywhere, Category="Mesh Settings|Advanced", DisplayName="Hull Retention Gradient", meta=(XAxisMin=0.0f, XAxisMax=1.0f, YAxisMin=0.0f, YAxisMax=1.0f))
+	FPVFloatRamp HullRetentionGradient;
 
-	UPROPERTY(EditAnywhere, Category="Segment Removal Settings", DisplayName="Scale Retention Gradient", meta=(XAxisMin=0.0f, XAxisMax=1.0f, YAxisMin=0.0f, YAxisMax=1.0f))
-	FPVFloatRamp SegmentScaleRetentionGradient;
+	UPROPERTY(EditAnywhere, Category="Mesh Settings|Advanced", DisplayName="Main Trunk Retention Gradient", meta=(XAxisMin=0.0f, XAxisMax=1.0f, YAxisMin=0.0f, YAxisMax=1.0f))
+	FPVFloatRamp MainTrunkRetentionGradient;
 
-	UPROPERTY(EditAnywhere, Category="Division Removal Settings", DisplayName="Hull Retention",
-		meta=(ClampMin=0.0f, ClampMax=1.0f, UIMin=0.0f, UIMax=1.0f))
-	float DivisionHullRetention = 0.0f;
-
-	UPROPERTY(EditAnywhere, Category="Division Removal Settings", DisplayName="Main Trunk Retention",
-		meta=(ClampMin=0.0f, ClampMax=1.0f, UIMin=0.0f, UIMax=1.0f))
-	float DivisionMainTrunkRetention = 0.0f;
+	UPROPERTY(EditAnywhere, Category="Mesh Settings|Advanced", DisplayName="Ground Retention Gradient", meta=(XAxisMin=0.0f, XAxisMax=1.0f, YAxisMin=0.0f, YAxisMax=1.0f))
+	FPVFloatRamp GroundRetentionGradient;
 
-	UPROPERTY(EditAnywhere, Category="Division Removal Settings", DisplayName="Ground Retention",
+	UPROPERTY(EditAnywhere, Category="Mesh Settings|Advanced", DisplayName="Scale Retention",
 		meta=(ClampMin=0.0f, ClampMax=1.0f, UIMin=0.0f, UIMax=1.0f))
-	float DivisionGroundRetention = 0.0f;
-
-	UPROPERTY(EditAnywhere, Category="Division Removal Settings", DisplayName="Hull Retention Gradient", meta=(XAxisMin=0.0f, XAxisMax=1.0f, YAxisMin=0.0f, YAxisMax=1.0f))
-	FPVFloatRamp DivisionHullRetentionGradient;
-
-	UPROPERTY(EditAnywhere, Category="Division Removal Settings", DisplayName="Main Trunk Retention Gradient", meta=(XAxisMin=0.0f, XAxisMax=1.0f, YAxisMin=0.0f, YAxisMax=1.0f))
-	FPVFloatRamp DivisionMainTrunkRetentionGradient;
-
-	UPROPERTY(EditAnywhere, Category="Division Removal Settings", DisplayName="Ground Retention Gradient", meta=(XAxisMin=0.0f, XAxisMax=1.0f, YAxisMin=0.0f, YAxisMax=1.0f))
-	FPVFloatRamp DivisionGroundRetentionGradient;
+	float ScaleRetention = 0.0f;
 
-	UPROPERTY(EditAnywhere, Category="Division Removal Settings", DisplayName="Scale Retention Gradient", meta=(XAxisMin=0.0f, XAxisMax=1.0f, YAxisMin=0.0f, YAxisMax=1.0f))
-	FPVFloatRamp DivisionScaleRetentionGradient;
+	UPROPERTY(EditAnywhere, Category="Mesh Settings|Advanced", DisplayName="Scale Retention Gradient", meta=(XAxisMin=0.0f, XAxisMax=1.0f, YAxisMin=0.0f, YAxisMax=1.0f))
+	FPVFloatRamp ScaleRetentionGradient;
 
 	UPROPERTY(EditAnywhere, Category="Plant Profile Settings", DisplayName="Profile",
 		meta=(EditCondition="bIsPlantProfileDropdownEnabled", EditConditionHides, HideEditConditionToggle, GetOptions="GetPlantProfileOptions"))
@@ -187,17 +142,10 @@ struct FPVMeshBuilderParams
 
 	FPVMeshBuilderParams()
 	{
-		InitializeLinearCurve(PointHullRetentionGradient);
-		InitializeLinearCurve(PointMainTrunkRetentionGradient);
-		InitializeLinearCurve(PointGroundRetentionGradient);
-		InitializeLinearCurve(SegmentHullRetentionGradient);
-		InitializeLinearCurve(SegmentMainTrunkRetentionGradient);
-		InitializeLinearCurve(SegmentGroundRetentionGradient);
-		InitializeLinearCurve(SegmentScaleRetentionGradient);
-		InitializeLinearCurve(DivisionHullRetentionGradient);
-		InitializeLinearCurve(DivisionMainTrunkRetentionGradient);
-		InitializeLinearCurve(DivisionGroundRetentionGradient);
-		InitializeLinearCurve(DivisionScaleRetentionGradient);
+		InitializeLinearCurve(HullRetentionGradient);
+		InitializeLinearCurve(MainTrunkRetentionGradient);
+		InitializeLinearCurve(GroundRetentionGradient);
+		InitializeLinearCurve(ScaleRetentionGradient);
 		InitializeLinearCurve(PlantProfileFallOff);
 	}
 
diff --git a/Engine/Plugins/Experimental/ProceduralVegetationEditor/Source/ProceduralVegetation/Private/Nodes/PVFoliageDistributorSettings.h b/Engine/Plugins/Experimental/ProceduralVegetationEditor/Source/ProceduralVegetation/Private/Nodes/PVFoliageDistributorSettings.h
--- a/Engine/Plugins/Experimental/ProceduralVegetationEditor/Source/ProceduralVegetation/Private/Nodes/PVFoliageDistributorSettings.h
+++ b/Engine/Plugins/Experimental/ProceduralVegetationEditor/Source/ProceduralVegetation/Private/Nodes/PVFoliageDistributorSettings.h
@@ -46,79 +46,79 @@ private:
 	UPROPERTY(EditAnywhere, Category="Distribution Settings")
 	bool OverrideDistribution = true;
 
-	UPROPERTY(EditAnywhere, Category="Distribution Settings", meta=(ClampMin=0.0f, ClampMax=1.0f, UIMin=0.0f, UIMax=1.0f))
+	UPROPERTY(EditAnywhere, Category="Distribution Settings", meta=(EditCondition="OverrideDistribution", ClampMin=0.0f, ClampMax=1.0f, UIMin=0.0f, UIMax=1.0f))
 	float InstanceSpacing = 0.3f;
 
-	UPROPERTY(EditAnywhere, Category="Distribution Settings", meta=(XAxisMin=0.0f, XAxisMax=1.0f, YAxisMin=0.0f, YAxisMax=1.0f))
+	UPROPERTY(EditAnywhere, Category="Distribution Settings", meta=(EditCondition="OverrideDistribution", XAxisMin=0.0f, XAxisMax=1.0f, YAxisMin=0.0f, YAxisMax=1.0f))
 	FPVFloatRamp InstanceSpacingRamp;
 
-	UPROPERTY(EditAnywhere, Category="Distribution Settings", meta=(ClampMin=0.0f, ClampMax=1.0f, UIMin=0.0f, UIMax=1.0f))
+	UPROPERTY(EditAnywhere, Category="Distribution Settings", meta=(EditCondition="OverrideDistribution", ClampMin=0.0f, ClampMax=1.0f, UIMin=0.0f, UIMax=1.0f))
 	float InstanceSpacingRampEffect = 0.0f;
 
-	UPROPERTY(EditAnywhere, Category="Distribution Settings", meta=(ClampMin=-1))
+	UPROPERTY(EditAnywhere, Category="Distribution Settings", meta=(EditCondition="OverrideDistribution", ClampMin=-1))
 	int32 MaxPerBranch = -1;
 
-	UPROPERTY(EditAnywhere, Category="Scale Settings", meta=(ClampMin=0.0f, ClampMax=2.0f, UIMin=0.0f, UIMax=2.0f))
+	UPROPERTY(EditAnywhere, Category="Scale Settings", meta=(EditCondition="OverrideDistribution", ClampMin=0.0f, ClampMax=2.0f, UIMin=0.0f, UIMax=2.0f))
 	float BaseScale = 1.0f;
 
-	UPROPERTY(EditAnywhere, Category="Scale Settings", meta=(ClampMin=0.0f, ClampMax=1.0f, UIMin=0.0f, UIMax=1.0f))
+	UPROPERTY(EditAnywhere, Category="Scale Settings", meta=(EditCondition="OverrideDistribution", ClampMin=0.0f, ClampMax=1.0f, UIMin=0.0f, UIMax=1.0f))
 	float BranchScaleImpact = 0.0f;
 
-	UPROPERTY(EditAnywhere, Category="Scale Settings", meta=(ClampMin=0.0f, ClampMax=1.0f, UIMin=0.0f, UIMax=1.0f))
+	UPROPERTY(EditAnywhere, Category="Scale Settings", meta=(EditCondition="OverrideDistribution", ClampMin=0.0f, ClampMax=1.0f, UIMin=0.0f, UIMax=1.0f))
 	float MinScale = 0.5f;
 
-	UPROPERTY(EditAnywhere, Category="Scale Settings", meta=(ClampMin=0.0f, ClampMax=1.0f, UIMin=0.0f, UIMax=1.0f))
+	UPROPERTY(EditAnywhere, Category="Scale Settings", meta=(EditCondition="OverrideDistribution", ClampMin=0.0f, ClampMax=1.0f, UIMin=0.0f, UIMax=1.0f))
 	float MaxScale = 1.0f;
 
-	UPROPERTY(EditAnywhere, Category="Scale Settings", meta=(ClampMin=0.0f, ClampMax=1.0f, UIMin=0.0f, UIMax=1.0f))
+	UPROPERTY(EditAnywhere, Category="Scale Settings", meta=(EditCondition="OverrideDistribution", ClampMin=0.0f, ClampMax=1.0f, UIMin=0.0f, UIMax=1.0f))
 	float RandomScaleMin = 1.0f;
 
-	UPROPERTY(EditAnywhere, Category="Scale Settings", meta=(ClampMin=1.0f, ClampMax=2.0f, UIMin=1.0f, UIMax=2.0f))
+	UPROPERTY(EditAnywhere, Category="Scale Settings", meta=(EditCondition="OverrideDistribution", ClampMin=1.0f, ClampMax=2.0f, UIMin=1.0f, UIMax=2.0f))
 	float RandomScaleMax = 1.0f;
 
-	UPROPERTY(EditAnywhere, Category="Scale Settings", meta=(XAxisMin=0.0f, XAxisMax=1.0f, YAxisMin=0.0f, YAxisMax=1.0f))
+	UPROPERTY(EditAnywhere, Category="Scale Settings", meta=(EditCondition="OverrideDistribution", XAxisMin=0.0f, XAxisMax=1.0f, YAxisMin=0.0f, YAxisMax=1.0f))
 	FPVFloatRamp ScaleRamp;
 
-	UPROPERTY(EditAnywhere, Category="Vector Settings")
+	UPROPERTY(EditAnywhere, Category="Vector Settings", meta=(EditCondition="OverrideDistribution"))
 	bool OverrideAxilAngle = true;
 
-	UPROPERTY(EditAnywhere, Category="Vector Settings", meta=(ClampMin=0.0f, ClampMax=90.0f, UIMin=0.0f, UIMax=90.0f))
+	UPROPERTY(EditAnywhere, Category="Vector Settings", meta=(EditCondition="OverrideDistribution", ClampMin=0.0f, ClampMax=90.0f, UIMin=0.0f, UIMax=90.0f))
 	float AxilAngle = 45.0f;
 
-	UPROPERTY(EditAnywhere, Category="Vector Settings", meta=(XAxisMin=0.0f, XAxisMax=1.0f, YAxisMin=0.0f, YAxisMax=1.0f))
+	UPROPERTY(EditAnywhere, Category="Vector Settings", meta=(EditCondition="OverrideDistribution", XAxisMin=0.0f, XAxisMax=1.0f, YAxisMin=0.0f, YAxisMax=1.0f))
 	FPVFloatRamp AxilAngleRamp;
 
-	UPROPERTY(EditAnywhere, Category="Vector Settings", meta=(ClampMin=0.0f, ClampMax=90.0f, UIMin=0.0f, UIMax=90.0f))
+	UPROPERTY(EditAnywhere, Category="Vector Settings", meta=(EditCondition="OverrideDistribution", ClampMin=0.0f, ClampMax=90.0f, UIMin=0.0f, UIMax=90.0f))
 	float AxilAngleRampUpperValue = 45.0f;
 
-	UPROPERTY(EditAnywhere, Category="Vector Settings", meta=(ClampMin=0.0f, ClampMax=1.0f, UIMin=0.0f, UIMax=1.0f))
+	UPROPERTY(EditAnywhere, Category="Vector Settings", meta=(EditCondition="OverrideDistribution", ClampMin=0.0f, ClampMax=1.0f, UIMin=0.0f, UIMax=1.0f))
 	float AxilAngleRampEffect = 0.0f;
 
-	UPROPERTY(EditAnywhere, Category="Phyllotaxy Settings")
+	UPROPERTY(EditAnywhere, Category="Phyllotaxy Settings", meta=(EditCondition="OverrideDistribution"))
 	bool OverridePhyllotaxy = false;
 
-	UPROPERTY(EditAnywhere, Category="Phyllotaxy Settings")
+	UPROPERTY(EditAnywhere, Category="Phyllotaxy Settings", meta=(EditCondition="OverrideDistribution"))
 	EPhyllotaxyType PhyllotaxyType = EPhyllotaxyType::Alternate;
 
-	UPROPERTY(EditAnywhere, Category="Phyllotaxy Settings", meta=(EditCondition="PhyllotaxyType == EPhyllotaxyType::Spiral", EditConditionHides))
+	UPROPERTY(EditAnywhere, Category="Phyllotaxy Settings", meta=(EditCondition="OverrideDistribution && PhyllotaxyType == EPhyllotaxyType::Spiral", EditConditionHides))
 	EPhyllotaxyFormation PhyllotaxyFormation = EPhyllotaxyFormation::Distichous;
 
-	UPROPERTY(EditAnywhere, Category="Phyllotaxy Settings", meta=(ClampMin=1, ClampMax=10, UIMin=1, UIMax=10))
+	UPROPERTY(EditAnywhere, Category="Phyllotaxy Settings", meta=(EditCondition="OverrideDistribution", ClampMin=1, ClampMax=10, UIMin=1, UIMax=10))
 	int32 MinimumNodeBuds = 2;
 
-	UPROPERTY(EditAnywhere, Category="Phyllotaxy Settings", meta=(ClampMin=1, ClampMax=10, UIMin=1, UIMax=10))
+	UPROPERTY(EditAnywhere, Category="Phyllotaxy Settings", meta=(EditCondition="OverrideDistribution", ClampMin=1, ClampMax=10, UIMin=1, UIMax=10))
 	int32 MaximumNodeBuds = 3;
 
-	UPROPERTY(EditAnywhere, Category="Phyllotaxy Settings")
+	UPROPERTY(EditAnywhere, Category="Phyllotaxy Settings", meta=(EditCondition="OverrideDistribution"))
 	bool ResetPhyllotaxy = false;
 
-	UPROPERTY(EditAnywhere, Category="Phyllotaxy Settings", meta=(ClampMin=0.0f, ClampMax=360.0f, UIMin=0.0f, UIMax=360.0f))
+	UPROPERTY(EditAnywhere, Category="Phyllotaxy Settings", meta=(EditCondition="OverrideDistribution", ClampMin=0.0f, ClampMax=360.0f, UIMin=0.0f, UIMax=360.0f))
 	float PhyllotaxyAdditionalAngle = 0.0f;
 
-	UPROPERTY(EditAnywhere, Category="Phyllotaxy Settings", meta=(ClampMin=0.0f, ClampMax=360.0f, UIMin=0.0f, UIMax=360.0f))
+	UPROPERTY(EditAnywhere, Category="Phyllotaxy Settings", meta=(EditCondition="OverrideDistribution", ClampMin=0.0f, ClampMax=360.0f, UIMin=0.0f, UIMax=360.0f))
 	float PhyllotaxyOffset = 0.0f;
 
-	UPROPERTY(EditAnywhere, Category="Misc Settings")
+	UPROPERTY(EditAnywhere, Category="Misc Settings", meta=(EditCondition="OverrideDistribution"))
 	int32 RandomSeed = 123456;
 };
 
