[PVE] Selection changes:
- Updated the viewport on node selection
- Select the first available output when the editor is opened.
- Removed extra panels.
- Auto focus viewport first time

Fixed issue: Selecting and trying to rename a comment box crashes PVE

Misc: Disabled viewport mannequin

diff --git a/Engine/Plugins/Experimental/ProceduralVegetationEditor/Source/ProceduralVegetation/Public/Nodes/PVBaseSettings.h b/Engine/Plugins/Experimental/ProceduralVegetationEditor/Source/ProceduralVegetation/Public/Nodes/PVBaseSettings.h
--- a/Engine/Plugins/Experimental/ProceduralVegetationEditor/Source/ProceduralVegetation/Public/Nodes/PVBaseSettings.h
+++ b/Engine/Plugins/Experimental/ProceduralVegetationEditor/Source/ProceduralVegetation/Public/Nodes/PVBaseSettings.h
@@ -13,6 +13,8 @@ class UPVBaseSettings : public UPCGSettings
 	GENERATED_BODY()
 
 public:
+	bool bAutoFocusViewport = false;
+	
 #if WITH_EDITOR
 	virtual EPCGSettingsType GetType() const override { return EPCGSettingsType::Spatial; }
 	virtual TArray<EPVRenderType> GetRenderType() const { return TArray<EPVRenderType>(); }
diff --git a/Engine/Plugins/Experimental/ProceduralVegetationEditor/Source/ProceduralVegetationEditor/Private/PVEditor.cpp b/Engine/Plugins/Experimental/ProceduralVegetationEditor/Source/ProceduralVegetationEditor/Private/PVEditor.cpp
--- a/Engine/Plugins/Experimental/ProceduralVegetationEditor/Source/ProceduralVegetationEditor/Private/PVEditor.cpp
+++ b/Engine/Plugins/Experimental/ProceduralVegetationEditor/Source/ProceduralVegetationEditor/Private/PVEditor.cpp
@@ -4,6 +4,7 @@
 #include "EngineAnalytics.h"
 #include "PVEditorCommands.h"
 #include "PCGDefaultExecutionSource.h"
+#include "PCGEditorSettings.h"
 #include "Subsystems/PCGEngineSubsystem.h"
 #include "PVEditorSchema.h"
 #include "ProceduralVegetation.h"
@@ -13,6 +14,7 @@
 #include "ToolMenuSection.h"
 #include "Dataflow/DataflowCollectionSpreadSheetWidget.h"
 #include "Nodes/PCGEditorGraphNodeBase.h"
+#include "PCGEditor/Private/PCGEditorGraph.h"
 #include "Widgets/SPVCreateDialog.h"
 #include "Widgets/SPVExportSelectionDialog.h"
 
@@ -72,6 +74,23 @@ void FPVEditor::Initialize(const EToolkitMode::Type InMode, const TSharedPtr<cla
 	// Ask for generation
 	ExecutionSource->Generate();
 
+	// Select the first available output
+	for (auto EdNode : GetPCGEditorGraph()->Nodes)
+	{
+		if (UPCGEditorGraphNodeBase* PCGEdNode = Cast<UPCGEditorGraphNodeBase>(EdNode))
+		{
+			if (auto PCGNode = PCGEdNode->GetPCGNode())
+			{
+				if(UPVOutputSettings* Settings = Cast<UPVOutputSettings>(PCGNode->GetSettings()))
+				{
+					GetPCGEditorGraph()->SelectNodeSet({PCGEdNode}, true);
+					
+					break;
+				}
+			}
+		}
+	}
+
 	SessionStartTime = FDateTime::Now();
 
 	if (FEngineAnalytics::IsAvailable())
@@ -230,6 +249,20 @@ TSharedRef<FTabManager::FLayout> FPVEditor::GetDefaultLayout() const
 	);
 }
 
+bool FPVEditor::IsPanelAvailable(const FName PanelID) const
+{
+	static TSet<FName> SupportedIDs = {
+		GetPanelID(EPCGEditorPanel::Find),
+		GetPanelID(EPCGEditorPanel::GraphEditor),
+		GetPanelID(EPCGEditorPanel::Log),
+		GetPanelID(EPCGEditorPanel::PropertyDetails1),
+		GetPanelID(EPCGEditorPanel::UserParams),
+		GetPanelID(EPCGEditorPanel::Viewport1),
+		GetPanelID(EPCGEditorPanel::Attributes1)
+	};
+	return SupportedIDs.Contains(PanelID);
+}
+
 void FPVEditor::RegisterToolbarInternal(FToolMenuSection& PCGSection) const
 {
 	RegisterToolbarButton(PCGSection, EPCGToolbarButtons::Find);
@@ -258,17 +291,29 @@ void FPVEditor::OnSelectedNodesChanged(const TSet<UObject*>& InNewSelection)
 {
 	FPCGEditor::OnSelectedNodesChanged(InNewSelection);
 
-#if PROCEDURAL_VEGETATION_DEBUG_SHOW_COLLECTION_TAB
 	const UPCGNode* PCGNode = nullptr;
 	
 	if (!InNewSelection.IsEmpty())
 	{
-		const UPCGEditorGraphNodeBase* Node = Cast<UPCGEditorGraphNodeBase>(*InNewSelection.CreateConstIterator());
-		PCGNode = Node->GetPCGNode();
-	}
+		UPCGEditorGraphNodeBase* Node = Cast<UPCGEditorGraphNodeBase>(*InNewSelection.CreateConstIterator());
 
-	UpdateCollectionSpreadSheet(PCGNode);
-#endif
+		if (Node && Node->GetPCGNode())
+		{
+			PCGNode = Node->GetPCGNode();
+
+#if PROCEDURAL_VEGETATION_DEBUG_SHOW_COLLECTION_TAB
+			UpdateCollectionSpreadSheet(PCGNode);
+#endif			
+
+			if (UPVBaseSettings* Settings = Cast<UPVBaseSettings>(PCGNode->GetSettings()))
+			{
+				Settings->bAutoFocusViewport = bFocusViewportOnSelection;
+				bFocusViewportOnSelection = false;
+			}
+			
+			SetNodeInspected(Node, true);
+		}
+	}
 }
 
 void FPVEditor::OnExport_Clicked()
diff --git a/Engine/Plugins/Experimental/ProceduralVegetationEditor/Source/ProceduralVegetationEditor/Private/PVEditor.h b/Engine/Plugins/Experimental/ProceduralVegetationEditor/Source/ProceduralVegetationEditor/Private/PVEditor.h
--- a/Engine/Plugins/Experimental/ProceduralVegetationEditor/Source/ProceduralVegetationEditor/Private/PVEditor.h
+++ b/Engine/Plugins/Experimental/ProceduralVegetationEditor/Source/ProceduralVegetationEditor/Private/PVEditor.h
@@ -39,6 +39,7 @@ protected:
 	// ~Begin FPCGEditor interface
 	virtual TAttribute<FGraphAppearanceInfo> GetAppearanceInfo() const override;
 	virtual TSharedRef<FTabManager::FLayout> GetDefaultLayout() const override;
+	virtual bool IsPanelAvailable(const FName PanelID) const override;
 	virtual void RegisterToolbarInternal(FToolMenuSection& PCGSection) const override;
 	virtual void BindCommands() override;
 	virtual void OnSelectedNodesChanged(const TSet<UObject*>& InNewSelection) override;
@@ -52,6 +53,8 @@ private:
 
 	FDateTime SessionStartTime;
 
+	bool bFocusViewportOnSelection = true;
+
 #if PROCEDURAL_VEGETATION_DEBUG_SHOW_COLLECTION_TAB
 	static const FName CollectionSpreadSheetTabId;
 
diff --git a/Engine/Plugins/Experimental/ProceduralVegetationEditor/Source/ProceduralVegetationEditor/Private/Visualizations/PVDataVisualization.cpp b/Engine/Plugins/Experimental/ProceduralVegetationEditor/Source/ProceduralVegetationEditor/Private/Visualizations/PVDataVisualization.cpp
--- a/Engine/Plugins/Experimental/ProceduralVegetationEditor/Source/ProceduralVegetationEditor/Private/Visualizations/PVDataVisualization.cpp
+++ b/Engine/Plugins/Experimental/ProceduralVegetationEditor/Source/ProceduralVegetationEditor/Private/Visualizations/PVDataVisualization.cpp
@@ -19,6 +19,7 @@
 #include "DataTypes/PVData.h"
 #include "Components/TextRenderComponent.h"
 #include "DebugVisualization/PVDebugVisualizer.h"
+#include "Nodes/PVBaseSettings.h"
 #include "Visualizations/PVSkeletonVisualizerComponent.h"
 
 #define LOCTEXT_NAMESPACE "PCGManagedArrayCollectionDataVisualization"
@@ -29,9 +30,6 @@ DEFINE_LOG_CATEGORY(LogProceduralVegetationDataVisualization);
 
 FPVDataVisualization::FPVDataVisualization()
 {
-	UPCGEditorSettings* PCGEditorSettings = GetMutableDefault<UPCGEditorSettings>();
-	PCGEditorSettings->bAutoFocusViewport = false;
-
 	RegisterVisualizations();
 }
 
@@ -51,7 +49,7 @@ void FPVDataVisualization::ExecuteDebugDisplay(FPCGContext* Context, const UPCGS
 
 FPCGSetupSceneFunc FPVDataVisualization::GetViewportSetupFunc(const UPCGSettingsInterface* SettingsInterface, const UPCGData* Data) const
 {
-	return [this, WeakData=TWeakObjectPtr<const UPVData>(Cast<UPVData>(Data)), WeakSettings = TWeakObjectPtr<const UPCGSettingsInterface>(Cast<UPCGSettingsInterface>(SettingsInterface))](
+	return [this, WeakData=TWeakObjectPtr<const UPVData>(Cast<UPVData>(Data)), WeakSettings = TWeakObjectPtr<const UPVBaseSettings>(Cast<UPVBaseSettings>(SettingsInterface))](
 		FPCGSceneSetupParams& InOutParams)
 	{
 		check(InOutParams.Scene);
@@ -61,7 +59,14 @@ FPCGSetupSceneFunc FPVDataVisualization::GetViewportSetupFunc(const UPCGSettings
 			UE_LOG(LogProceduralVegetationDataVisualization, Error, TEXT("Failed to setup data viewport, the data was lost or invalid."));
 			return;
 		}
+			
+		UPCGEditorSettings* PCGEditorSettings = GetMutableDefault<UPCGEditorSettings>();
 
+		if (PCGEditorSettings && WeakData.IsValid())
+		{
+			PCGEditorSettings->bAutoFocusViewport = WeakSettings->bAutoFocusViewport;
+		}
+			
 		const FManagedArrayCollection& Collection = WeakData->GetCollection();
 
 		FBoxSphereBounds Bounds(ForceInit);
@@ -93,7 +98,7 @@ FPCGSetupSceneFunc FPVDataVisualization::GetViewportSetupFunc(const UPCGSettings
 			TEXT("/ProceduralVegetationEditor/Mannequin/Viewport_Mannequin_T_Pose.Viewport_Mannequin_T_Pose")
 		);
 
-		constexpr bool bShowMannequin = true;
+		constexpr bool bShowMannequin = false;
 		if (bShowMannequin)
 		{
 			UStaticMeshComponent* MannequinComponent = NewObject<UStaticMeshComponent>(GetTransientPackage(), NAME_None, RF_Transient);
