[PVE] Fix bughawk issues

diff --git a/Engine/Plugins/Experimental/ProceduralVegetationEditor/Source/ProceduralVegetationEditor/Private/Visualizations/PVDataVisualization.cpp b/Engine/Plugins/Experimental/ProceduralVegetationEditor/Source/ProceduralVegetationEditor/Private/Visualizations/PVDataVisualization.cpp
--- a/Engine/Plugins/Experimental/ProceduralVegetationEditor/Source/ProceduralVegetationEditor/Private/Visualizations/PVDataVisualization.cpp
+++ b/Engine/Plugins/Experimental/ProceduralVegetationEditor/Source/ProceduralVegetationEditor/Private/Visualizations/PVDataVisualization.cpp
@@ -62,7 +62,7 @@ FPCGSetupSceneFunc FPVDataVisualization::GetViewportSetupFunc(const UPCGSettings
 			
 		UPCGEditorSettings* PCGEditorSettings = GetMutableDefault<UPCGEditorSettings>();
 
-		if (PCGEditorSettings && WeakData.IsValid())
+		if (PCGEditorSettings && WeakSettings.IsValid())
 		{
 			PCGEditorSettings->bAutoFocusViewport = WeakSettings->bAutoFocusViewport;
 		}
