[PVE] Build Health issue fix
Output to run on main thread

diff --git a/Engine/Plugins/Experimental/ProceduralVegetationEditor/Source/ProceduralVegetation/Private/Helpers/PVExportHelper.cpp b/Engine/Plugins/Experimental/ProceduralVegetationEditor/Source/ProceduralVegetation/Private/Helpers/PVExportHelper.cpp
--- a/Engine/Plugins/Experimental/ProceduralVegetationEditor/Source/ProceduralVegetation/Private/Helpers/PVExportHelper.cpp
+++ b/Engine/Plugins/Experimental/ProceduralVegetationEditor/Source/ProceduralVegetation/Private/Helpers/PVExportHelper.cpp
@@ -35,6 +35,7 @@
 #include "SkinnedAssetCompiler.h"
 #include "StaticMeshCompiler.h"
 #include "GeometryCollection/Facades/CollectionConstraintOverrideFacade.h"
+#include "AssetRegistry/AssetData.h"
 
 namespace PV::Export::Internal
 {
@@ -1008,7 +1009,7 @@ namespace PV::Export
 			}
 
 			TObjectPtr<USkeleton> ExportSkeleton = nullptr;
-			if (Utilities::CanOverwriteAsset(*(FullPath + "_Skeleton"), USkeletalMesh::StaticClass(), AssetData, ErrorMessage))
+			if (Utilities::CanOverwriteAsset(*(FullPath + "_Skeleton"), USkeleton::StaticClass(), AssetData, ErrorMessage))
 			{
 				ExportSkeleton = CastChecked<USkeleton>(AssetData.GetAsset());
 			}
diff --git a/Engine/Plugins/Experimental/ProceduralVegetationEditor/Source/ProceduralVegetation/Public/Helpers/PVUtilities.h b/Engine/Plugins/Experimental/ProceduralVegetationEditor/Source/ProceduralVegetation/Public/Helpers/PVUtilities.h
--- a/Engine/Plugins/Experimental/ProceduralVegetationEditor/Source/ProceduralVegetation/Public/Helpers/PVUtilities.h
+++ b/Engine/Plugins/Experimental/ProceduralVegetationEditor/Source/ProceduralVegetation/Public/Helpers/PVUtilities.h
@@ -6,6 +6,7 @@ class UInstancedStaticMeshComponent;
 class UMeshComponent;
 class FFoliageFacade;
 class AActor;
+struct FAssetData;
 
 namespace PV::Utilities
 {
diff --git a/Engine/Plugins/Experimental/ProceduralVegetationEditor/Source/ProceduralVegetation/Public/Nodes/PVOutputSettings.h b/Engine/Plugins/Experimental/ProceduralVegetationEditor/Source/ProceduralVegetation/Public/Nodes/PVOutputSettings.h
--- a/Engine/Plugins/Experimental/ProceduralVegetationEditor/Source/ProceduralVegetation/Public/Nodes/PVOutputSettings.h
+++ b/Engine/Plugins/Experimental/ProceduralVegetationEditor/Source/ProceduralVegetation/Public/Nodes/PVOutputSettings.h
@@ -34,4 +34,6 @@ public:
 class FPVOutputElement : public FPVBaseElement
 {
 	virtual bool ExecuteInternal(FPCGContext* Context) const override;
+
+	virtual bool CanExecuteOnlyOnMainThread(FPCGContext* Context) const override { return true; }
 };
\ No newline at end of file
