[PVE] Fix issue where ramp defaults are inverted in the mesh details

diff --git a/Engine/Plugins/Experimental/ProceduralVegetationEditor/Source/ProceduralVegetation/Private/Implementations/PVMeshBuilder.h b/Engine/Plugins/Experimental/ProceduralVegetationEditor/Source/ProceduralVegetation/Private/Implementations/PVMeshBuilder.h
--- a/Engine/Plugins/Experimental/ProceduralVegetationEditor/Source/ProceduralVegetation/Private/Implementations/PVMeshBuilder.h
+++ b/Engine/Plugins/Experimental/ProceduralVegetationEditor/Source/ProceduralVegetation/Private/Implementations/PVMeshBuilder.h
@@ -155,8 +155,8 @@ private:
 		if (FRichCurve* RichCurve = Curve.GetRichCurve())
 		{
 			RichCurve->Reset();
-			const FKeyHandle Handle0 = RichCurve->AddKey(0.0f, 0.0f);
-			const FKeyHandle Handle1 = RichCurve->AddKey(1.0f, 1.0f);
+			const FKeyHandle Handle0 = RichCurve->AddKey(0.0f, 1.0f);
+			const FKeyHandle Handle1 = RichCurve->AddKey(1.0f, 0.0f);
 
 			RichCurve->SetKeyInterpMode(Handle0, RCIM_Linear);
 			RichCurve->SetKeyInterpMode(Handle1, RCIM_Linear);
