Replace PROCEDURAL_VEGETATION_DEBUG_SHOW_COLLECTION_TAB define with PVDebugMode CVar

diff --git a/Engine/Plugins/Experimental/ProceduralVegetationEditor/Source/ProceduralVegetationEditor/Private/PVEditor.cpp b/Engine/Plugins/Experimental/ProceduralVegetationEditor/Source/ProceduralVegetationEditor/Private/PVEditor.cpp
--- a/Engine/Plugins/Experimental/ProceduralVegetationEditor/Source/ProceduralVegetationEditor/Private/PVEditor.cpp
+++ b/Engine/Plugins/Experimental/ProceduralVegetationEditor/Source/ProceduralVegetationEditor/Private/PVEditor.cpp
@@ -17,16 +17,11 @@
 #include "PCGEditor/Private/PCGEditorGraph.h"
 #include "Widgets/SPVCreateDialog.h"
 #include "Widgets/SPVExportSelectionDialog.h"
-
-#if PROCEDURAL_VEGETATION_DEBUG_SHOW_COLLECTION_TAB
 #include "Widgets/Docking/SDockTab.h"
-#endif
 
 #define LOCTEXT_NAMESPACE "PVEditor"
 
-#if PROCEDURAL_VEGETATION_DEBUG_SHOW_COLLECTION_TAB
 const FName FPVEditor::CollectionSpreadSheetTabId (TEXT("PVEditor_CollectionSpreadSheet"));
-#endif
 
 FPVEditor::FPVEditor()
 {
@@ -44,9 +39,10 @@ void FPVEditor::Initialize(const EToolkitMode::Type InMode, const TSharedPtr<cla
 
 	// Initialize widgets before calling base
 
-#if PROCEDURAL_VEGETATION_DEBUG_SHOW_COLLECTION_TAB
-	CollectionSpreadSheetWidget = CreateCollectionSpreadSheetWidget();
-#endif
+	if (PV::Utilities::DebugModeEnabled())
+	{
+		CollectionSpreadSheetWidget = CreateCollectionSpreadSheetWidget();
+	}
 	
 	FPCGEditor::Initialize(InMode, InToolkitHost, ProceduralVegetationGraph);
 
@@ -164,20 +160,19 @@ void FPVEditor::RegisterTabSpawners(const TSharedRef<class FTabManager>& InTabMa
 {
 	FPCGEditor::RegisterTabSpawners(InTabManager);
 
-#if PROCEDURAL_VEGETATION_DEBUG_SHOW_COLLECTION_TAB
-	InTabManager->RegisterTabSpawner(CollectionSpreadSheetTabId, FOnSpawnTab::CreateSP(this, &FPVEditor::SpawnTab_CollectionSpreadSheet))
-		.SetDisplayName(LOCTEXT("CollectionSpreadsheetTab", "Collection Spreadsheet"))
-		.SetGroup(WorkspaceMenuCategory.ToSharedRef());
-#endif
+	if (CollectionSpreadSheetWidget.IsValid())
+	{
+		InTabManager->RegisterTabSpawner(CollectionSpreadSheetTabId, FOnSpawnTab::CreateSP(this, &FPVEditor::SpawnTab_CollectionSpreadSheet))
+			.SetDisplayName(LOCTEXT("CollectionSpreadsheetTab", "Collection Spreadsheet"))
+			.SetGroup(WorkspaceMenuCategory.ToSharedRef());
+	}
 }
 
 void FPVEditor::UnregisterTabSpawners(const TSharedRef<class FTabManager>& InTabManager)
 {
 	FPCGEditor::UnregisterTabSpawners(InTabManager);
 
-#if PROCEDURAL_VEGETATION_DEBUG_SHOW_COLLECTION_TAB
 	InTabManager->UnregisterTabSpawner(CollectionSpreadSheetTabId);
-#endif
 }
 
 IPCGBaseSubsystem* FPVEditor::GetSubsystem() const
@@ -216,15 +211,12 @@ TSharedRef<FTabManager::FLayout> FPVEditor::GetDefaultLayout() const
 					->SetHideTabWell(true)
 				)
 				->Split // Bottom Section - ALV
-				(FTabManager::NewStack()
+				(
+					FTabManager::NewStack()
 					->SetSizeCoefficient(0.05f)
 					->SetHideTabWell(false)
 					->AddTab(GetPanelID(EPCGEditorPanel::Attributes1), ETabState::ClosedTab)
-
-#if PROCEDURAL_VEGETATION_DEBUG_SHOW_COLLECTION_TAB
-					->AddTab(CollectionSpreadSheetTabId, ETabState::ClosedTab)
-#endif
-				)	
+				)
 			)
 			->Split // Graph and details panel
 			(
@@ -233,16 +225,16 @@ TSharedRef<FTabManager::FLayout> FPVEditor::GetDefaultLayout() const
 				->SetSizeCoefficient(0.25f)
 				->Split
 				(
-				FTabManager::NewStack()
+					FTabManager::NewStack()
 					->SetSizeCoefficient(0.6f)
 					->AddTab(GetPanelID(EPCGEditorPanel::GraphEditor), ETabState::OpenedTab)
 					->SetHideTabWell(true)
 				)
 				->Split
 				(
-				FTabManager::NewStack()
-				->SetSizeCoefficient(0.4f)
-				->AddTab(GetPanelID(EPCGEditorPanel::PropertyDetails1), ETabState::OpenedTab)
+					FTabManager::NewStack()
+					->SetSizeCoefficient(0.4f)
+					->AddTab(GetPanelID(EPCGEditorPanel::PropertyDetails1), ETabState::OpenedTab)
 				)
 			)
 		)
@@ -301,9 +293,7 @@ void FPVEditor::OnSelectedNodesChanged(const TSet<UObject*>& InNewSelection)
 		{
 			PCGNode = Node->GetPCGNode();
 
-#if PROCEDURAL_VEGETATION_DEBUG_SHOW_COLLECTION_TAB
 			UpdateCollectionSpreadSheet(PCGNode);
-#endif			
 
 			if (UPVBaseSettings* Settings = Cast<UPVBaseSettings>(PCGNode->GetSettings()))
 			{
@@ -438,7 +428,6 @@ bool FPVEditor::GetCollectionFromNode(const UPCGNode* InNode, FName& OutPinName,
 	return false;
 }
 
-#if PROCEDURAL_VEGETATION_DEBUG_SHOW_COLLECTION_TAB
 TSharedRef<SCollectionSpreadSheetWidget> FPVEditor::CreateCollectionSpreadSheetWidget()
 {
 	return SNew(SCollectionSpreadSheetWidget);
@@ -480,6 +469,5 @@ void FPVEditor::UpdateCollectionSpreadSheet(const UPCGNode* InSelectedNode) cons
 		CollectionSpreadSheetWidget->RefreshWidget();	
 	}
 }
-#endif
 
 #undef LOCTEXT_NAMESPACE
diff --git a/Engine/Plugins/Experimental/ProceduralVegetationEditor/Source/ProceduralVegetationEditor/Private/PVEditor.h b/Engine/Plugins/Experimental/ProceduralVegetationEditor/Source/ProceduralVegetationEditor/Private/PVEditor.h
--- a/Engine/Plugins/Experimental/ProceduralVegetationEditor/Source/ProceduralVegetationEditor/Private/PVEditor.h
+++ b/Engine/Plugins/Experimental/ProceduralVegetationEditor/Source/ProceduralVegetationEditor/Private/PVEditor.h
@@ -9,8 +9,6 @@ struct FManagedArrayCollection;
 class UPCGDefaultExecutionSource;
 class UProceduralVegetation;
 
-#define PROCEDURAL_VEGETATION_DEBUG_SHOW_COLLECTION_TAB !UE_BUILD_SHIPPING_WITH_EDITOR 
-
 class FPVEditor : public FPCGEditor
 {
 public:
@@ -55,7 +53,6 @@ private:
 
 	bool bFocusViewportOnSelection = true;
 
-#if PROCEDURAL_VEGETATION_DEBUG_SHOW_COLLECTION_TAB
 	static const FName CollectionSpreadSheetTabId;
 
 	TWeakPtr<SDockTab> CollectionSpreadSheetTab;
@@ -66,7 +63,6 @@ private:
 	TSharedRef<SDockTab> SpawnTab_CollectionSpreadSheet(const FSpawnTabArgs& SpawnTabArgs);
 	
 	void UpdateCollectionSpreadSheet(const UPCGNode* InSelectedNode) const;
-#endif
 	
 	bool GetCollectionFromNode(const UPCGNode* InNode, FName& OutPinName, FManagedArrayCollection& OutCollection) const;
 	void OnExport_Clicked();
