[PVE] non unity compile issue fixed

diff --git a/Engine/Plugins/Experimental/ProceduralVegetationEditor/Source/ProceduralVegetation/Private/Nodes/PVOutputSettings.cpp b/Engine/Plugins/Experimental/ProceduralVegetationEditor/Source/ProceduralVegetation/Private/Nodes/PVOutputSettings.cpp
--- a/Engine/Plugins/Experimental/ProceduralVegetationEditor/Source/ProceduralVegetation/Private/Nodes/PVOutputSettings.cpp
+++ b/Engine/Plugins/Experimental/ProceduralVegetationEditor/Source/ProceduralVegetation/Private/Nodes/PVOutputSettings.cpp
@@ -6,6 +6,9 @@
 #include "DataTypes/PVFoliageMeshData.h"
 #include "Helpers/PCGHelpers.h"
 #include "Misc/PackageName.h"
+#include "Engine/StaticMesh.h"
+#include "Engine/SkeletalMesh.h"
+
 
 #define LOCTEXT_NAMESPACE "PVOutputSettings"
 
