[PVE] Wind data fix

diff --git a/Engine/Plugins/Experimental/ProceduralVegetationEditor/Source/ProceduralVegetation/Private/Facades/PVBoneFacade.cpp b/Engine/Plugins/Experimental/ProceduralVegetationEditor/Source/ProceduralVegetation/Private/Facades/PVBoneFacade.cpp
--- a/Engine/Plugins/Experimental/ProceduralVegetationEditor/Source/ProceduralVegetation/Private/Facades/PVBoneFacade.cpp
+++ b/Engine/Plugins/Experimental/ProceduralVegetationEditor/Source/ProceduralVegetation/Private/Facades/PVBoneFacade.cpp
@@ -212,11 +212,11 @@ namespace PV::Facades
 			Min = FMath::Min(Min, Pscale);
 			Max = FMath::Max(Max, Pscale);
 		}
-
+		
 		for (auto& Pscale : BranchRootPScales)
 		{
 			const UE::Math::TVector2<float> InputRange(Min, Max);
-			const UE::Math::TVector2<float> OutputRange(0, 1);
+			const UE::Math::TVector2<float> OutputRange(1, 0);
 			Pscale = FMath::GetMappedRangeValueClamped(InputRange, OutputRange, Pscale);
 		}
 
@@ -225,7 +225,7 @@ namespace PV::Facades
 			TArray<int32> ParentIndices = BranchFacade.GetParentBranchIndices(BranchIdx);
 			int32 ParentIndex = ParentIndices.Num() > 0 ? ParentIndices.Last() : -1;
 
-			if (BranchRootPScales.IsValidIndex(ParentIndex))
+			if (BranchRootPScales.IsValidIndex(ParentIndex) && ParentIndex > 0 )
 			{
 				BranchRootPScales[ParentIndex] = FMath::Min(BranchRootPScales[BranchIdx], BranchRootPScales[ParentIndex]);
 			}
