Add user-facing node tooltips

diff --git a/Engine/Plugins/Experimental/ProceduralVegetationEditor/Source/ProceduralVegetation/Private/Nodes/PVBoneReductionSettings.cpp b/Engine/Plugins/Experimental/ProceduralVegetationEditor/Source/ProceduralVegetation/Private/Nodes/PVBoneReductionSettings.cpp
--- a/Engine/Plugins/Experimental/ProceduralVegetationEditor/Source/ProceduralVegetation/Private/Nodes/PVBoneReductionSettings.cpp
+++ b/Engine/Plugins/Experimental/ProceduralVegetationEditor/Source/ProceduralVegetation/Private/Nodes/PVBoneReductionSettings.cpp
@@ -8,6 +8,23 @@
 #include "GeometryCollection/ManagedArrayCollection.h"
 #include "Helpers/PCGSettingsHelpers.h"
 
+#define LOCTEXT_NAMESPACE "PVBoneReductionSettings"
+
+#if WITH_EDITOR
+FText UPVBoneReductionSettings::GetDefaultNodeTitle() const 
+{ 
+	return LOCTEXT("NodeTitle", "Bone Reduction"); 
+}
+
+FText UPVBoneReductionSettings::GetNodeTooltipText() const 
+{ 
+	return LOCTEXT("NodeTooltip", 
+		"Allows the user to reduce the complexity of the skeletal bones generated. "
+		"Reduced bone complexity helps in optimizing performance but can effect the accuracy of wind simulation"
+	);
+}
+#endif
+
 FPCGDataTypeIdentifier UPVBoneReductionSettings::GetInputPinTypeIdentifier() const
 {
 	return FPCGDataTypeIdentifier{ FPVDataTypeInfoMesh::AsId() };
@@ -67,3 +84,5 @@ bool FPVBoneReductionElement::ExecuteInternal(FPCGContext* InContext) const
 
 	return true;
 }
+
+#undef LOCTEXT_NAMESPACE
\ No newline at end of file
diff --git a/Engine/Plugins/Experimental/ProceduralVegetationEditor/Source/ProceduralVegetation/Private/Nodes/PVBoneReductionSettings.h b/Engine/Plugins/Experimental/ProceduralVegetationEditor/Source/ProceduralVegetation/Private/Nodes/PVBoneReductionSettings.h
--- a/Engine/Plugins/Experimental/ProceduralVegetationEditor/Source/ProceduralVegetation/Private/Nodes/PVBoneReductionSettings.h
+++ b/Engine/Plugins/Experimental/ProceduralVegetationEditor/Source/ProceduralVegetation/Private/Nodes/PVBoneReductionSettings.h
@@ -15,8 +15,8 @@ class UPVBoneReductionSettings : public UPVBaseSettings
 public:
 #if WITH_EDITOR
 	virtual FName GetDefaultNodeName() const override { return FName(TEXT("ProceduralVegetationBoneReduction")); }
-	virtual FText GetDefaultNodeTitle() const override { return NSLOCTEXT("PVBoneReductionSettings", "NodeTitle", "Bone Reduction"); }
-	virtual FText GetNodeTooltipText() const override { return NSLOCTEXT("PVBoneReductionSettings", "NodeTooltip", "Controls the number of bones created for skeletal mesh"); }
+	virtual FText GetDefaultNodeTitle() const override;
+	virtual FText GetNodeTooltipText() const override;
 	virtual FLinearColor GetNodeTitleColor() const override { return FColor::Cyan; }
 
 	virtual TArray<EPVRenderType> GetRenderType() const override { return TArray{EPVRenderType::Bones, EPVRenderType::Mesh}; }
diff --git a/Engine/Plugins/Experimental/ProceduralVegetationEditor/Source/ProceduralVegetation/Private/Nodes/PVCarveSettings.cpp b/Engine/Plugins/Experimental/ProceduralVegetationEditor/Source/ProceduralVegetation/Private/Nodes/PVCarveSettings.cpp
--- a/Engine/Plugins/Experimental/ProceduralVegetationEditor/Source/ProceduralVegetation/Private/Nodes/PVCarveSettings.cpp
+++ b/Engine/Plugins/Experimental/ProceduralVegetationEditor/Source/ProceduralVegetation/Private/Nodes/PVCarveSettings.cpp
@@ -7,6 +7,26 @@
 #include "GeometryCollection/ManagedArrayCollection.h"
 #include "Helpers/PCGSettingsHelpers.h"
 
+#define LOCTEXT_NAMESPACE "PVCarveSettings"
+
+#if WITH_EDITOR
+FText UPVCarveSettings::GetDefaultNodeTitle() const
+{ 
+	return LOCTEXT("NodeTitle", "Carve"); 
+}
+
+FText UPVCarveSettings::GetNodeTooltipText() const
+{ 
+	return LOCTEXT("NodeTooltip", 
+		"Allows the user to trim the vegetation structure.The Carve Basis determines the reference used for the cut. \n\n"
+		"LengthFromRoot – distance measured from the plant’s root / trunk base.\n"
+		"From Bottom – relative height starting from the plant’s lowest point.\n"
+		"ZPosition – absolute world - space Z axis position.\n"
+		"Radius – Thickness of branch"
+	); 
+}
+#endif
+
 FPCGDataTypeIdentifier UPVCarveSettings::GetInputPinTypeIdentifier() const
 {
 	return FPCGDataTypeIdentifier{ FPVDataTypeInfoGrowth::AsId() };
@@ -57,3 +77,5 @@ bool FPVCarveElement::ExecuteInternal(FPCGContext* InContext) const
 
 	return true;
 }
+
+#undef LOCTEXT_NAMESPACE
\ No newline at end of file
diff --git a/Engine/Plugins/Experimental/ProceduralVegetationEditor/Source/ProceduralVegetation/Private/Nodes/PVCarveSettings.h b/Engine/Plugins/Experimental/ProceduralVegetationEditor/Source/ProceduralVegetation/Private/Nodes/PVCarveSettings.h
--- a/Engine/Plugins/Experimental/ProceduralVegetationEditor/Source/ProceduralVegetation/Private/Nodes/PVCarveSettings.h
+++ b/Engine/Plugins/Experimental/ProceduralVegetationEditor/Source/ProceduralVegetation/Private/Nodes/PVCarveSettings.h
@@ -17,8 +17,8 @@ class UPVCarveSettings : public UPVBaseSettings
 public:
 #if WITH_EDITOR
 	virtual FName GetDefaultNodeName() const override { return FName(TEXT("ProceduralVegetationCarve")); }
-	virtual FText GetDefaultNodeTitle() const override { return NSLOCTEXT("PVCarveSettings", "NodeTitle", "Carve"); }
-	virtual FText GetNodeTooltipText() const override { return NSLOCTEXT("PVCarveSettings", "NodeTooltip", "Carve branches or trunk"); }
+	virtual FText GetDefaultNodeTitle() const override;
+	virtual FText GetNodeTooltipText() const override;
 	virtual FLinearColor GetNodeTitleColor() const override { return FLinearColor::Yellow; }
 
 	virtual TArray<EPVRenderType> GetRenderType() const override { return TArray{EPVRenderType::Skeleton}; }
diff --git a/Engine/Plugins/Experimental/ProceduralVegetationEditor/Source/ProceduralVegetation/Private/Nodes/PVFoliageDistributorSettings.cpp b/Engine/Plugins/Experimental/ProceduralVegetationEditor/Source/ProceduralVegetation/Private/Nodes/PVFoliageDistributorSettings.cpp
--- a/Engine/Plugins/Experimental/ProceduralVegetationEditor/Source/ProceduralVegetation/Private/Nodes/PVFoliageDistributorSettings.cpp
+++ b/Engine/Plugins/Experimental/ProceduralVegetationEditor/Source/ProceduralVegetation/Private/Nodes/PVFoliageDistributorSettings.cpp
@@ -9,6 +9,26 @@
 #include "Helpers/PCGHelpers.h"
 #include "Helpers/PCGSettingsHelpers.h"
 
+#define LOCTEXT_NAMESPACE "PVFoliageDistributorSettings"
+
+#if WITH_EDITOR
+FText UPVFoliageDistributorSettings::GetDefaultNodeTitle() const 
+{ 
+	return LOCTEXT("NodeTitle", "Foliage Distributor"); 
+}
+
+FText UPVFoliageDistributorSettings::GetNodeTooltipText() const 
+{ 
+	return LOCTEXT("NodeTooltip", 
+		"Allows the user to customize the distribution of the foliage. \n"
+		"Ethylene plays a key role in plant aging and leaf abscission (falling).\n"
+		"EthyleneThreshold: Defines how sensitive the procedural system is to \"ethylene levels.\"\n"
+		"\tHigher Threshold = more buds/branches retained, making the plant denser and bushier.\n"
+		"\tLower Threshold = fewer buds/branches retained, resulting in a sparser, pruned structure."
+	); 
+}
+#endif
+
 UPVFoliageDistributorSettings::UPVFoliageDistributorSettings()
 {
 	if (PCGHelpers::IsNewObjectAndNotDefault(this))
@@ -124,4 +144,6 @@ bool FPVFoliageDistributorElement::ExecuteInternal(FPCGContext* InContext) const
 	}
 	
 	return true;
-}
\ No newline at end of file
+}
+
+#undef LOCTEXT_NAMESPACE
\ No newline at end of file
diff --git a/Engine/Plugins/Experimental/ProceduralVegetationEditor/Source/ProceduralVegetation/Private/Nodes/PVFoliageDistributorSettings.h b/Engine/Plugins/Experimental/ProceduralVegetationEditor/Source/ProceduralVegetation/Private/Nodes/PVFoliageDistributorSettings.h
--- a/Engine/Plugins/Experimental/ProceduralVegetationEditor/Source/ProceduralVegetation/Private/Nodes/PVFoliageDistributorSettings.h
+++ b/Engine/Plugins/Experimental/ProceduralVegetationEditor/Source/ProceduralVegetation/Private/Nodes/PVFoliageDistributorSettings.h
@@ -24,8 +24,8 @@ public:
 
 	#if WITH_EDITOR
 	virtual FName GetDefaultNodeName() const override { return FName(TEXT("ProceduralVegetationFoliageDistributor")); }
-	virtual FText GetDefaultNodeTitle() const override { return NSLOCTEXT("PVFoliageDistributorSettings", "NodeTitle", "Foliage Distributor"); }
-	virtual FText GetNodeTooltipText() const override { return NSLOCTEXT("PVFoliageDistributorSettings", "NodeTooltip", "Distributes the Foliage"); }
+	virtual FText GetDefaultNodeTitle() const override;
+	virtual FText GetNodeTooltipText() const override;
 	virtual FLinearColor GetNodeTitleColor() const override { return FLinearColor::Blue; }
 
 	virtual TArray<EPVRenderType> GetRenderType() const override { return TArray{ EPVRenderType::Mesh, EPVRenderType::Foliage }; }
diff --git a/Engine/Plugins/Experimental/ProceduralVegetationEditor/Source/ProceduralVegetation/Private/Nodes/PVFoliagePaletteSettings.cpp b/Engine/Plugins/Experimental/ProceduralVegetationEditor/Source/ProceduralVegetation/Private/Nodes/PVFoliagePaletteSettings.cpp
--- a/Engine/Plugins/Experimental/ProceduralVegetationEditor/Source/ProceduralVegetation/Private/Nodes/PVFoliagePaletteSettings.cpp
+++ b/Engine/Plugins/Experimental/ProceduralVegetationEditor/Source/ProceduralVegetation/Private/Nodes/PVFoliagePaletteSettings.cpp
@@ -12,7 +12,23 @@
 #include "Helpers/PCGHelpers.h"
 #include "Helpers/PCGSettingsHelpers.h"
 
+#define LOCTEXT_NAMESPACE "PVFoliagePaletteSettings"
+
 #if WITH_EDITOR
+FText UPVFoliagePaletteSettings::GetDefaultNodeTitle() const 
+{ 
+	return LOCTEXT("NodeTitle", "Foliage Palette");
+}
+
+FText UPVFoliagePaletteSettings::GetNodeTooltipText() const 
+{ 
+	return LOCTEXT("NodeTooltip", 
+		"Allows the user to visualize and optionally change the foliage meshes. "
+		"Foliage mesh data is loaded from the Procedrual Vegetation preset. "
+		"The user can replace or remove foliage meshes. The placement of folliage meshes is driven by the procedural vegetaiton preset and cannot be modifed"
+	);
+}
+
 void UPVFoliagePaletteSettings::PostEditChangeProperty(FPropertyChangedEvent& PropertyChangedEvent)
 {
 	Super::PostEditChangeProperty(PropertyChangedEvent);
@@ -117,3 +133,5 @@ bool FPVFoliageElement::ExecuteInternal(FPCGContext* InContext) const
 
 	return true;
 }
+
+#undef LOCTEXT_NAMESPACE
\ No newline at end of file
diff --git a/Engine/Plugins/Experimental/ProceduralVegetationEditor/Source/ProceduralVegetation/Private/Nodes/PVFoliagePaletteSettings.h b/Engine/Plugins/Experimental/ProceduralVegetationEditor/Source/ProceduralVegetation/Private/Nodes/PVFoliagePaletteSettings.h
--- a/Engine/Plugins/Experimental/ProceduralVegetationEditor/Source/ProceduralVegetation/Private/Nodes/PVFoliagePaletteSettings.h
+++ b/Engine/Plugins/Experimental/ProceduralVegetationEditor/Source/ProceduralVegetation/Private/Nodes/PVFoliagePaletteSettings.h
@@ -14,8 +14,8 @@ class UPVFoliagePaletteSettings : public UPVBaseSettings
 public:
 #if WITH_EDITOR
 	virtual FName GetDefaultNodeName() const override { return FName(TEXT("ProceduralVegetationFoliage")); }
-	virtual FText GetDefaultNodeTitle() const override { return NSLOCTEXT("PVFoliagePaletteSettings", "NodeTitle", "Foliage Palette"); }
-	virtual FText GetNodeTooltipText() const override { return NSLOCTEXT("PVFoliagePaletteSettings", "NodeTooltip", "Foliage setup for the tree"); }
+	virtual FText GetDefaultNodeTitle() const override;
+	virtual FText GetNodeTooltipText() const override;
 	virtual FLinearColor GetNodeTitleColor() const override { return FLinearColor::Green; }
 
 	virtual TArray<EPVRenderType> GetRenderType() const override { return TArray{EPVRenderType::Mesh, EPVRenderType::Foliage}; }
diff --git a/Engine/Plugins/Experimental/ProceduralVegetationEditor/Source/ProceduralVegetation/Private/Nodes/PVGravitySettings.cpp b/Engine/Plugins/Experimental/ProceduralVegetationEditor/Source/ProceduralVegetation/Private/Nodes/PVGravitySettings.cpp
--- a/Engine/Plugins/Experimental/ProceduralVegetationEditor/Source/ProceduralVegetation/Private/Nodes/PVGravitySettings.cpp
+++ b/Engine/Plugins/Experimental/ProceduralVegetationEditor/Source/ProceduralVegetation/Private/Nodes/PVGravitySettings.cpp
@@ -8,6 +8,27 @@
 #include "Helpers/PCGSettingsHelpers.h"
 #include "Implementations/PVGravity.h"
 
+#define LOCTEXT_NAMESPACE "PVGravitySettings"
+
+#if WITH_EDITOR
+FText UPVGravitySettings::GetDefaultNodeTitle() const
+{ 
+	return LOCTEXT("NodeTitle", "Gravity"); 
+}
+
+FText UPVGravitySettings::GetNodeTooltipText() const
+{ 
+	return LOCTEXT("NodeTooltip", 
+		"This node controls how generated vegetation responds to directional forces, mimicking natural phenomena such as the downward pull of gravity and the growth tendency toward the direction of the optimal light conditions.(phototropism).\n\n"
+
+		"Gravity Mode: Defines the behavior — either Gravity (branches bend downward under weight) or Phototropic (branches orient upward toward light).\n"
+		"Strength Controls: Parameters adjust how strongly branches respond, allowing subtle or dramatic curvature.\n"
+		"Direction Vector: When in gravity mode, you can set the exact vector (e.g., world down, custom direction).\n"
+		"Angle Correction & Bias: Additional properties fine-tune how branches preserve their initial angles or bias between optimal light direction/shadow avoidance.\n"
+	); 
+}
+#endif
+
 FPCGDataTypeIdentifier UPVGravitySettings::GetInputPinTypeIdentifier() const
 {
 	return FPCGDataTypeIdentifier{ FPVDataTypeInfoGrowth::AsId() };
@@ -57,3 +78,5 @@ bool FPVGravityElement::ExecuteInternal(FPCGContext* InContext) const
 	
 	return true;
 }
+
+#undef LOCTEXT_NAMESPACE
\ No newline at end of file
diff --git a/Engine/Plugins/Experimental/ProceduralVegetationEditor/Source/ProceduralVegetation/Private/Nodes/PVGravitySettings.h b/Engine/Plugins/Experimental/ProceduralVegetationEditor/Source/ProceduralVegetation/Private/Nodes/PVGravitySettings.h
--- a/Engine/Plugins/Experimental/ProceduralVegetationEditor/Source/ProceduralVegetation/Private/Nodes/PVGravitySettings.h
+++ b/Engine/Plugins/Experimental/ProceduralVegetationEditor/Source/ProceduralVegetation/Private/Nodes/PVGravitySettings.h
@@ -14,8 +14,8 @@ class UPVGravitySettings : public UPVBaseSettings
 public:
 #if WITH_EDITOR
 	virtual FName GetDefaultNodeName() const override { return FName(TEXT("ProceduralVegetationGravity")); }
-	virtual FText GetDefaultNodeTitle() const override { return NSLOCTEXT("PVGravitySettings", "NodeTitle", "Gravity"); }
-	virtual FText GetNodeTooltipText() const override { return NSLOCTEXT("PVGravitySettings", "NodeTooltip", "Apply gravity to branches and main trunk"); }
+	virtual FText GetDefaultNodeTitle() const override;
+	virtual FText GetNodeTooltipText() const override;
 	virtual FLinearColor GetNodeTitleColor() const override { return FLinearColor::Yellow; }
 	
 	virtual TArray<EPVRenderType> GetRenderType() const override { return TArray{ EPVRenderType::Skeleton }; }
diff --git a/Engine/Plugins/Experimental/ProceduralVegetationEditor/Source/ProceduralVegetation/Private/Nodes/PVMeshBuilderSettings.cpp b/Engine/Plugins/Experimental/ProceduralVegetationEditor/Source/ProceduralVegetation/Private/Nodes/PVMeshBuilderSettings.cpp
--- a/Engine/Plugins/Experimental/ProceduralVegetationEditor/Source/ProceduralVegetation/Private/Nodes/PVMeshBuilderSettings.cpp
+++ b/Engine/Plugins/Experimental/ProceduralVegetationEditor/Source/ProceduralVegetation/Private/Nodes/PVMeshBuilderSettings.cpp
@@ -12,7 +12,25 @@
 #include "Implementations/PVMeshBuilder.h"
 #include "Materials/MaterialInterface.h"
 
+#define LOCTEXT_NAMESPACE "PVMeshBuilderSettings"
+
 #if WITH_EDITOR
+FText UPVMeshBuilderSettings::GetDefaultNodeTitle() const
+{ 
+	return LOCTEXT("NodeTitle", "Mesh Builder"); 
+}
+
+FText UPVMeshBuilderSettings::GetNodeTooltipText() const 
+{ 
+	return LOCTEXT("NodeTooltip", 
+		"Builds a 3D Mesh for the plant based on the structure incoming from previous nodes.\n\n"
+		"Users can specify materials to be used for the mesh. By default it uses the materials specified in the Procedural Vegetation Preset. It has 3 groups of settings : \n"
+		"- Material\n"
+		"- Mesh\n"
+		"- Displacement\n"
+	); 
+}
+
 void UPVMeshBuilderSettings::PostEditChangeProperty(FPropertyChangedEvent& PropertyChangedEvent)
 {
 	Super::PostEditChangeProperty(PropertyChangedEvent);
@@ -152,3 +170,5 @@ bool FPVMeshBuilderElement::ExecuteInternal(FPCGContext* InContext) const
 
 	return true;
 }
+
+#undef LOCTEXT_NAMESPACE
\ No newline at end of file
diff --git a/Engine/Plugins/Experimental/ProceduralVegetationEditor/Source/ProceduralVegetation/Private/Nodes/PVMeshBuilderSettings.h b/Engine/Plugins/Experimental/ProceduralVegetationEditor/Source/ProceduralVegetation/Private/Nodes/PVMeshBuilderSettings.h
--- a/Engine/Plugins/Experimental/ProceduralVegetationEditor/Source/ProceduralVegetation/Private/Nodes/PVMeshBuilderSettings.h
+++ b/Engine/Plugins/Experimental/ProceduralVegetationEditor/Source/ProceduralVegetation/Private/Nodes/PVMeshBuilderSettings.h
@@ -16,8 +16,8 @@ class UPVMeshBuilderSettings : public UPVBaseSettings
 public:
 #if WITH_EDITOR
 	virtual FName GetDefaultNodeName() const override { return FName(TEXT("ProceduralVegetationMesher")); }
-	virtual FText GetDefaultNodeTitle() const override { return NSLOCTEXT("PVMeshBuilderSettings", "NodeTitle", "Mesh Builder"); }
-	virtual FText GetNodeTooltipText() const override { return NSLOCTEXT("PVMeshBuilderSettings", "NodeTooltip", "Create a mesh for the skeleton"); }
+	virtual FText GetDefaultNodeTitle() const override;
+	virtual FText GetNodeTooltipText() const override;
 	virtual FLinearColor GetNodeTitleColor() const override { return FLinearColor(FColor::Cyan); }
 	virtual TArray<EPVRenderType> GetRenderType() const override { return TArray{EPVRenderType::Mesh}; }
 
diff --git a/Engine/Plugins/Experimental/ProceduralVegetationEditor/Source/ProceduralVegetation/Private/Nodes/PVOutputSettings.cpp b/Engine/Plugins/Experimental/ProceduralVegetationEditor/Source/ProceduralVegetation/Private/Nodes/PVOutputSettings.cpp
--- a/Engine/Plugins/Experimental/ProceduralVegetationEditor/Source/ProceduralVegetation/Private/Nodes/PVOutputSettings.cpp
+++ b/Engine/Plugins/Experimental/ProceduralVegetationEditor/Source/ProceduralVegetation/Private/Nodes/PVOutputSettings.cpp
@@ -12,6 +12,28 @@
 
 #define LOCTEXT_NAMESPACE "PVOutputSettings"
 
+#if WITH_EDITOR
+FText UPVOutputSettings::GetDefaultNodeTitle() const
+{
+	return ExportSettings.MeshName.IsNone() 
+		? LOCTEXT("NodeTitle", "Output") 
+		: FText::FromName(ExportSettings.MeshName);
+}
+
+FText UPVOutputSettings::GetNodeTooltipText() const 
+{ 
+	return LOCTEXT("NodeTooltip", 
+		"This node allows the user to save the output of the graph as:\n"
+		"- Static Mesh\n"
+		"- Static Mesh with Nanite Foliage enabled\n"
+		"- Skeletal Mesh\n"
+		"- Skeletal Mesh with Nanite Foliage enabled\n\n"
+
+		"If the foliage meshes do not contain skeletal data and Skeletal mesh export is selected, it will build the skeletal data for foliage meshes"
+	); 
+}
+#endif
+
 UPVOutputSettings::UPVOutputSettings()
 {
 	if (PCGHelpers::IsNewObjectAndNotDefault(this))
diff --git a/Engine/Plugins/Experimental/ProceduralVegetationEditor/Source/ProceduralVegetation/Private/Nodes/PVPresetLoaderSettings.cpp b/Engine/Plugins/Experimental/ProceduralVegetationEditor/Source/ProceduralVegetation/Private/Nodes/PVPresetLoaderSettings.cpp
--- a/Engine/Plugins/Experimental/ProceduralVegetationEditor/Source/ProceduralVegetation/Private/Nodes/PVPresetLoaderSettings.cpp
+++ b/Engine/Plugins/Experimental/ProceduralVegetationEditor/Source/ProceduralVegetation/Private/Nodes/PVPresetLoaderSettings.cpp
@@ -5,6 +5,25 @@
 #include "DataTypes/PVGrowthData.h"
 #include "ProceduralVegetationPreset.h"
 
+#define LOCTEXT_NAMESPACE "PVPresetLoaderSettings"
+
+#if WITH_EDITOR
+FText UPVPresetLoaderSettings::GetDefaultNodeTitle() const
+{ 
+	return LOCTEXT("NodeTitle", "Procedural Vegetation Preset Loader"); 
+}
+
+FText UPVPresetLoaderSettings::GetNodeTooltipText() const
+{	
+	return LOCTEXT("NodeTitle", 
+		"This should be the first node. "
+		"This allows the user to load a procedural vegetation preset. "
+		"A procedural vegatation preset contains data for a specie of vegetaion. "
+		"It contains hormonal information, the main structure of the vegetation, data about the branches, where foliage will / can be attached, references to foliage meshes, materials to be applied"
+	);
+}
+#endif
+
 TArray<FPCGPinProperties> UPVPresetLoaderSettings::InputPinProperties() const
 {
 	return TArray<FPCGPinProperties>();
@@ -62,3 +81,5 @@ bool FPVPresetLoaderElement::ExecuteInternal(FPCGContext* InContext) const
 
 	return true;
 }
+
+#undef LOCTEXT_NAMESPACE
\ No newline at end of file
diff --git a/Engine/Plugins/Experimental/ProceduralVegetationEditor/Source/ProceduralVegetation/Private/Nodes/PVPresetLoaderSettings.h b/Engine/Plugins/Experimental/ProceduralVegetationEditor/Source/ProceduralVegetation/Private/Nodes/PVPresetLoaderSettings.h
--- a/Engine/Plugins/Experimental/ProceduralVegetationEditor/Source/ProceduralVegetation/Private/Nodes/PVPresetLoaderSettings.h
+++ b/Engine/Plugins/Experimental/ProceduralVegetationEditor/Source/ProceduralVegetation/Private/Nodes/PVPresetLoaderSettings.h
@@ -14,8 +14,8 @@ class UPVPresetLoaderSettings : public UPVBaseSettings
 public:
 #if WITH_EDITOR
 	virtual FName GetDefaultNodeName() const override { return FName(TEXT("ProceduralVegetationPresetLoader")); }
-	virtual FText GetDefaultNodeTitle() const override { return NSLOCTEXT("PVPresetLoaderSettings", "NodeTitle", "Procedural Vegetation Preset Loader"); }
-	virtual FText GetNodeTooltipText() const override { return NSLOCTEXT("PVPresetLoaderSettings", "NodeTooltip", "Load Procedural Vegetation Preset"); }
+	virtual FText GetDefaultNodeTitle() const override;
+	virtual FText GetNodeTooltipText() const override;
 
 	virtual TArray<EPVRenderType> GetRenderType() const override { return TArray{ EPVRenderType::Skeleton }; }
 #endif
diff --git a/Engine/Plugins/Experimental/ProceduralVegetationEditor/Source/ProceduralVegetation/Private/Nodes/PVRemoveBranchesSettings.cpp b/Engine/Plugins/Experimental/ProceduralVegetationEditor/Source/ProceduralVegetation/Private/Nodes/PVRemoveBranchesSettings.cpp
--- a/Engine/Plugins/Experimental/ProceduralVegetationEditor/Source/ProceduralVegetation/Private/Nodes/PVRemoveBranchesSettings.cpp
+++ b/Engine/Plugins/Experimental/ProceduralVegetationEditor/Source/ProceduralVegetation/Private/Nodes/PVRemoveBranchesSettings.cpp
@@ -8,6 +8,20 @@
 #include "Helpers/PCGSettingsHelpers.h"
 #include "Implementations/PVRemoveBranches.h"
 
+#define LOCTEXT_NAMESPACE "PVRemoveBranchesSettings"
+
+#if WITH_EDITOR
+FText UPVRemoveBranchesSettings::GetDefaultNodeTitle() const
+{ 
+	return LOCTEXT("NodeTitle", "Remove Branches");
+}
+
+FText UPVRemoveBranchesSettings::GetNodeTooltipText() const
+{ 
+	return LOCTEXT("NodeTooltip", "Removes certain branches on the basis of Length, Radius, Light (How exposed a branch is to the light), Age and Generation"); 
+}
+#endif
+
 FPCGDataTypeIdentifier UPVRemoveBranchesSettings::GetInputPinTypeIdentifier() const
 {
 	return FPCGDataTypeIdentifier{ FPVDataTypeInfoGrowth::AsId() };
@@ -54,3 +68,5 @@ bool FPVRemoveBranchesElement::ExecuteInternal(FPCGContext* InContext) const
 	
 	return true;
 }
+
+#undef LOCTEXT_NAMESPACE
\ No newline at end of file
diff --git a/Engine/Plugins/Experimental/ProceduralVegetationEditor/Source/ProceduralVegetation/Private/Nodes/PVRemoveBranchesSettings.h b/Engine/Plugins/Experimental/ProceduralVegetationEditor/Source/ProceduralVegetation/Private/Nodes/PVRemoveBranchesSettings.h
--- a/Engine/Plugins/Experimental/ProceduralVegetationEditor/Source/ProceduralVegetation/Private/Nodes/PVRemoveBranchesSettings.h
+++ b/Engine/Plugins/Experimental/ProceduralVegetationEditor/Source/ProceduralVegetation/Private/Nodes/PVRemoveBranchesSettings.h
@@ -14,8 +14,8 @@ class UPVRemoveBranchesSettings : public UPVBaseSettings
 public:
 #if WITH_EDITOR
 	virtual FName GetDefaultNodeName() const override { return FName(TEXT("ProceduralVegetationRemoveBranches")); }
-	virtual FText GetDefaultNodeTitle() const override { return NSLOCTEXT("PVRemoveBranchesSettings", "NodeTitle", "Remove Branches"); }
-	virtual FText GetNodeTooltipText() const override { return NSLOCTEXT("PVRemoveBranchesSettings", "NodeTooltip", "Remove branches on the basis of selected criteria"); }
+	virtual FText GetDefaultNodeTitle() const override;
+	virtual FText GetNodeTooltipText() const override;
 	virtual FLinearColor GetNodeTitleColor() const override { return FLinearColor::Yellow; }
 
 	virtual TArray<EPVRenderType> GetRenderType() const override { return TArray{ EPVRenderType::Skeleton }; }
diff --git a/Engine/Plugins/Experimental/ProceduralVegetationEditor/Source/ProceduralVegetation/Private/Nodes/PVScaleSettings.cpp b/Engine/Plugins/Experimental/ProceduralVegetationEditor/Source/ProceduralVegetation/Private/Nodes/PVScaleSettings.cpp
--- a/Engine/Plugins/Experimental/ProceduralVegetationEditor/Source/ProceduralVegetation/Private/Nodes/PVScaleSettings.cpp
+++ b/Engine/Plugins/Experimental/ProceduralVegetationEditor/Source/ProceduralVegetation/Private/Nodes/PVScaleSettings.cpp
@@ -8,6 +8,20 @@
 #include "Helpers/PCGSettingsHelpers.h"
 #include "Implementations/PVScale.h"
 
+#define LOCTEXT_NAMESPACE "PVScaleSettings"
+
+#if WITH_EDITOR
+FText UPVScaleSettings::GetDefaultNodeTitle() const
+{ 
+	return LOCTEXT("NodeTitle", "Scale"); 
+}
+
+FText UPVScaleSettings::GetNodeTooltipText() const
+{ 
+	return LOCTEXT("NodeTooltip", "This node defines how large or small vegetation elements appear within the generated structure. Scaling affects trunks, branches, leaves, or entire sub-structures.");
+}
+#endif
+
 FPCGDataTypeIdentifier UPVScaleSettings::GetInputPinTypeIdentifier() const
 {
 	return FPCGDataTypeIdentifier{ FPVDataTypeInfoGrowth::AsId() };
@@ -54,3 +68,5 @@ bool FPVScaleElement::ExecuteInternal(FPCGContext* InContext) const
 	
 	return true;
 }
+
+#undef LOCTEXT_NAMESPACE
\ No newline at end of file
diff --git a/Engine/Plugins/Experimental/ProceduralVegetationEditor/Source/ProceduralVegetation/Private/Nodes/PVScaleSettings.h b/Engine/Plugins/Experimental/ProceduralVegetationEditor/Source/ProceduralVegetation/Private/Nodes/PVScaleSettings.h
--- a/Engine/Plugins/Experimental/ProceduralVegetationEditor/Source/ProceduralVegetation/Private/Nodes/PVScaleSettings.h
+++ b/Engine/Plugins/Experimental/ProceduralVegetationEditor/Source/ProceduralVegetation/Private/Nodes/PVScaleSettings.h
@@ -13,8 +13,8 @@ class UPVScaleSettings : public UPVBaseSettings
 public:
 #if WITH_EDITOR
 	virtual FName GetDefaultNodeName() const override { return FName(TEXT("ProceduralVegetationScale")); }
-	virtual FText GetDefaultNodeTitle() const override { return NSLOCTEXT("PVScaleSettings", "NodeTitle", "Scale"); }
-	virtual FText GetNodeTooltipText() const override { return NSLOCTEXT("PVScaleSettings", "NodeTooltip", "Scale the tree based on the settings"); }
+	virtual FText GetDefaultNodeTitle() const override;
+	virtual FText GetNodeTooltipText() const override;
 	virtual FLinearColor GetNodeTitleColor() const override { return FLinearColor::Yellow; }
 
 	virtual TArray<EPVRenderType> GetRenderType() const override { return TArray{ EPVRenderType::Skeleton }; }
diff --git a/Engine/Plugins/Experimental/ProceduralVegetationEditor/Source/ProceduralVegetation/Public/Nodes/PVOutputSettings.h b/Engine/Plugins/Experimental/ProceduralVegetationEditor/Source/ProceduralVegetation/Public/Nodes/PVOutputSettings.h
--- a/Engine/Plugins/Experimental/ProceduralVegetationEditor/Source/ProceduralVegetation/Public/Nodes/PVOutputSettings.h
+++ b/Engine/Plugins/Experimental/ProceduralVegetationEditor/Source/ProceduralVegetation/Public/Nodes/PVOutputSettings.h
@@ -15,8 +15,9 @@ public:
 	UPVOutputSettings();
 
 #if WITH_EDITOR
-	virtual FName GetDefaultNodeName() const override { return ExportSettings.MeshName.IsNone() ? "ProceduralVegetation Output" : ExportSettings.MeshName; }
-	virtual FText GetNodeTooltipText() const override { return NSLOCTEXT("PVOutputSettings", "NodeTooltip", "Output node for Procedural Vegetation"); }
+	virtual FName GetDefaultNodeName() const override { return ExportSettings.MeshName.IsNone() ? "Output" : ExportSettings.MeshName; }
+	virtual FText GetDefaultNodeTitle() const override;
+	virtual FText GetNodeTooltipText() const override;
 	virtual EPCGSettingsType GetType() const override { return EPCGSettingsType::InputOutput; }
 #endif
 
