Address minor comments in 45785638
- Change loctext key to NodeTooltip
- Reword tooltip of PVPresetLoaderSettings & PVGravitySettings

#DisabledValidations Preflight

diff --git a/Engine/Plugins/Experimental/ProceduralVegetationEditor/Source/ProceduralVegetation/Private/Nodes/PVGravitySettings.cpp b/Engine/Plugins/Experimental/ProceduralVegetationEditor/Source/ProceduralVegetation/Private/Nodes/PVGravitySettings.cpp
--- a/Engine/Plugins/Experimental/ProceduralVegetationEditor/Source/ProceduralVegetation/Private/Nodes/PVGravitySettings.cpp
+++ b/Engine/Plugins/Experimental/ProceduralVegetationEditor/Source/ProceduralVegetation/Private/Nodes/PVGravitySettings.cpp
@@ -21,7 +21,7 @@ FText UPVGravitySettings::GetNodeTooltipText() const
 	return LOCTEXT("NodeTooltip", 
 		"This node controls how generated vegetation responds to directional forces, mimicking natural phenomena such as the downward pull of gravity and the growth tendency toward the direction of the optimal light conditions.(phototropism).\n\n"
 
-		"Gravity Mode: Defines the behavior — either Gravity (branches bend downward under weight) or Phototropic (branches orient upward toward light).\n"
+		"Gravity Mode: Defines the behavior — either Gravity (branches bend downward under weight) or Phototropic (branches orient toward light).\n"
 		"Strength Controls: Parameters adjust how strongly branches respond, allowing subtle or dramatic curvature.\n"
 		"Direction Vector: When in gravity mode, you can set the exact vector (e.g., world down, custom direction).\n"
 		"Angle Correction & Bias: Additional properties fine-tune how branches preserve their initial angles or bias between optimal light direction/shadow avoidance.\n"
diff --git a/Engine/Plugins/Experimental/ProceduralVegetationEditor/Source/ProceduralVegetation/Private/Nodes/PVPresetLoaderSettings.cpp b/Engine/Plugins/Experimental/ProceduralVegetationEditor/Source/ProceduralVegetation/Private/Nodes/PVPresetLoaderSettings.cpp
--- a/Engine/Plugins/Experimental/ProceduralVegetationEditor/Source/ProceduralVegetation/Private/Nodes/PVPresetLoaderSettings.cpp
+++ b/Engine/Plugins/Experimental/ProceduralVegetationEditor/Source/ProceduralVegetation/Private/Nodes/PVPresetLoaderSettings.cpp
@@ -15,8 +15,7 @@ FText UPVPresetLoaderSettings::GetDefaultNodeTitle() const
 
 FText UPVPresetLoaderSettings::GetNodeTooltipText() const
 {	
-	return LOCTEXT("NodeTitle", 
-		"This should be the first node. "
+	return LOCTEXT("NodeTooltip", 
 		"This allows the user to load a procedural vegetation preset. "
 		"A procedural vegatation preset contains data for a specie of vegetaion. "
 		"It contains hormonal information, the main structure of the vegetation, data about the branches, where foliage will / can be attached, references to foliage meshes, materials to be applied"
