[PVE] Multiple fixes
- Mannequin show/hide
- Mannequin bounds issue
- Enable Nanite on Mannequin
- Set editor viewport to realtime (removes noisy shadows)

#virtualized

diff --git a/Engine/Plugins/Experimental/ProceduralVegetationEditor/Source/ProceduralVegetationEditor/Private/PVEditor.cpp b/Engine/Plugins/Experimental/ProceduralVegetationEditor/Source/ProceduralVegetationEditor/Private/PVEditor.cpp
--- a/Engine/Plugins/Experimental/ProceduralVegetationEditor/Source/ProceduralVegetationEditor/Private/PVEditor.cpp
+++ b/Engine/Plugins/Experimental/ProceduralVegetationEditor/Source/ProceduralVegetationEditor/Private/PVEditor.cpp
@@ -1,21 +1,30 @@
 ﻿// Copyright Epic Games, Inc. All Rights Reserved.
 
 #include "PVEditor.h"
+
 #include "EngineAnalytics.h"
-#include "PVEditorCommands.h"
 #include "PCGDefaultExecutionSource.h"
-#include "PCGEditorSettings.h"
-#include "Subsystems/PCGEngineSubsystem.h"
-#include "PVEditorSchema.h"
+#include "PCGEditorCommands.h"
 #include "ProceduralVegetation.h"
-#include "DataTypes/PVData.h"
-#include "Nodes/PVOutputSettings.h"
+#include "PVEditorCommands.h"
+#include "PVEditorSchema.h"
+#include "PVEditorSettings.h"
 #include "ToolMenus.h"
 #include "ToolMenuSection.h"
+
 #include "Dataflow/DataflowCollectionSpreadSheetWidget.h"
+
+#include "DataTypes/PVData.h"
+
+#include "Helpers/PVUtilities.h"
+
 #include "Nodes/PCGEditorGraphNodeBase.h"
+#include "Nodes/PVOutputSettings.h"
+
 #include "PCGEditor/Private/PCGEditorGraph.h"
-#include "Widgets/SPVCreateDialog.h"
+
+#include "Subsystems/PCGEngineSubsystem.h"
+
 #include "Widgets/SPVExportSelectionDialog.h"
 #include "Widgets/Docking/SDockTab.h"
 
@@ -73,14 +82,13 @@ void FPVEditor::Initialize(const EToolkitMode::Type InMode, const TSharedPtr<cla
 	// Select the first available output
 	for (auto EdNode : GetPCGEditorGraph()->Nodes)
 	{
-		if (UPCGEditorGraphNodeBase* PCGEdNode = Cast<UPCGEditorGraphNodeBase>(EdNode))
+		if (const UPCGEditorGraphNodeBase* PCGEdNode = Cast<UPCGEditorGraphNodeBase>(EdNode))
 		{
-			if (auto PCGNode = PCGEdNode->GetPCGNode())
+			if (const UPCGNode* PCGNode = PCGEdNode->GetPCGNode())
 			{
-				if(UPVOutputSettings* Settings = Cast<UPVOutputSettings>(PCGNode->GetSettings()))
+				if(PCGNode->GetSettings()->IsA<UPVOutputSettings>())
 				{
 					GetPCGEditorGraph()->SelectNodeSet({PCGEdNode}, true);
-					
 					break;
 				}
 			}
@@ -266,6 +274,14 @@ void FPVEditor::RegisterToolbarInternal(FToolMenuSection& PCGSection) const
 			TAttribute<FText>(),
 			TAttribute<FText>(),
 			FSlateIcon(FAppStyle::GetAppStyleSetName(), "Icons.Toolbar.Export")));
+
+	PCGSection.AddSeparator("Debug Settings");
+
+	PCGSection.AddEntry(FToolMenuEntry::InitToolBarButton(
+		Commands.ShowMannequin,
+		TAttribute<FText>(),
+		TAttribute<FText>(),
+		FSlateIcon(FAppStyle::GetAppStyleSetName(), "ClassIcon.PoseAsset")));
 }
 
 void FPVEditor::BindCommands()
@@ -277,6 +293,24 @@ void FPVEditor::BindCommands()
 	ToolkitCommands->MapAction(
 		EditorCommands.Export,
 		FExecuteAction::CreateSP(this, &FPVEditor::OnExport_Clicked));
+
+	ToolkitCommands->MapAction(
+		EditorCommands.ShowMannequin,
+		FExecuteAction::CreateLambda(
+			[this]
+				{
+					UPVEditorSettings* Settings = GetMutableDefault<UPVEditorSettings>();
+					Settings->bShowMannequin = !Settings->bShowMannequin;
+					GetPCGEditorGraph()->GetPCGGraph()->ForceNotificationForEditor(EPCGChangeType::Settings);
+				}
+		),
+		FCanExecuteAction(),
+		FIsActionChecked::CreateLambda([]
+			{
+				return GetDefault<UPVEditorSettings>()->bShowMannequin;
+			}
+		)
+	);
 }
 
 void FPVEditor::OnSelectedNodesChanged(const TSet<UObject*>& InNewSelection)
diff --git a/Engine/Plugins/Experimental/ProceduralVegetationEditor/Source/ProceduralVegetationEditor/Private/PVEditorCommands.cpp b/Engine/Plugins/Experimental/ProceduralVegetationEditor/Source/ProceduralVegetationEditor/Private/PVEditorCommands.cpp
--- a/Engine/Plugins/Experimental/ProceduralVegetationEditor/Source/ProceduralVegetationEditor/Private/PVEditorCommands.cpp
+++ b/Engine/Plugins/Experimental/ProceduralVegetationEditor/Source/ProceduralVegetationEditor/Private/PVEditorCommands.cpp
@@ -1,6 +1,7 @@
 // Copyright Epic Games, Inc. All Rights Reserved.
 
 #include "PVEditorCommands.h"
+
 #include "Styling/AppStyle.h"
 
 #define LOCTEXT_NAMESPACE "PVEditorCommands"
@@ -17,6 +18,7 @@ FPVEditorCommands::FPVEditorCommands()
 void FPVEditorCommands::RegisterCommands()
 {
 	UI_COMMAND(Export, "Export", "Export outputs.", EUserInterfaceActionType::Button, FInputChord(EModifierKey::Control, EKeys::E));
+	UI_COMMAND(ShowMannequin, "Mannequin", "Show/Hide Mannequin", EUserInterfaceActionType::ToggleButton, FInputChord(EModifierKey::Control, EKeys::M));
 }
 
 #undef LOCTEXT_NAMESPACE
diff --git a/Engine/Plugins/Experimental/ProceduralVegetationEditor/Source/ProceduralVegetationEditor/Private/PVEditorCommands.h b/Engine/Plugins/Experimental/ProceduralVegetationEditor/Source/ProceduralVegetationEditor/Private/PVEditorCommands.h
--- a/Engine/Plugins/Experimental/ProceduralVegetationEditor/Source/ProceduralVegetationEditor/Private/PVEditorCommands.h
+++ b/Engine/Plugins/Experimental/ProceduralVegetationEditor/Source/ProceduralVegetationEditor/Private/PVEditorCommands.h
@@ -14,4 +14,6 @@ public:
 	// ~End TCommands<> interface
 
 	TSharedPtr<FUICommandInfo> Export;
+	
+	TSharedPtr<FUICommandInfo> ShowMannequin;
 };
\ No newline at end of file
diff --git a/Engine/Plugins/Experimental/ProceduralVegetationEditor/Source/ProceduralVegetationEditor/Private/Visualizations/DebugVisualization/PVDebugVisualizationBase.cpp b/Engine/Plugins/Experimental/ProceduralVegetationEditor/Source/ProceduralVegetationEditor/Private/Visualizations/DebugVisualization/PVDebugVisualizationBase.cpp
--- a/Engine/Plugins/Experimental/ProceduralVegetationEditor/Source/ProceduralVegetationEditor/Private/Visualizations/DebugVisualization/PVDebugVisualizationBase.cpp
+++ b/Engine/Plugins/Experimental/ProceduralVegetationEditor/Source/ProceduralVegetationEditor/Private/Visualizations/DebugVisualization/PVDebugVisualizationBase.cpp
@@ -2,14 +2,19 @@
 
 #include "PVDebugVisualizationBase.h"
 
-#include "UObject/Package.h"
-#include "Engine/StaticMesh.h"
 #include "AdvancedPreviewScene.h"
 #include "PVDataAttributeValueInterface.h"
 #include "PVDebugVisualizer.h"
-#include "DataTypes/PVData.h"
+
 #include "Components/InstancedStaticMeshComponent.h"
 #include "Components/TextRenderComponent.h"
+
+#include "DataTypes/PVData.h"
+
+#include "Engine/StaticMesh.h"
+
+#include "UObject/Package.h"
+
 #include "Visualizations/PVLineBatchComponent.h"
 
 #define LOCTEXT_NAMESPACE "PVDebugVisualization"
@@ -61,8 +66,6 @@ void FPVDebugVisualizationBase::DrawAttributes(FVisualizerDrawContext& InContext
 
 			FPVAttributeValuePtr AttributeValue = FPVAttributeValueInterface::Create(GroupName, AttributeName, AttributeType, Collection);
 
-			check(AttributeValue);
-			
 			if (AttributeValue)
 			{
 				if (AttributeValue->IsScalar() || VisualizationSettings.VisualizationMode == EPVDebugValueVisualizationMode::Text)
diff --git a/Engine/Plugins/Experimental/ProceduralVegetationEditor/Source/ProceduralVegetationEditor/Private/Visualizations/PVDataVisualization.cpp b/Engine/Plugins/Experimental/ProceduralVegetationEditor/Source/ProceduralVegetationEditor/Private/Visualizations/PVDataVisualization.cpp
--- a/Engine/Plugins/Experimental/ProceduralVegetationEditor/Source/ProceduralVegetationEditor/Private/Visualizations/PVDataVisualization.cpp
+++ b/Engine/Plugins/Experimental/ProceduralVegetationEditor/Source/ProceduralVegetationEditor/Private/Visualizations/PVDataVisualization.cpp
@@ -1,25 +1,34 @@
 // Copyright Epic Games, Inc. All Rights Reserved.
 
 #include "PVDataVisualization.h"
+
 #include "AdvancedPreviewScene.h"
+#include "EditorViewportClient.h"
 #include "GeometryCollectionToDynamicMesh.h"
-#include "PVBoneComponent.h"
 #include "PCGComponent.h"
 #include "PCGContext.h"
 #include "PCGDataVisualization.h"
 #include "PCGEditorSettings.h"
 #include "PlanarCut.h"
+#include "PVBoneComponent.h"
+#include "PVEditorSettings.h"
+
 #include "Components/DynamicMeshComponent.h"
-#include "Components/InstancedStaticMeshComponent.h"
 #include "Components/LineBatchComponent.h"
+#include "Components/TextRenderComponent.h"
+
+#include "DataTypes/PVData.h"
+
+#include "DebugVisualization/PVDebugVisualizer.h"
+
 #include "Facades/PVFoliageFacade.h"
+
 #include "GeometryCollection/GeometryCollectionObject.h"
-#include "GeometryCollection/Facades/CollectionRenderingFacade.h"
+
 #include "Helpers/PVUtilities.h"
-#include "DataTypes/PVData.h"
-#include "Components/TextRenderComponent.h"
-#include "DebugVisualization/PVDebugVisualizer.h"
+
 #include "Nodes/PVBaseSettings.h"
+
 #include "Visualizations/PVSkeletonVisualizerComponent.h"
 
 #define LOCTEXT_NAMESPACE "PCGManagedArrayCollectionDataVisualization"
@@ -60,6 +69,24 @@ FPCGSetupSceneFunc FPVDataVisualization::GetViewportSetupFunc(const UPCGSettings
 			return;
 		}
 			
+		InOutParams.EditorViewportClient->SetRealtime(true);
+
+		if (!InOutParams.Scene->IsGridEnabled())
+		{
+			InOutParams.Scene->HandleToggleGrid();
+		}
+		if (!InOutParams.Scene->IsEnvironmentEnabled())
+		{
+			InOutParams.Scene->HandleToggleEnvironment();
+		}
+		if (!InOutParams.Scene->IsPostProcessingEnabled())
+		{
+			InOutParams.Scene->HandleTogglePostProcessing();
+		}
+
+		InOutParams.Scene->SetFloorVisibility(true);
+		InOutParams.Scene->SetFloorOffset(1.0f);
+			
 		UPCGEditorSettings* PCGEditorSettings = GetMutableDefault<UPCGEditorSettings>();
 
 		if (PCGEditorSettings && WeakSettings.IsValid())
@@ -80,32 +107,21 @@ FPCGSetupSceneFunc FPVDataVisualization::GetViewportSetupFunc(const UPCGSettings
 			}
 		}
 
-		if (!InOutParams.Scene->IsGridEnabled())
-		{
-			InOutParams.Scene->HandleToggleGrid();
-		}
-		if (!InOutParams.Scene->IsEnvironmentEnabled())
-		{
-			InOutParams.Scene->HandleToggleEnvironment();
-		}
-		if (!InOutParams.Scene->IsPostProcessingEnabled())
-		{
-			InOutParams.Scene->HandleTogglePostProcessing();
-		}
-
-		UStaticMesh* MannequinMesh = LoadObject<UStaticMesh>(
-			nullptr,
-			TEXT("/ProceduralVegetationEditor/Mannequin/Viewport_Mannequin_T_Pose.Viewport_Mannequin_T_Pose")
-		);
-
-		constexpr bool bShowMannequin = false;
-		if (bShowMannequin)
+		if (GetDefault<UPVEditorSettings>()->bShowMannequin)
 		{
+			UStaticMesh* MannequinMesh = LoadObject<UStaticMesh>(
+				nullptr,
+				TEXT("/ProceduralVegetationEditor/Mannequin/Viewport_Mannequin_T_Pose.Viewport_Mannequin_T_Pose")
+			);
 			UStaticMeshComponent* MannequinComponent = NewObject<UStaticMeshComponent>(GetTransientPackage(), NAME_None, RF_Transient);
 			MannequinComponent->SetStaticMesh(MannequinMesh);
+			FBoxSphereBounds MannequinBounds = MannequinComponent->CalcLocalBounds();
 
 			InOutParams.ManagedResources.Add(MannequinComponent);
-			InOutParams.Scene->AddComponent(MannequinComponent, FTransform(FVector(-Bounds.BoxExtent.X * 2.0, 0, 0)));
+			InOutParams.Scene->AddComponent(
+				MannequinComponent,
+				FTransform(FVector((-Bounds.BoxExtent.X - MannequinBounds.BoxExtent.X) * 1.5f, 0, 0))
+			);
 		}
 
 		if (PV::Utilities::DebugModeEnabled())
@@ -116,9 +132,6 @@ FPCGSetupSceneFunc FPVDataVisualization::GetViewportSetupFunc(const UPCGSettings
 			}
 		}
 
-		InOutParams.Scene->SetFloorVisibility(true);
-		InOutParams.Scene->SetFloorOffset(1.0f);
-
 		InOutParams.FocusBounds = FBoxSphereBounds(Bounds.Origin, Bounds.BoxExtent, Bounds.SphereRadius);
 	};
 }
diff --git a/Engine/Plugins/Experimental/ProceduralVegetationEditor/Source/ProceduralVegetationEditor/Public/PVEditorSettings.h b/Engine/Plugins/Experimental/ProceduralVegetationEditor/Source/ProceduralVegetationEditor/Public/PVEditorSettings.h
--- a/Engine/Plugins/Experimental/ProceduralVegetationEditor/Source/ProceduralVegetationEditor/Public/PVEditorSettings.h
+++ b/Engine/Plugins/Experimental/ProceduralVegetationEditor/Source/ProceduralVegetationEditor/Public/PVEditorSettings.h
@@ -2,7 +2,7 @@
 
 #pragma once
 
-#include "Engine/DeveloperSettings.h"
+#include "CoreMinimal.h"
 
 #include "PVEditorSettings.generated.h"
 
@@ -22,4 +22,7 @@ public:
 	/** Color used for data pins of type FoliageMeshData */
 	UPROPERTY(EditAnywhere, config, Category = Node, meta = (HideAlphaChannel))
 	FLinearColor FoliageMeshDataPinColor = FColor(0, 255, 0);
+
+	UPROPERTY(EditAnywhere, config, Category = Debug)
+	bool bShowMannequin = false;
 };
