[PVE] Fix issues related to bone skeleton calculation and other minor issues

diff --git a/Engine/Plugins/Experimental/ProceduralVegetationEditor/Source/ProceduralVegetation/Private/Facades/PVBoneFacade.cpp b/Engine/Plugins/Experimental/ProceduralVegetationEditor/Source/ProceduralVegetation/Private/Facades/PVBoneFacade.cpp
--- a/Engine/Plugins/Experimental/ProceduralVegetationEditor/Source/ProceduralVegetation/Private/Facades/PVBoneFacade.cpp
+++ b/Engine/Plugins/Experimental/ProceduralVegetationEditor/Source/ProceduralVegetation/Private/Facades/PVBoneFacade.cpp
@@ -223,7 +223,7 @@ namespace PV::Facades
 		for (int32 BranchIdx = 0; BranchIdx < NumBranches; ++BranchIdx)
 		{
 			TArray<int32> ParentIndices = BranchFacade.GetParentBranchIndices(BranchIdx);
-			int32 ParentIndex = ParentIndices.Num() > 0 ? ParentIndices.Last() : -1;
+			int32 ParentIndex = BranchFacade.GetParentIndex(BranchIdx);
 
 			if (BranchRootPScales.IsValidIndex(ParentIndex) && ParentIndex > 0 )
 			{
@@ -289,12 +289,11 @@ namespace PV::Facades
 			{
 				int32 BranchIndex = SortedBranchIndices[i];
 				int32 BranchNumber = BranchFacade.GetBranchNumber(BranchIndex);
-				UE_LOG(LogProceduralVegetation, Log, TEXT("Sorted i{%i} Branch Index{%i} Branch Number{%i}"), i, BranchIndex, BranchNumber);
-			
+				
 				auto BranchPoints = BranchFacade.GetPoints(BranchIndex);
 			
 				check(BranchPoints.IsValidIndex(0));
-				int32 ParentIndex = BranchFacade.GetParentBranchIndices(BranchIndex).Last();
+				int32 ParentIndex = BranchFacade.GetParentIndex(BranchIndex);
 				TArray<FBoneNode> ParentBoneNodes = GetParentBranchBones(Bones, ParentIndex);
 			
 				FString Message;
@@ -326,9 +325,9 @@ namespace PV::Facades
 					const float PScale = PointFacade.GetPointScale(PointIdx);
 				
 					FVector NextPosition = PointPosition;
-					if (BranchPoints.IsValidIndex(Idx + 1))
+					if (Idx + 1 < BranchPoints.Num())
 					{
-						NextPosition = FVector(BranchPoints[Idx + 1]);
+						NextPosition = FVector(Positions[BranchPoints[Idx + 1]]);
 					}
 				
 					float NjordPixelID = PointFacade.GetNjordPixelIndex(PointIdx);
diff --git a/Engine/Plugins/Experimental/ProceduralVegetationEditor/Source/ProceduralVegetation/Private/Facades/PVBranchFacade.cpp b/Engine/Plugins/Experimental/ProceduralVegetationEditor/Source/ProceduralVegetation/Private/Facades/PVBranchFacade.cpp
--- a/Engine/Plugins/Experimental/ProceduralVegetationEditor/Source/ProceduralVegetation/Private/Facades/PVBranchFacade.cpp
+++ b/Engine/Plugins/Experimental/ProceduralVegetationEditor/Source/ProceduralVegetation/Private/Facades/PVBranchFacade.cpp
@@ -18,6 +18,8 @@ namespace PV::Facades
 		, BranchSimulationGroupIndex(InCollection, AttributeNames::BranchSimulationGroupIndex, GroupNames::BranchGroup)
 		, TrunkMaterialPathAttribute(InCollection, AttributeNames::TrunkMaterialPath, GroupNames::DetailsGroup)
 		, TrunkURangeAttribute(InCollection, AttributeNames::TrunkURange, GroupNames::DetailsGroup)
+		, BranchParentNumbers(InCollection, AttributeNames::BranchParentNumber, GroupNames::BranchGroup)
+		, PlantNumbers(InCollection, AttributeNames::PlantNumber, GroupNames::BranchGroup)
 	{
 	}
 
@@ -34,13 +36,21 @@ namespace PV::Facades
 		, BranchSimulationGroupIndex(InCollection, AttributeNames::BranchSimulationGroupIndex, GroupNames::BranchGroup)
 		, TrunkMaterialPathAttribute(InCollection, AttributeNames::TrunkMaterialPath, GroupNames::DetailsGroup)
 		, TrunkURangeAttribute(InCollection, AttributeNames::TrunkURange, GroupNames::DetailsGroup)
+		, BranchParentNumbers(InCollection, AttributeNames::BranchParentNumber, GroupNames::BranchGroup)
+		, PlantNumbers(InCollection, AttributeNames::PlantNumber, GroupNames::BranchGroup)
 	{
 	}
 
 	bool FBranchFacade::IsValid() const
 	{
-		return Parents.IsValid() && Children.IsValid() && BranchPoints.IsValid() && BranchNumbers.IsValid() && BranchSourceBudNumber.IsValid() &&
-			BranchHierarchyNumber.IsValid();
+		return Parents.IsValid()
+			&& Children.IsValid()
+			&& BranchPoints.IsValid()
+			&& BranchNumbers.IsValid()
+			&& BranchSourceBudNumber.IsValid()
+			&& BranchHierarchyNumber.IsValid()
+			&& BranchParentNumbers.IsValid()
+			&& PlantNumbers.IsValid();
 	}
 
 	int32 FBranchFacade::GetElementCount() const
@@ -97,6 +107,26 @@ namespace PV::Facades
 		return EmptyParents;
 	}
 
+	int32 FBranchFacade::GetParentIndex(const int32 BranchIndex) const
+	{
+		if (BranchParentNumbers.IsValid() && BranchParentNumbers.IsValidIndex(BranchIndex))
+		{
+			return BranchNumbers.Get().Find(BranchParentNumbers[BranchIndex]);
+		}
+
+		return INDEX_NONE;
+	}
+
+	int32 FBranchFacade::GetParentBranchNumber(const int32 BranchIndex) const
+	{
+		if (BranchParentNumbers.IsValid() && BranchParentNumbers.IsValidIndex(BranchIndex))
+		{
+			return BranchParentNumbers[BranchIndex];
+		}
+
+		return INDEX_NONE;
+	}
+
 	void FBranchFacade::SetParents(const int32 Index, const TArray<int32>& InParents)
 	{
 		if (Parents.IsValid() && Parents.IsValidIndex(Index))
@@ -294,6 +324,8 @@ namespace PV::Facades
 			BranchSourceBudNumber.ModifyAt(ToIndex, BranchSourceBudNumber[FromIndex]);
 			BranchFoliageIDs.ModifyAt(ToIndex, BranchFoliageIDs[FromIndex]);
 			BranchHierarchyNumber.ModifyAt(ToIndex, BranchHierarchyNumber[FromIndex]);
+			BranchParentNumbers.ModifyAt(ToIndex, BranchParentNumbers[FromIndex]);
+			PlantNumbers.ModifyAt(ToIndex, PlantNumbers[FromIndex]);
 			if (BranchUVMaterial.IsValid())
 			{
 				BranchUVMaterial.ModifyAt(ToIndex, BranchUVMaterial[FromIndex]);
diff --git a/Engine/Plugins/Experimental/ProceduralVegetationEditor/Source/ProceduralVegetation/Private/Facades/PVFoliageFacade.cpp b/Engine/Plugins/Experimental/ProceduralVegetationEditor/Source/ProceduralVegetation/Private/Facades/PVFoliageFacade.cpp
--- a/Engine/Plugins/Experimental/ProceduralVegetationEditor/Source/ProceduralVegetation/Private/Facades/PVFoliageFacade.cpp
+++ b/Engine/Plugins/Experimental/ProceduralVegetationEditor/Source/ProceduralVegetation/Private/Facades/PVFoliageFacade.cpp
@@ -219,7 +219,7 @@ namespace PV::Facades
 	void FFoliageFacade::SetFoliageIdsArray(const int32 Index, const TArray<int32>& InputIds)
 	{
 		check(!IsConst());
-		if (IsValid() && Index >= 0)
+		if (IsValid() && FoliageIdsAttribute.IsValidIndex(Index))
 		{
 			FoliageIdsAttribute.Modify()[Index] = InputIds;
 		}
diff --git a/Engine/Plugins/Experimental/ProceduralVegetationEditor/Source/ProceduralVegetation/Private/Implementations/PVCarve.cpp b/Engine/Plugins/Experimental/ProceduralVegetationEditor/Source/ProceduralVegetation/Private/Implementations/PVCarve.cpp
--- a/Engine/Plugins/Experimental/ProceduralVegetationEditor/Source/ProceduralVegetation/Private/Implementations/PVCarve.cpp
+++ b/Engine/Plugins/Experimental/ProceduralVegetationEditor/Source/ProceduralVegetation/Private/Implementations/PVCarve.cpp
@@ -49,79 +49,6 @@ void FPVCarve::UpdatePointScales(PV::Facades::FPointFacade& PointFacadeOut, cons
 	}
 }
 
-void FPVCarve::UpdateBranchLinkedIndices(PV::Facades::FBranchFacade& BranchFacadeOut,
-                                         PV::Facades::FFoliageFacade& FoliageFacadeOut, const TMap<int32, int32>& PointsOldIDsToNewIDs,
-                                         const TArray<bool>& PointsRemoved, const TMap<int32, int32>& FoliageInstancesOldIDsToNewIDs,
-                                         const TArray<bool>& FoliageInstancesRemoved)
-{
-	TRACE_CPUPROFILER_EVENT_SCOPE(PVCarve::UpdateBranchLinkedIndices);
-
-	const int32 NumOfBranches = BranchFacadeOut.GetElementCount();
-	TSet<int32> CurrentBranchNumbers;
-	CurrentBranchNumbers.Reserve(NumOfBranches);
-	for (int32 BranchIndex = 0; BranchIndex < NumOfBranches; ++BranchIndex)
-	{
-		CurrentBranchNumbers.Add(BranchFacadeOut.GetBranchNumber(BranchIndex));
-	}
-
-	for (int32 BranchIndex = 0; BranchIndex < NumOfBranches; ++BranchIndex)
-	{
-		TArray<int32> UpdatedPoints;
-		const TArray<int32>& BranchPoints = BranchFacadeOut.GetPoints(BranchIndex);
-		for (const int32& PointIndex : BranchPoints)
-		{
-			if (PointsOldIDsToNewIDs.Contains(PointIndex))
-			{
-				UpdatedPoints.Add(PointsOldIDsToNewIDs[PointIndex]);
-			}
-			else if (!PointsRemoved[PointIndex])
-			{
-				UpdatedPoints.Add(PointIndex);
-			}
-		}
-		BranchFacadeOut.SetPoints(BranchIndex, UpdatedPoints);
-
-		TArray<int32> UpdatedParents;
-		const TArray<int32>& BranchParents = BranchFacadeOut.GetParents(BranchIndex);
-		for (const int32& BranchNumber : BranchParents)
-		{
-			if (BranchNumber == 0 || CurrentBranchNumbers.Contains(BranchNumber))
-			{
-				UpdatedParents.Add(BranchNumber);
-			}
-		}
-		BranchFacadeOut.SetParents(BranchIndex, UpdatedParents);
-
-		TArray<int32> UpdatedChildren;
-		const TArray<int32>& BranchChildren = BranchFacadeOut.GetChildren(BranchIndex);
-		for (const int32& BranchNumber : BranchChildren)
-		{
-			if (CurrentBranchNumbers.Contains(BranchNumber))
-			{
-				UpdatedChildren.Add(BranchNumber);
-			}
-		}
-		BranchFacadeOut.SetChildren(BranchIndex, UpdatedChildren);
-
-		TArray<int32> UpdatedFoliageInstances;
-		const TArray<int32>& FoliageInstances = FoliageFacadeOut.GetFoliageEntryIdsForBranch(BranchIndex);
-		for (const int32& FoliageInstanceIndex : FoliageInstances)
-		{
-			if (FoliageInstancesOldIDsToNewIDs.Contains(FoliageInstanceIndex))
-			{
-				const int32 NewFoliageInstanceID = FoliageInstancesOldIDsToNewIDs[FoliageInstanceIndex];
-				UpdatedFoliageInstances.Add(NewFoliageInstanceID);
-				FoliageFacadeOut.SetFoliageBranchId(NewFoliageInstanceID, BranchIndex);
-			}
-			else if (!FoliageInstancesRemoved[FoliageInstanceIndex])
-			{
-				UpdatedFoliageInstances.Add(FoliageInstanceIndex);
-			}
-		}
-		FoliageFacadeOut.SetFoliageIdsArray(BranchIndex, UpdatedFoliageInstances);
-	}
-}
-
 void FPVCarve::RecomputeAttributes(FManagedArrayCollection& OutCollection, const FManagedArrayCollection& SourceCollection,
                                    const TMap<int32, int32>& BranchesNewIDsToOldIDs)
 {
@@ -316,7 +243,7 @@ void FPVCarve::CarveFromTop(FManagedArrayCollection& OutCollection, const FManag
 			{
 				const FVector3f& FirstPointPosition = PointFacadeSource.GetPosition(BranchPoints[0]);
 
-				const int32 ParentBranchNumber = BranchFacadeSource.GetParents(BranchIndex).Last();
+				const int32 ParentBranchNumber = BranchFacadeSource.GetParentBranchNumber(BranchIndex);
 				const int32 ParentBranchIndex = BranchNumbersToBranchIDs[ParentBranchNumber];
 				const TArray<int32>& ParentBranchPoints = BranchFacadeSource.GetPoints(ParentBranchIndex);
 
diff --git a/Engine/Plugins/Experimental/ProceduralVegetationEditor/Source/ProceduralVegetation/Private/Implementations/PVCarve.h b/Engine/Plugins/Experimental/ProceduralVegetationEditor/Source/ProceduralVegetation/Private/Implementations/PVCarve.h
--- a/Engine/Plugins/Experimental/ProceduralVegetationEditor/Source/ProceduralVegetation/Private/Implementations/PVCarve.h
+++ b/Engine/Plugins/Experimental/ProceduralVegetationEditor/Source/ProceduralVegetation/Private/Implementations/PVCarve.h
@@ -41,10 +41,6 @@ private:
 	                              const TArray<int>& BranchPoints, const float LastPointScale, const int32 LastPointIndex, const float CarveRatio,
 	                              const int32 EndIndex, const float FirstPointTargetPScale);
 
-	static void UpdateBranchLinkedIndices(PV::Facades::FBranchFacade& BranchFacadeOut, PV::Facades::FFoliageFacade& FoliageFacadeOut,
-	                                      const TMap<int32, int32>& PointsOldIDsToNewIDs, const TArray<bool>& PointsRemoved,
-	                                      const TMap<int32, int32>& FoliageInstancesOldIDsToNewIDs, const TArray<bool>& FoliageInstancesRemoved);
-
 	static void RecomputeAttributes(FManagedArrayCollection& OutCollection, const FManagedArrayCollection& SourceCollection,
 	                                const TMap<int32, int32>& BranchesNewIDsToOldIDs);
 
diff --git a/Engine/Plugins/Experimental/ProceduralVegetationEditor/Source/ProceduralVegetation/Public/Facades/PVBranchFacade.h b/Engine/Plugins/Experimental/ProceduralVegetationEditor/Source/ProceduralVegetation/Public/Facades/PVBranchFacade.h
--- a/Engine/Plugins/Experimental/ProceduralVegetationEditor/Source/ProceduralVegetation/Public/Facades/PVBranchFacade.h
+++ b/Engine/Plugins/Experimental/ProceduralVegetationEditor/Source/ProceduralVegetation/Public/Facades/PVBranchFacade.h
@@ -35,6 +35,10 @@ namespace PV::Facades
 
 		const TArray<int32>& GetParents(const int32 Index) const;
 
+		int32 GetParentIndex(const int32 BranchIndex) const;
+
+		int32 GetParentBranchNumber(const int32 BranchIndex) const;
+
 		void SetParents(const int32 Index, const TArray<int32>& InParents);
 
 		int32 GetBranchNumber(const int32 Index) const;
@@ -89,5 +93,7 @@ namespace PV::Facades
 		TManagedArrayAccessor<int32> BranchSimulationGroupIndex;
 		TManagedArrayAccessor<FString> TrunkMaterialPathAttribute;
 		TManagedArrayAccessor<TArray<FVector2f>> TrunkURangeAttribute;
+		TManagedArrayAccessor<int32> BranchParentNumbers;
+		TManagedArrayAccessor<int32> PlantNumbers;
 	};
 }
