[PVE] Remove unnecessary settings in foliage distributor node

diff --git a/Engine/Plugins/Experimental/ProceduralVegetationEditor/Source/ProceduralVegetation/Private/Implementations/PVFoliage.h b/Engine/Plugins/Experimental/ProceduralVegetationEditor/Source/ProceduralVegetation/Private/Implementations/PVFoliage.h
--- a/Engine/Plugins/Experimental/ProceduralVegetationEditor/Source/ProceduralVegetation/Private/Implementations/PVFoliage.h
+++ b/Engine/Plugins/Experimental/ProceduralVegetationEditor/Source/ProceduralVegetation/Private/Implementations/PVFoliage.h
@@ -69,9 +69,7 @@ struct FPhyllotaxySettings
 	const EPhyllotaxyFormation PhyllotaxyFormation;
 	const int32 MinimumNodeBuds;
 	const int32 MaximumNodeBuds;
-	const bool ResetPhyllotaxy;
 	const float PhyllotaxyAdditionalAngle;
-	const float PhyllotaxyOffset;
 };
 
 struct FMiscSettings
diff --git a/Engine/Plugins/Experimental/ProceduralVegetationEditor/Source/ProceduralVegetation/Private/Nodes/PVFoliageDistributorSettings.cpp b/Engine/Plugins/Experimental/ProceduralVegetationEditor/Source/ProceduralVegetation/Private/Nodes/PVFoliageDistributorSettings.cpp
--- a/Engine/Plugins/Experimental/ProceduralVegetationEditor/Source/ProceduralVegetation/Private/Nodes/PVFoliageDistributorSettings.cpp
+++ b/Engine/Plugins/Experimental/ProceduralVegetationEditor/Source/ProceduralVegetation/Private/Nodes/PVFoliageDistributorSettings.cpp
@@ -121,9 +121,7 @@ bool FPVFoliageDistributorElement::ExecuteInternal(FPCGContext* InContext) const
 				.PhyllotaxyFormation = Settings->PhyllotaxyFormation,
 				.MinimumNodeBuds = Settings->MinimumNodeBuds,
 				.MaximumNodeBuds = Settings->MaximumNodeBuds,
-				.ResetPhyllotaxy = Settings->ResetPhyllotaxy,
-				.PhyllotaxyAdditionalAngle = Settings->PhyllotaxyAdditionalAngle,
-				.PhyllotaxyOffset = Settings->PhyllotaxyOffset
+				.PhyllotaxyAdditionalAngle = Settings->PhyllotaxyAdditionalAngle
 			};
 
 			const FMiscSettings MiscSettings = {
diff --git a/Engine/Plugins/Experimental/ProceduralVegetationEditor/Source/ProceduralVegetation/Private/Nodes/PVFoliageDistributorSettings.h b/Engine/Plugins/Experimental/ProceduralVegetationEditor/Source/ProceduralVegetation/Private/Nodes/PVFoliageDistributorSettings.h
--- a/Engine/Plugins/Experimental/ProceduralVegetationEditor/Source/ProceduralVegetation/Private/Nodes/PVFoliageDistributorSettings.h
+++ b/Engine/Plugins/Experimental/ProceduralVegetationEditor/Source/ProceduralVegetation/Private/Nodes/PVFoliageDistributorSettings.h
@@ -109,15 +109,9 @@ private:
 	UPROPERTY(EditAnywhere, Category="Phyllotaxy Settings", meta=(EditCondition="OverrideDistribution", ClampMin=1, ClampMax=10, UIMin=1, UIMax=10))
 	int32 MaximumNodeBuds = 3;
 
-	UPROPERTY(EditAnywhere, Category="Phyllotaxy Settings", meta=(EditCondition="OverrideDistribution"))
-	bool ResetPhyllotaxy = false;
-
 	UPROPERTY(EditAnywhere, Category="Phyllotaxy Settings", meta=(EditCondition="OverrideDistribution", ClampMin=0.0f, ClampMax=360.0f, UIMin=0.0f, UIMax=360.0f))
 	float PhyllotaxyAdditionalAngle = 0.0f;
 
-	UPROPERTY(EditAnywhere, Category="Phyllotaxy Settings", meta=(EditCondition="OverrideDistribution", ClampMin=0.0f, ClampMax=360.0f, UIMin=0.0f, UIMax=360.0f))
-	float PhyllotaxyOffset = 0.0f;
-
 	UPROPERTY(EditAnywhere, Category="Misc Settings", meta=(EditCondition="OverrideDistribution"))
 	int32 RandomSeed = 123456;
 };
