[PVE] Fix issue with gravity node not working on certain sample assets

diff --git a/Engine/Plugins/Experimental/ProceduralVegetationEditor/Source/ProceduralVegetation/Private/Implementations/PVGravity.cpp b/Engine/Plugins/Experimental/ProceduralVegetationEditor/Source/ProceduralVegetation/Private/Implementations/PVGravity.cpp
--- a/Engine/Plugins/Experimental/ProceduralVegetationEditor/Source/ProceduralVegetation/Private/Implementations/PVGravity.cpp
+++ b/Engine/Plugins/Experimental/ProceduralVegetationEditor/Source/ProceduralVegetation/Private/Implementations/PVGravity.cpp
@@ -129,9 +129,9 @@ void FPVGravity::ApplyGravity(
 	for (const int32 Child : CurrentBranchChildren)
 	{
 		const int32 ChildIndex = BranchNumbers.Find(Child);
-		const TArray<int32>& Parents = BranchFacade.GetParents(ChildIndex);
 
-		if (ChildIndex != INDEX_NONE && !Parents.IsEmpty() && Parents[Parents.Num() - 1] == BranchNumbers[BranchIndex])
+		if (const int32 ParentIndex = BranchFacade.GetParentIndex(ChildIndex);
+			ChildIndex != INDEX_NONE && ParentIndex == BranchIndex)
 		{
 			ChildrenIndices.Add(ChildIndex);
 		}
