[PVE] Add Tooltips to node properties

diff --git a/Engine/Plugins/Experimental/ProceduralVegetationEditor/Source/ProceduralVegetation/Private/Implementations/PVCarve.h b/Engine/Plugins/Experimental/ProceduralVegetationEditor/Source/ProceduralVegetation/Private/Implementations/PVCarve.h
--- a/Engine/Plugins/Experimental/ProceduralVegetationEditor/Source/ProceduralVegetation/Private/Implementations/PVCarve.h
+++ b/Engine/Plugins/Experimental/ProceduralVegetationEditor/Source/ProceduralVegetation/Private/Implementations/PVCarve.h
@@ -24,10 +24,10 @@ struct FPVCarveParams
 {
 	GENERATED_BODY()
 
-	UPROPERTY(EditAnywhere, Category="Carve")
+	UPROPERTY(EditAnywhere, Category="Carve", meta=(Tooltip="Sets the reference used to trim the plant.\n\nChooses the basis for trimming the plant structure (e.g., length from root, bottom height, world Z, or radial distance). \nLengthFromRoot – distance measured from the plant’s root/trunk base.\n\nFrom Bottom – relative height starting from the plant’s lowest point.\n\nZPosition – absolute world-space Z axis position.\n\nRadius – Radius/Thickness of branch"))
 	ECarveBasis CarveBasis = ECarveBasis::LengthFromRoot;
 
-	UPROPERTY(EditAnywhere, Category="Carve", meta=(ClampMin=0.0f, ClampMax=1.0f, UIMin=0.0f, UIMax=1.0f))
+	UPROPERTY(EditAnywhere, Category="Carve", meta=(ClampMin=0.0f, ClampMax=1.0f, UIMin=0.0f, UIMax=1.0f, Tooltip="Amount of trimming based on the selected basis.\n\nControls how much of the structure is removed according to the chosen Carve Basis. Lower values keep more geometry; higher values trim more. Use to reveal or mask regions procedurally."))
 	float Carve = 0.0f;
 };
 
diff --git a/Engine/Plugins/Experimental/ProceduralVegetationEditor/Source/ProceduralVegetation/Private/Implementations/PVGravity.h b/Engine/Plugins/Experimental/ProceduralVegetationEditor/Source/ProceduralVegetation/Private/Implementations/PVGravity.h
--- a/Engine/Plugins/Experimental/ProceduralVegetationEditor/Source/ProceduralVegetation/Private/Implementations/PVGravity.h
+++ b/Engine/Plugins/Experimental/ProceduralVegetationEditor/Source/ProceduralVegetation/Private/Implementations/PVGravity.h
@@ -24,19 +24,19 @@ struct FPVGravityParams
 {
 	GENERATED_BODY()
 	
-	UPROPERTY(EditAnywhere, Category="Gravity")
+	UPROPERTY(EditAnywhere, Category="Gravity", meta=(ToolTip="Selects gravity behavior (bending mode).\n\nChooses how bending is applied: Gravity (downward pull) or Phototropic (growth toward light/up). Affects the direction and style of branch curvature."))
 	EGravityMode Mode = EGravityMode::Gravity;
 
-	UPROPERTY(EditAnywhere, Category="Gravity", meta=(ClampMin=0.0f, ClampMax=10.0f, UIMin=0.0f, UIMax=1.0f))
+	UPROPERTY(EditAnywhere, Category="Gravity", meta=(ClampMin=0.0f, ClampMax=10.0f, UIMin=0.0f, UIMax=1.0f, Tooltip="Strength of gravity/phototropic effect.\n\nSets how strongly branches bend under gravity or orient toward light in phototropic mode. Higher values yield more pronounced curvature."))
 	float Gravity = 0;
 
-	UPROPERTY(EditAnywhere, Category="Gravity", meta=(EditCondition = "Mode == EGravityMode::Gravity"))
+	UPROPERTY(EditAnywhere, Category="Gravity", meta=(EditCondition = "Mode == EGravityMode::Gravity", Tooltip="Direction vector used for gravity.\n\nDefines the gravity vector for bending (e.g., world down or a custom direction). Only applies in gravity mode."))
 	FVector3f Direction = FVector3f::DownVector;
 
-	UPROPERTY(EditAnywhere, Category="Gravity", meta=(ClampMin=0.0f, ClampMax=1.0f, UIMin=0.0f, UIMax=1.0f))
+	UPROPERTY(EditAnywhere, Category="Gravity", meta=(ClampMin=0.0f, ClampMax=1.0f, UIMin=0.0f, UIMax=1.0f, Tooltip="Preserves initial branch angle during bending.\n\nBias that resists over-bending by maintaining the original branching angle. Increase to keep initial structure while still responding to forces."))
 	float AngleCorrection = 0;
 
-	UPROPERTY(EditAnywhere, Category="Gravity", meta=(ClampMin=0.0f, ClampMax=1.0f, UIMin=0.0f, UIMax=1.0f, EditCondition = "Mode == EGravityMode::Phototropic"))
+	UPROPERTY(EditAnywhere, Category="Gravity", meta=(ClampMin=0.0f, ClampMax=1.0f, UIMin=0.0f, UIMax=1.0f, EditCondition = "Mode == EGravityMode::Phototropic", Tooltip="Bias toward light-facing growth.\n\nBias between optimal light direction and shadow avoidance.  0 Being Light optimal, 1 being shadow avoidance"))
 	float PhototropicBias = 0;
 };
 
diff --git a/Engine/Plugins/Experimental/ProceduralVegetationEditor/Source/ProceduralVegetation/Private/Implementations/PVMaterialSettings.h b/Engine/Plugins/Experimental/ProceduralVegetationEditor/Source/ProceduralVegetation/Private/Implementations/PVMaterialSettings.h
--- a/Engine/Plugins/Experimental/ProceduralVegetationEditor/Source/ProceduralVegetation/Private/Implementations/PVMaterialSettings.h
+++ b/Engine/Plugins/Experimental/ProceduralVegetationEditor/Source/ProceduralVegetation/Private/Implementations/PVMaterialSettings.h
@@ -47,25 +47,25 @@ struct FTrunkGenerationMaterialSetup
 	UPROPERTY(EditAnywhere, Category="Materials")
 	TObjectPtr<UMaterialInterface> Material;
 
-	UPROPERTY(EditAnywhere, Category="Materials" , meta=(DisplayName = "Y Texture Mode"))
+	UPROPERTY(EditAnywhere, Category="Materials" , meta=(DisplayName = "Y Texture Mode", Tooltip="Projection/alignment mode along Y.\n\nChooses how textures are aligned or projected along the Y axis of the mesh (e.g., to align bark grain vertically on trunk and branches)."))
 	EYTextureMode YTextureMode = EYTextureMode::Default;
 
-	UPROPERTY(EditAnywhere, Category="Materials", meta=(DisplayName = "Y Scale", EditCondition = "YTextureMode == EYTextureMode::Default"))
+	UPROPERTY(EditAnywhere, Category="Materials", meta=(DisplayName = "Y Scale", EditCondition = "YTextureMode == EYTextureMode::Default", Tooltip="Texture tiling scale along Y.\n\nAdjusts vertical texture scale. Higher values increase tiling frequency (finer detail); lower values broaden patterns."))
 	float YScale = 1.0f;
 
-	UPROPERTY(EditAnywhere, Category="Materials", meta=(DisplayName = "Y Offset", EditCondition = "YTextureMode == EYTextureMode::Default"))
+	UPROPERTY(EditAnywhere, Category="Materials", meta=(DisplayName = "Y Offset", EditCondition = "YTextureMode == EYTextureMode::Default", Tooltip="Shifts texture position along Y.\n\nMoves the material vertically without changing scale. Useful for aligning repeats or hiding seams."))
 	float YOffset = 0.0f;
 
-	UPROPERTY(EditAnywhere, Category="Materials", meta=(DisplayName = "Y Offset Random", EditCondition = "YTextureMode == EYTextureMode::Default"))
+	UPROPERTY(EditAnywhere, Category="Materials", meta=(DisplayName = "Y Offset Random", EditCondition = "YTextureMode == EYTextureMode::Default", Tooltip="Randomizes Y offset per instance.\n\nAdds per-instance variation to the vertical texture offset to avoid visible repetitions across branches."))
 	float YOffsetRandom = 0.0f;
 
-	UPROPERTY(EditAnywhere, Category="Materials")
+	UPROPERTY(EditAnywhere, Category="Materials", meta=(Tooltip="Shifts texture position along X.\n\nMoves the material horizontally across the mesh to fine-tune alignment or hide seams."))
 	float XOffset = 0.0f;
 
-	UPROPERTY(EditAnywhere, Category="Materials" , meta=(DisplayName = "X Offset Random"))
+	UPROPERTY(EditAnywhere, Category="Materials" , meta=(DisplayName = "X Offset Random", Tooltip="Randomizes X offset per instance.\n\nAdds per-instance variation to the horizontal texture offset to break pattern repetition."))
 	float XOffsetRandom = 0.0f;
 
-	UPROPERTY(EditAnywhere, Category="Materials", meta=(DisplayName = "X Range"))
+	UPROPERTY(EditAnywhere, Category="Materials", meta=(DisplayName = "X Range", Tooltip="Limits horizontal coverage range.\n\nSets the portion of the mesh affected along the X axis, useful for restricting materials (e.g., moss/lichen) to a subsection."))
 	FFloatRange URange = FFloatRange(0,1);
 	
 };
@@ -78,19 +78,19 @@ struct FPVMaterialSettings
 {
 	GENERATED_BODY()
 
-	UPROPERTY(EditAnywhere, Category="Materials")
+	UPROPERTY(EditAnywhere, Category="Materials", meta=(Tooltip="Shifts which branch generation the material targets.\n\nAdjusts the generation index used for material assignment. For example, an offset of 1 targets branches one level above the default (trunk → gen0, first branches → gen1, etc.)."))
 	int32 GenerationOffset = 0;
 
-	UPROPERTY(EditAnywhere, Category="Materials")
+	UPROPERTY(EditAnywhere, Category="Materials", meta=(Tooltip="How generation offset is interpreted.\n\nDefines handling of GenerationOffset: Clamped keeps values within valid generations; Refit rescales offsets across available generations for smoother distribution."))
 	EGenerationOffsetMethod OffsetMethod = EGenerationOffsetMethod::Clamped;
 
-	UPROPERTY(EditAnywhere, Category="Materials")
+	UPROPERTY(EditAnywhere, Category="Materials", meta=(Tooltip="How materials are selected based on number of materials available vs generations\n\nIf number of generations exceed the number of materials, Fit will fit all generations between available materials. Repeat will apply the last material entry to all generations exceeding number of materials"))
 	EMaterialDistributionMethod DistributionMethod = EMaterialDistributionMethod::Repeat;
 
-	UPROPERTY(EditAnywhere, Category="Materials" , DisplayName="Generation Materials")
+	UPROPERTY(EditAnywhere, Category="Materials" , DisplayName="Generation Materials", meta=(Tooltip="List of materials applied to the mesh.\n\nDefines the collection of materials that can be distributed by generation, age, or radius. Enables complex setups (e.g., bark, moss, twig materials by region)."))
 	TArray<FTrunkGenerationMaterialSetup> MaterialSetups;
 
-	UPROPERTY(EditAnywhere, Category="Materials")
+	UPROPERTY(EditAnywhere, Category="Materials", meta=(Tooltip="Distributes materials by generation, age, or radius.\n\nSelects the logic for assigning materials: Generation (branch hierarchy), Age (segment age), or Radius (thickness). Useful for realistic variation across the plant."))
 	EUVMaterialMode MaterialMode = EUVMaterialMode::Generation;
 
 	/*
diff --git a/Engine/Plugins/Experimental/ProceduralVegetationEditor/Source/ProceduralVegetation/Private/Implementations/PVMeshBuilder.h b/Engine/Plugins/Experimental/ProceduralVegetationEditor/Source/ProceduralVegetation/Private/Implementations/PVMeshBuilder.h
--- a/Engine/Plugins/Experimental/ProceduralVegetationEditor/Source/ProceduralVegetation/Private/Implementations/PVMeshBuilder.h
+++ b/Engine/Plugins/Experimental/ProceduralVegetationEditor/Source/ProceduralVegetation/Private/Implementations/PVMeshBuilder.h
@@ -37,63 +37,63 @@ struct FPVMeshBuilderParams
 	FPVMaterialSettings MaterialSettings;
 
 	UPROPERTY(EditAnywhere, Category="Mesh Settings", DisplayName="Point Removal",
-		meta=(ClampMin=0.0f, ClampMax=0.1f, UIMin=0.0f, UIMax=0.1f))
+		meta=(ClampMin=0.0f, ClampMax=0.1f, UIMin=0.0f, UIMax=0.1f, Tooltip="Removes small points to simplify the mesh.\n\nDefines the strength of point removal. Higher values eliminate more minor vertices or detail points to reduce complexity at the cost of fine detail."))
 	float PointRemoval = 0.0f;
 
 	UPROPERTY(EditAnywhere, Category="Mesh Settings", DisplayName="Segment Reduction",
-		meta=(ClampMin=0.0f, ClampMax=0.1f, UIMin=0.0f, UIMax=0.1f))
+		meta=(ClampMin=0.0f, ClampMax=0.1f, UIMin=0.0f, UIMax=0.1f, Tooltip="Controls how aggressively segments are reduced.\n\nSets the strength of segment merging/decimation. Higher values reduce edge count for lighter meshes; lower values preserve detail."))
 	float SegmentReduction = 0.0f;
 
 	UPROPERTY(EditAnywhere, Category="Mesh Settings", DisplayName="Segment Retention Impact",
-		meta=(ClampMin=0.0f, ClampMax=2.0f, UIMin=0.0f, UIMax=2.0f))
+		meta=(ClampMin=0.0f, ClampMax=2.0f, UIMin=0.0f, UIMax=2.0f, Tooltip="Control the impact of reduced segments based on other criteria\n\nSegment retention impact allows you to reintroduce points that have been reduced based on user defined criteria\'s. See; Hull Retention, Ground Retention and Trunk Retention."))
 	float SegmentRetentionImpact = 1.0f;
 
 	UPROPERTY(EditAnywhere, Category="Mesh Settings", DisplayName="Longest Segment Allowed",
-		meta=(ClampMin=0.0f, ClampMax=50.0f, UIMin=0.0f, UIMax=10.0f))
+		meta=(ClampMin=0.0f, ClampMax=50.0f, UIMin=0.0f, UIMax=10.0f, Tooltip="Maximum allowed segment length.\n\nCaps segment length after simplification to maintain even geometry and prevent stretched polygons. Specified in mesh units."))
 	float LongestSegmentLength = 0.1f;
 
 	UPROPERTY(EditAnywhere, Category="Mesh Settings", DisplayName="Shortest Segment Allowed",
-		meta=(ClampMin=0.0f, ClampMax=1.0f, UIMin=0.0f, UIMax=1.0f))
+		meta=(ClampMin=0.0f, ClampMax=1.0f, UIMin=0.0f, UIMax=1.0f, Tooltip="Minimum allowed segment length.\n\nSegments shorter than this are merged or removed to avoid unnecessary fine detail that inflates polycount."))
 	float ShortestSegmentLength = 0.0f;
 
 	UPROPERTY(EditAnywhere, Category="Mesh Settings", DisplayName="Min Divisions",
-		meta=(ClampMin=3, ClampMax=1024, UIMin=3, UIMax=36))
+		meta=(ClampMin=3, ClampMax=1024, UIMin=3, UIMax=36, Tooltip="Minimum radial divisions.\n\nThe lowest polygon count around cross-sections. Higher values keep sections smooth; lower values produce blockier geometry with better performance."))
 	int32 MinMeshDivisions = 6;
 
 	UPROPERTY(EditAnywhere, Category="Mesh Settings", DisplayName="Max Divisions",
-		meta=(ClampMin=3, ClampMax=1024, UIMin=3, UIMax=36))
+		meta=(ClampMin=3, ClampMax=1024, UIMin=3, UIMax=36, Tooltip="Maximum radial divisions.\n\nUpper limit for cross-section detail on trunks/branches. Balance with Min Divisions to control density across sizes."))
 	int32 MaxMeshDivisions = 12;
 
 	UPROPERTY(EditAnywhere, Category="Mesh Settings|Advanced", DisplayName="Reduction Accuracy",
-		meta=(ClampMin=3, ClampMax=15, UIMin=1, UIMax=15))
+		meta=(ClampMin=3, ClampMax=15, UIMin=1, UIMax=15, Tooltip="Accuracy level for mesh reduction.\n\nControls how closely the simplified mesh follows the original. It specifies number of iterations that are used to evaluate the curve reduction path. The more number of iterations the more accurate the reduction will be at the cost of processing power."))
 	int32 Accuracy = 5;
 
 	UPROPERTY(EditAnywhere, Category="Mesh Settings|Advanced", DisplayName="Hull Retention",
-		meta=(ClampMin=0.0f, ClampMax=1.0f, UIMin=0.0f, UIMax=1.0f))
+		meta=(ClampMin=0.0f, ClampMax=1.0f, UIMin=0.0f, UIMax=1.0f, Tooltip="Strength of preserving the outer hull.\n\nAdjusts retention of the mesh silhouette during reduction. Higher values protect external shape; lower values allow more simplification."))
 	float HullRetention = 0.0f;
 
 	UPROPERTY(EditAnywhere, Category="Mesh Settings|Advanced", DisplayName="Main Trunk Retention",
-		meta=(ClampMin=0.0f, ClampMax=1.0f, UIMin=0.0f, UIMax=1.0f))
+		meta=(ClampMin=0.0f, ClampMax=1.0f, UIMin=0.0f, UIMax=1.0f, Tooltip="Preserves detail in the main trunk.\n\nControls how much trunk geometry is retained during simplification. High values keep trunk smooth and detailed; low values reduce complexity."))
 	float MainTrunkRetention = 0.0f;
 
 	UPROPERTY(EditAnywhere, Category="Mesh Settings|Advanced", DisplayName="Ground Retention",
-		meta=(ClampMin=0.0f, ClampMax=1.0f, UIMin=0.0f, UIMax=1.0f))
+		meta=(ClampMin=0.0f, ClampMax=1.0f, UIMin=0.0f, UIMax=1.0f, Tooltip="Preserves detail near ground regions.\n\nAdjusts retention for areas close to the base of the plant. Use higher values to keep root/ground transitions detailed."))
 	float GroundRetention = 0.0f;
 
-	UPROPERTY(EditAnywhere, Category="Mesh Settings|Advanced", DisplayName="Hull Retention Gradient", meta=(XAxisMin=0.0f, XAxisMax=1.0f, YAxisMin=0.0f, YAxisMax=1.0f))
+	UPROPERTY(EditAnywhere, Category="Mesh Settings|Advanced", DisplayName="Hull Retention Gradient", meta=(XAxisMin=0.0f, XAxisMax=1.0f, YAxisMin=0.0f, YAxisMax=1.0f, Tooltip="Gradient for preserving outer hull detail.\n\nA ramp that varies outer-hull retention across the model. Use to keep silhouettes crisp where needed while simplifying elsewhere."))
 	FPVFloatRamp HullRetentionGradient;
 
-	UPROPERTY(EditAnywhere, Category="Mesh Settings|Advanced", DisplayName="Main Trunk Retention Gradient", meta=(XAxisMin=0.0f, XAxisMax=1.0f, YAxisMin=0.0f, YAxisMax=1.0f))
+	UPROPERTY(EditAnywhere, Category="Mesh Settings|Advanced", DisplayName="Main Trunk Retention Gradient", meta=(XAxisMin=0.0f, XAxisMax=1.0f, YAxisMin=0.0f, YAxisMax=1.0f, Tooltip="Gradient controlling trunk detail retention.\n\nVaries trunk retention via a ramp. Useful to preserve key trunk regions while simplifying others."))
 	FPVFloatRamp MainTrunkRetentionGradient;
 
-	UPROPERTY(EditAnywhere, Category="Mesh Settings|Advanced", DisplayName="Ground Retention Gradient", meta=(XAxisMin=0.0f, XAxisMax=1.0f, YAxisMin=0.0f, YAxisMax=1.0f))
+	UPROPERTY(EditAnywhere, Category="Mesh Settings|Advanced", DisplayName="Ground Retention Gradient", meta=(XAxisMin=0.0f, XAxisMax=1.0f, YAxisMin=0.0f, YAxisMax=1.0f, Tooltip="Gradient for preserving ground-region detail.\n\nA ramp that modulates retention around base regions. Tune to keep roots/basal features while allowing simplification above."))
 	FPVFloatRamp GroundRetentionGradient;
 
 	UPROPERTY(EditAnywhere, Category="Mesh Settings|Advanced", DisplayName="Scale Retention",
-		meta=(ClampMin=0.0f, ClampMax=1.0f, UIMin=0.0f, UIMax=1.0f))
+		meta=(ClampMin=0.0f, ClampMax=1.0f, UIMin=0.0f, UIMax=1.0f, Tooltip="Preserves detail relative to branch size.\n\nRetains more detail on larger/thicker structures and simplifies smaller ones to prioritize visually important geometry."))
 	float ScaleRetention = 0.0f;
 
-	UPROPERTY(EditAnywhere, Category="Mesh Settings|Advanced", DisplayName="Scale Retention Gradient", meta=(XAxisMin=0.0f, XAxisMax=1.0f, YAxisMin=0.0f, YAxisMax=1.0f))
+	UPROPERTY(EditAnywhere, Category="Mesh Settings|Advanced", DisplayName="Scale Retention Gradient", meta=(XAxisMin=0.0f, XAxisMax=1.0f, YAxisMin=0.0f, YAxisMax=1.0f, Tooltip="Gradient controlling size-based retention.\n\nUses a ramp to vary how size-dependent retention is applied across the model. Keep large structures crisp while simplifying finer parts."))
 	FPVFloatRamp ScaleRetentionGradient;
 
 	UPROPERTY(EditAnywhere, Category="Plant Profile Settings", DisplayName="Profile",
@@ -119,25 +119,25 @@ struct FPVMeshBuilderParams
 		meta=(EditCondition="bIsPlantProfileDropdownEnabled", EditConditionHides, HideEditConditionToggle))
 	bool bApplyProfileToBranches = false;
 
-	UPROPERTY(EditAnywhere, Category = "Displacement Settings")
+	UPROPERTY(EditAnywhere, Category = "Displacement Settings", meta=(Tooltip="Texture used to drive mesh displacement.\n\nAssigns the texture map that controls displacement. Bright and dark areas in the texture push or pull the mesh surface to create additional geometric detail such as bark roughness, grooves, or ridges."))
 	TObjectPtr<UTexture2D> DisplacementTexture = nullptr;
 
 	UPROPERTY(Transient, NonTransactional)
 	TArray<float> DisplacementValues;
 
 	UPROPERTY(EditAnywhere, Category="Displacement Settings", DisplayName="Displacement Scale",
-		meta=(ClampMin=0.0f, ClampMax=10.0f, UIMin=0.0f, UIMax=1.0f))
+		meta=(ClampMin=0.0f, ClampMax=10.0f, UIMin=0.0f, UIMax=1.0f, Tooltip="Intensity of displacement effect.\n\nSets the overall strength of the displacement. Higher values exaggerate surface depth and detail, while lower values apply subtle variation. Use to balance realism with mesh stability."))
 	float DisplacementStrength = 0.5f;
 
 	UPROPERTY(EditAnywhere, Category="Displacement Settings", DisplayName="Displacement Bias",
-		meta=(ClampMin=0.0f, ClampMax=10.0f, UIMin=0.0f, UIMax=10.0f))
+		meta=(ClampMin=0.0f, ClampMax=10.0f, UIMin=0.0f, UIMax=10.0f, Tooltip="Shifts the midpoint of displacement.\n\nOffsets the neutral level of the displacement map. Positive bias lifts the surface outward. Useful for correcting textures that displace unevenly or need centering."))
 	float DisplacementBias = 0.0f;
 
-	UPROPERTY(EditAnywhere, Category="Displacement Settings", DisplayName="Displacement UV Scale")
+	UPROPERTY(EditAnywhere, Category="Displacement Settings", DisplayName="Displacement UV Scale", meta=(Tooltip="Controls tiling of displacement texture.\n\nAdjusts the repetition of the displacement texture across the mesh in U and V directions. Higher values increase tiling frequency (smaller details repeated more often), while lower values stretch the texture over a larger area."))
 	FIntVector2 DisplacementUVScale = FIntVector2(1.0f, 1.0f);
 
 	UPROPERTY(EditAnywhere, Category="Displacement Settings", DisplayName="Apply Displacement up to generation",
-		meta=(ClampMin=1, ClampMax=10, UIMin=1, UIMax=10))
+		meta=(ClampMin=1, ClampMax=10, UIMin=1, UIMax=10, Tooltip="Controls which generations the displacement is applied on \n\nControls which generations the displacement is applied on "))
 	int32 DisplacementGenerationUpperLimit = 1;
 
 	FPVMeshBuilderParams()
diff --git a/Engine/Plugins/Experimental/ProceduralVegetationEditor/Source/ProceduralVegetation/Private/Nodes/PVPresetLoaderSettings.h b/Engine/Plugins/Experimental/ProceduralVegetationEditor/Source/ProceduralVegetation/Private/Nodes/PVPresetLoaderSettings.h
--- a/Engine/Plugins/Experimental/ProceduralVegetationEditor/Source/ProceduralVegetation/Private/Nodes/PVPresetLoaderSettings.h
+++ b/Engine/Plugins/Experimental/ProceduralVegetationEditor/Source/ProceduralVegetation/Private/Nodes/PVPresetLoaderSettings.h
@@ -27,7 +27,7 @@ protected:
 	virtual FPCGElementPtr CreateElement() const override;
 
 public:
-	UPROPERTY(EditAnywhere, Category="Preset")
+	UPROPERTY(EditAnywhere, Category="Preset", meta=(Tooltip="Loads the selected procedural vegetation preset\n\nThis allows the user to load a procedural vegetation preset. A procedural vegatation preset contains data for a specie of vegetaion. It contains hormonal information, the main structure of the vegetation, data about the branches, where foliage will/can be attached, references to foliage meshes, materials to be applied"))
 	TObjectPtr<class UProceduralVegetationPreset> Preset = nullptr;
 };
 
diff --git a/Engine/Plugins/Experimental/ProceduralVegetationEditor/Source/ProceduralVegetation/Private/Nodes/PVRemoveBranchesSettings.h b/Engine/Plugins/Experimental/ProceduralVegetationEditor/Source/ProceduralVegetation/Private/Nodes/PVRemoveBranchesSettings.h
--- a/Engine/Plugins/Experimental/ProceduralVegetationEditor/Source/ProceduralVegetation/Private/Nodes/PVRemoveBranchesSettings.h
+++ b/Engine/Plugins/Experimental/ProceduralVegetationEditor/Source/ProceduralVegetation/Private/Nodes/PVRemoveBranchesSettings.h
@@ -28,10 +28,10 @@ protected:
 	virtual FPCGElementPtr CreateElement() const override;
 
 public:
-	UPROPERTY(EditAnywhere, Category = Settings)
+	UPROPERTY(EditAnywhere, Category = Settings, meta=(Tooltip="Sets the rule used to remove branches.\n\nChooses the basis for pruning (e.g., generation, length from root, age, or radius). Determines which branches are candidates for removal."))
 	ERemoveBranchesBasis BranchRemoveBasis = ERemoveBranchesBasis::Length;
 
-	UPROPERTY(EditAnywhere, Category = Settings, meta=(ClampMin = 0f, ClampMax = 10.0f, UIMin = 0f, UIMax = 1.0f))
+	UPROPERTY(EditAnywhere, Category = Settings, meta=(ClampMin = 0f, ClampMax = 10.0f, UIMin = 0f, UIMax = 1.0f, Tooltip="Removal threshold for the chosen basis.\n\nBranches exceeding or falling below this threshold (per the selected basis) are removed. Use lower/higher values to prune sparsely or aggressively."))
 	float Threshold = 0;
 };
 
diff --git a/Engine/Plugins/Experimental/ProceduralVegetationEditor/Source/ProceduralVegetation/Private/Nodes/PVScaleSettings.h b/Engine/Plugins/Experimental/ProceduralVegetationEditor/Source/ProceduralVegetation/Private/Nodes/PVScaleSettings.h
--- a/Engine/Plugins/Experimental/ProceduralVegetationEditor/Source/ProceduralVegetation/Private/Nodes/PVScaleSettings.h
+++ b/Engine/Plugins/Experimental/ProceduralVegetationEditor/Source/ProceduralVegetation/Private/Nodes/PVScaleSettings.h
@@ -27,7 +27,7 @@ protected:
 	virtual FPCGElementPtr CreateElement() const override;
 
 public:
-	UPROPERTY(EditAnywhere, Category = Settings, meta=(ClampMin = 0.01f, ClampMax = 100.0f, UIMin = 0.01f, UIMax = 10.0f))
+	UPROPERTY(EditAnywhere, Category = Settings, meta=(ClampMin = 0.01f, ClampMax = 100.0f, UIMin = 0.01f, UIMax = 10.0f, Tooltip="Controls overall size of generated elements.\n\nUniform size multiplier for branches, leaves, or sub-structures. Combine with randomness or ramps for natural variation."))
 	float Scale = 1;
 
 };
