[PVE] Material issue fixed where Branches get disappeared on chaning the offset or changing the Material Mode to radius.

diff --git a/Engine/Plugins/Experimental/ProceduralVegetationEditor/Source/ProceduralVegetation/Private/Implementations/PVMeshBuilder.cpp b/Engine/Plugins/Experimental/ProceduralVegetationEditor/Source/ProceduralVegetation/Private/Implementations/PVMeshBuilder.cpp
--- a/Engine/Plugins/Experimental/ProceduralVegetationEditor/Source/ProceduralVegetation/Private/Implementations/PVMeshBuilder.cpp
+++ b/Engine/Plugins/Experimental/ProceduralVegetationEditor/Source/ProceduralVegetation/Private/Implementations/PVMeshBuilder.cpp
@@ -56,9 +56,12 @@ void FPVMeshBuilder::GenerateGeometryCollection(const FManagedArrayCollection& I
 	TManagedArray<FVector3f>& Normals = OutGeometryCollection.Normal;
 	TManagedArray<FVector2f>& UVs = *OutGeometryCollection.FindUVLayer(0);
 	TManagedArray<FIntVector>& Indices = OutGeometryCollection.Indices;
-	TManagedArray<int32>& MaterialIndex = OutGeometryCollection.MaterialIndex;
 	TManagedArray<int32>& MaterialIDAttribute = OutGeometryCollection.MaterialID;
 
+	const FName MaterialPathAttributeName("MaterialPath");
+	auto& MaterialArray = OutGeometryCollection.AddAttribute<FString>(MaterialPathAttributeName, FGeometryCollection::MaterialGroup);
+	auto& MaterialSetups = MeshBuilderParams.MaterialSettings.MaterialSetups;
+	
 	struct FSection
 	{
 		int32 MaterialID;
@@ -72,20 +75,46 @@ void FPVMeshBuilder::GenerateGeometryCollection(const FManagedArrayCollection& I
 	int32 FaceIndex = 0;
 	int32 PrevMaterialID = -1;
 	TArray<FSection> Sections;
-
+	
 	//Sorting by material id so sections are created according to material ids
 	MeshesDataCollection.Sort([](const FLocalDynamicMeshData& A, const FLocalDynamicMeshData& B)
-		{
-			return A.MaterialID > B.MaterialID;
-		});
+	{
+		return A.MaterialID > B.MaterialID;
+	});
 
+	int32 SectionId = -1;
 	for (auto MeshIndex = MeshesDataCollection.Num() - 1; MeshIndex >= 0; --MeshIndex)
 	{
 		int MaterialID = MeshesDataCollection[MeshIndex].MaterialID;
 		if (PrevMaterialID != MaterialID)
 		{
 			PrevMaterialID = MaterialID;
+			
+			TObjectPtr<UMaterialInterface> SectionMaterial;
+
+			if (MaterialSetups.IsValidIndex(MaterialID))
+			{
+				SectionMaterial = MaterialSetups[MaterialID].Material;
+			}
+			else
+			{
+				UE_LOG(LogProceduralVegetation, Warning, TEXT("Unassigned Material for section %i, MaterialID %i not found in MaterialSetups."), SectionId, MaterialID);
+			}
+
+			FString MaterialPathName = "";
+			if (SectionMaterial)
+			{
+				MaterialPathName = SectionMaterial->GetPathName();
+			}
+			else
+			{
+				UE_LOG(LogProceduralVegetation, Warning, TEXT("Null Material for section %i, Assign Material slot no %i in MaterialSetups."), SectionId, MaterialID);
+			}
+
+			int Element = OutGeometryCollection.AddElements(1, FGeometryCollection::MaterialGroup);
+			MaterialArray[Element] = SectionMaterial.GetPathName();
 			Sections.Add({MaterialID, FaceIndex, 0, VertexIndex, 0});
+			SectionId++;
 		}
 
 		const FLocalDynamicMeshData& MeshesData = MeshesDataCollection[MeshIndex];
@@ -93,8 +122,7 @@ void FPVMeshBuilder::GenerateGeometryCollection(const FManagedArrayCollection& I
 		{
 			const FIntVector4 OffsetTriangle = Triangle + FIntVector4(VertexIndex);
 			Indices[FaceIndex] = FIntVector(OffsetTriangle);
-			MaterialIndex[FaceIndex] = OffsetTriangle.W;
-			MaterialIDAttribute[FaceIndex] = MaterialID;
+			MaterialIDAttribute[FaceIndex] = SectionId;
 			++FaceIndex;
 		}
 		for (const FLocalDynamicMeshData::FVertex& Vertex : MeshesData.Vertices)
@@ -129,23 +157,10 @@ void FPVMeshBuilder::GenerateGeometryCollection(const FManagedArrayCollection& I
 	OutGeometryCollection.TangentU.Fill(FVector3f::ForwardVector);
 	OutGeometryCollection.TangentV.Fill(FVector3f::RightVector);
 	OutGeometryCollection.Color.Fill(FLinearColor::White);
-
-	auto& MaterialSetups = MeshBuilderParams.MaterialSettings.MaterialSetups;
-	OutGeometryCollection.AddElements(Sections.Num(), FGeometryCollection::MaterialGroup);
-	auto& MaterialArray = OutGeometryCollection.AddAttribute<FString>("MaterialPath", FGeometryCollection::MaterialGroup);
-
+	
 	int32 SectionIndex = 0;
 	for (auto Section : Sections)
 	{
-		TObjectPtr<UMaterialInterface> SectionMaterial;
-		if (MaterialSetups.IsValidIndex(Section.MaterialID))
-		{
-			SectionMaterial = MaterialSetups[Section.MaterialID].Material;
-		}
-		MaterialArray[SectionIndex] = SectionMaterial
-			? SectionMaterial->GetPathName()
-			: "";
-
 		TManagedArray<FGeometryCollectionSection>& GeometrySections = OutGeometryCollection.Sections;
 
 		GeometrySections[SectionIndex].MaterialID = Section.MaterialID;
