[PVE] Material not getting assigned in material settings only default materials provided are getting used in material settings

diff --git a/Engine/Plugins/Experimental/ProceduralVegetationEditor/Source/ProceduralVegetation/Private/Implementations/PVMaterialSettings.h b/Engine/Plugins/Experimental/ProceduralVegetationEditor/Source/ProceduralVegetation/Private/Implementations/PVMaterialSettings.h
--- a/Engine/Plugins/Experimental/ProceduralVegetationEditor/Source/ProceduralVegetation/Private/Implementations/PVMaterialSettings.h
+++ b/Engine/Plugins/Experimental/ProceduralVegetationEditor/Source/ProceduralVegetation/Private/Implementations/PVMaterialSettings.h
@@ -78,7 +78,7 @@ struct FPVMaterialSettings
 {
 	GENERATED_BODY()
 
-	UPROPERTY(EditAnywhere, Category="Materials", meta=(Tooltip="Shifts which branch generation the material targets.\n\nAdjusts the generation index used for material assignment. For example, an offset of 1 targets branches one level above the default (trunk → gen0, first branches → gen1, etc.)."))
+	UPROPERTY(EditAnywhere, Category="Materials", meta=(ClampMin = "0", Tooltip="Shifts which branch generation the material targets.\n\nAdjusts the generation index used for material assignment. For example, an offset of 1 targets branches one level above the default (trunk → gen0, first branches → gen1, etc.)."))
 	int32 GenerationOffset = 0;
 
 	UPROPERTY(EditAnywhere, Category="Materials", meta=(Tooltip="How generation offset is interpreted.\n\nDefines handling of GenerationOffset: Clamped keeps values within valid generations; Refit rescales offsets across available generations for smoother distribution."))
diff --git a/Engine/Plugins/Experimental/ProceduralVegetationEditor/Source/ProceduralVegetation/Private/Implementations/PVMeshBuilder.h b/Engine/Plugins/Experimental/ProceduralVegetationEditor/Source/ProceduralVegetation/Private/Implementations/PVMeshBuilder.h
--- a/Engine/Plugins/Experimental/ProceduralVegetationEditor/Source/ProceduralVegetation/Private/Implementations/PVMeshBuilder.h
+++ b/Engine/Plugins/Experimental/ProceduralVegetationEditor/Source/ProceduralVegetation/Private/Implementations/PVMeshBuilder.h
@@ -30,10 +30,10 @@ struct FPVMeshBuilderParams
 {
 	GENERATED_BODY()
 
-	UPROPERTY(EditAnywhere, Category="Materials", meta=(PinHiddenByDefault, InlineEditConditionToggle))
+	UPROPERTY(EditAnywhere, Category="Materials", meta=(PinHiddenByDefault, InlineEditConditionToggle, Tooltip="Toggles between default and custom material settings.\n\nDefault material setting is defined in Procedural Vegetation Preset.\n\nCustom material settings allow for editing."))
 	bool bOverrideMaterial = false;
 
-	UPROPERTY(EditAnywhere, Category="Materials", meta = (ShowOnlyInnerProperties))
+	UPROPERTY(EditAnywhere, Category="Materials", meta = (EditCondition="bOverrideMaterial"))
 	FPVMaterialSettings MaterialSettings;
 
 	UPROPERTY(EditAnywhere, Category="Mesh Settings", DisplayName="Point Removal",
diff --git a/Engine/Plugins/Experimental/ProceduralVegetationEditor/Source/ProceduralVegetation/Private/Nodes/PVMeshBuilderSettings.cpp b/Engine/Plugins/Experimental/ProceduralVegetationEditor/Source/ProceduralVegetation/Private/Nodes/PVMeshBuilderSettings.cpp
--- a/Engine/Plugins/Experimental/ProceduralVegetationEditor/Source/ProceduralVegetation/Private/Nodes/PVMeshBuilderSettings.cpp
+++ b/Engine/Plugins/Experimental/ProceduralVegetationEditor/Source/ProceduralVegetation/Private/Nodes/PVMeshBuilderSettings.cpp
@@ -49,6 +49,18 @@ void UPVMeshBuilderSettings::PostEditChangeProperty(FPropertyChangedEvent& Prope
 
 		Modify();
 	}
+	if (PropertyChangedEvent.Property
+		&& PropertyChangedEvent.Property->GetFName() == GET_MEMBER_NAME_CHECKED(FPVMeshBuilderParams, MaterialSettings))
+	{
+		if (MesherSettings.bOverrideMaterial)
+		{
+			MesherSettings.MaterialSettings = CustomMaterialSettings;
+		}
+		else
+		{
+			MesherSettings.MaterialSettings = DefaultMaterialSettings;
+		}
+	}
 }
 
 void UPVMeshBuilderSettings::PostLoad()
@@ -104,7 +116,7 @@ bool FPVMeshBuilderElement::ExecuteInternal(FPCGContext* InContext) const
 			const FManagedArrayCollection& SkeletonCollection = InputData->GetCollection();
 			auto& MaterialSettings = MutableSettings->MesherSettings.MaterialSettings;
 					
-			if (!Settings->MesherSettings.bOverrideMaterial || MaterialSettings.MaterialSetups.Num() == 0)
+			if (!Settings->MesherSettings.bOverrideMaterial )
 			{
 				const PV::Facades::FBranchFacade BranchFacade(SkeletonCollection);
 				const FString TrunkMaterial = BranchFacade.GetTrunkMaterialPath();
@@ -113,16 +125,21 @@ bool FPVMeshBuilderElement::ExecuteInternal(FPCGContext* InContext) const
 				int32 Generation = 0;
 				for (const FVector2f& URange : TrunkURanges)
 				{
-					if (!TrunkMaterial.IsEmpty())
-					{
-						MaterialSettings.SetMaterial(TrunkMaterial, URange, Generation);
-						Generation++;
-					}
-					else
-					{
-						MutableSettings->MesherSettings.bOverrideMaterial = true;
-					}
+					MaterialSettings.SetMaterial(TrunkMaterial, URange, Generation);
+					Generation++;
 				}
+
+				MutableSettings->DefaultMaterialSettings = MaterialSettings;
+
+				if (!MutableSettings->bCustomMaterialSet)
+				{
+					MutableSettings->CustomMaterialSettings = MaterialSettings;
+					MutableSettings->bCustomMaterialSet = true;
+				}
+			}
+			else
+			{
+				MutableSettings->CustomMaterialSettings = MaterialSettings;
 			}
 
 			if (const PV::Facades::FPlantProfileFacade PlantProfileFacade = PV::Facades::FPlantProfileFacade(SkeletonCollection);
diff --git a/Engine/Plugins/Experimental/ProceduralVegetationEditor/Source/ProceduralVegetation/Private/Nodes/PVMeshBuilderSettings.h b/Engine/Plugins/Experimental/ProceduralVegetationEditor/Source/ProceduralVegetation/Private/Nodes/PVMeshBuilderSettings.h
--- a/Engine/Plugins/Experimental/ProceduralVegetationEditor/Source/ProceduralVegetation/Private/Nodes/PVMeshBuilderSettings.h
+++ b/Engine/Plugins/Experimental/ProceduralVegetationEditor/Source/ProceduralVegetation/Private/Nodes/PVMeshBuilderSettings.h
@@ -3,6 +3,7 @@
 
 #include "CoreMinimal.h"
 #include "Implementations/PVMeshBuilder.h"
+#include "Implementations/PVMaterialSettings.h"
 #include "Nodes/PVBaseSettings.h"
 #include "PVMeshBuilderSettings.generated.h"
 
@@ -36,11 +37,21 @@ public:
 	UPROPERTY(EditAnywhere, Category = "Mesh", meta = (ShowOnlyInnerProperties))
 	FPVMeshBuilderParams MesherSettings;
 
+	UPROPERTY()
+	FPVMaterialSettings CustomMaterialSettings;
+
+	UPROPERTY()
+	FPVMaterialSettings DefaultMaterialSettings;
+
+	UPROPERTY()
+	bool bCustomMaterialSet = false;
+	
 	UFUNCTION()
 	TArray<FString> GetPlantProfileOptions() const
 	{
 		return MesherSettings.PlantProfileOptions;
 	}
+	
 };
 
 class FPVMeshBuilderElement : public FPVBaseElement
