Fix crash during PVE Undo due to Procedural Vegetation Graph not being flagged as RF_Transactional

diff --git a/Engine/Plugins/Experimental/ProceduralVegetationEditor/Source/ProceduralVegetation/Private/ProceduralVegetation.cpp b/Engine/Plugins/Experimental/ProceduralVegetationEditor/Source/ProceduralVegetation/Private/ProceduralVegetation.cpp
new file mode 100644
--- /dev/null
+++ b/Engine/Plugins/Experimental/ProceduralVegetationEditor/Source/ProceduralVegetation/Private/ProceduralVegetation.cpp
@@ -0,0 +1,101 @@
+// Copyright Epic Games, Inc. All Rights Reserved.
+
+#include "ProceduralVegetation.h"
+#include "Nodes/PVFoliagePaletteSettings.h"
+#include "Nodes/PVGravitySettings.h"
+#include "Nodes/PVMeshBuilderSettings.h"
+#include "Nodes/PVPresetLoaderSettings.h"
+#include "Nodes/PVOutputSettings.h"
+#include "ProceduralVegetationPreset.h"
+
+namespace PV
+{
+	void SanitizeFolderName(FString& FolderName)
+	{
+		const TCHAR* InvalidChar = INVALID_OBJECTPATH_CHARACTERS INVALID_LONGPACKAGE_CHARACTERS;
+
+		while (*InvalidChar)
+		{
+			FolderName.ReplaceCharInline(*InvalidChar, TCHAR('_'), ESearchCase::CaseSensitive);
+			++InvalidChar;
+		}
+
+		FolderName.ReplaceCharInline('/', TCHAR('_'), ESearchCase::CaseSensitive);
+	}
+
+	const static FName DefaultGraphName = TEXT("ProceduralVegetationGraph");
+};
+
+void UProceduralVegetationGraph::PostLoad()
+{
+	Super::PostLoad();
+
+#if WITH_EDITOR
+	if (HasAllFlags(RF_Transactional) == false)
+	{
+		SetFlags(RF_Transactional);
+	}
+#endif
+}
+
+void UProceduralVegetation::CreateGraph(const UProceduralVegetationGraph* InGraph /*= nullptr*/)
+{
+	if (InGraph)
+	{
+		Graph = DuplicateObject(InGraph, this);
+		Graph->SetFlags(RF_Transactional);
+	}
+	else
+	{
+		Graph = NewObject<UProceduralVegetationGraph>(this, PV::DefaultGraphName, RF_Transactional);
+	}
+}
+
+#if WITH_EDITOR
+void UProceduralVegetation::CreateGraphFromPreset(const TObjectPtr<UProceduralVegetationPreset> InPreset)
+{
+	Graph = NewObject<UProceduralVegetationGraph>(this, PV::DefaultGraphName, RF_Transactional);
+
+	UPCGSettings* DefaultNodeSettings = nullptr;
+
+	UPVPresetLoaderSettings* PresetLoaderSettings = NewObject<UPVPresetLoaderSettings>(Graph);
+	PresetLoaderSettings->Preset = InPreset;
+	UPCGNode* PresetLoaderNode = Graph->AddNodeCopy(PresetLoaderSettings, DefaultNodeSettings);
+
+	FInt32Vector2 Position(0, 300);
+
+	PresetLoaderNode->SetNodePosition(Position.X, Position.Y);
+	Position.X += 100;
+
+	for (const auto& [VariantName, VariantData] : InPreset->Variants)
+	{
+		const UPVGravitySettings* GravitySettings = NewObject<UPVGravitySettings>(Graph);
+		const UPVMeshBuilderSettings* MeshBuilderSettings = NewObject<UPVMeshBuilderSettings>(Graph);
+		const UPVFoliagePaletteSettings* FoliagePaletteSettings = NewObject<UPVFoliagePaletteSettings>(Graph);
+		UPVOutputSettings* OutputSettings = NewObject<UPVOutputSettings>(Graph);
+		FString MeshName = *VariantName;
+		PV::SanitizeFolderName(MeshName);
+		OutputSettings->ExportSettings.MeshName = *MeshName;
+
+		UPCGNode* GravityNode = Graph->AddNodeCopy(GravitySettings, DefaultNodeSettings);
+		UPCGNode* MeshBuilderNode = Graph->AddNodeCopy(MeshBuilderSettings, DefaultNodeSettings);
+		UPCGNode* FoliagePaletteNode = Graph->AddNodeCopy(FoliagePaletteSettings, DefaultNodeSettings);
+		UPCGNode* OutputNode = Graph->AddNodeCopy(OutputSettings, DefaultNodeSettings);
+
+		OutputNode->NodeTitle = *VariantName;
+
+		PresetLoaderNode
+			->AddEdgeTo(*VariantName, GravityNode, PCGPinConstants::DefaultInputLabel)
+			->AddEdgeTo(PCGPinConstants::DefaultOutputLabel, MeshBuilderNode, PCGPinConstants::DefaultInputLabel)
+			->AddEdgeTo(PCGPinConstants::DefaultOutputLabel, FoliagePaletteNode, PCGPinConstants::DefaultInputLabel)
+			->AddEdgeTo(PCGPinConstants::DefaultOutputLabel, OutputNode, PCGPinConstants::DefaultInputLabel);
+
+		GravityNode->SetNodePosition(Position.X + 200, Position.Y);
+		MeshBuilderNode->SetNodePosition(Position.X + 400, Position.Y);
+		FoliagePaletteNode->SetNodePosition(Position.X + 600, Position.Y);
+		OutputNode->SetNodePosition(Position.X + 800, Position.Y);
+
+		Position.Y += 150;
+	}
+}
+#endif
diff --git a/Engine/Plugins/Experimental/ProceduralVegetationEditor/Source/ProceduralVegetation/Private/ProceduralVegetationLink.cpp b/Engine/Plugins/Experimental/ProceduralVegetationEditor/Source/ProceduralVegetation/Private/ProceduralVegetationLink.cpp
deleted file mode 100644
diff --git a/Engine/Plugins/Experimental/ProceduralVegetationEditor/Source/ProceduralVegetation/Public/ProceduralVegetation.h b/Engine/Plugins/Experimental/ProceduralVegetationEditor/Source/ProceduralVegetation/Public/ProceduralVegetation.h
--- a/Engine/Plugins/Experimental/ProceduralVegetationEditor/Source/ProceduralVegetation/Public/ProceduralVegetation.h
+++ b/Engine/Plugins/Experimental/ProceduralVegetationEditor/Source/ProceduralVegetation/Public/ProceduralVegetation.h
@@ -28,6 +28,8 @@ public:
 #endif
 	}
 
+	virtual void PostLoad() override;
+
 #if WITH_EDITOR
 	virtual bool ShouldDisplayDebuggingProperties() const override { return false; }
 	virtual bool CanToggleStandaloneGraph() const override { return false; };
