Improve PVE export progress reporting.

- Add traces scopes to all export functions
- Improve export flow to make progress reporting more accurate

diff --git a/Engine/Plugins/Experimental/ProceduralVegetationEditor/Source/ProceduralVegetation/Private/Helpers/PVExportHelper.cpp b/Engine/Plugins/Experimental/ProceduralVegetationEditor/Source/ProceduralVegetation/Private/Helpers/PVExportHelper.cpp
--- a/Engine/Plugins/Experimental/ProceduralVegetationEditor/Source/ProceduralVegetation/Private/Helpers/PVExportHelper.cpp
+++ b/Engine/Plugins/Experimental/ProceduralVegetationEditor/Source/ProceduralVegetation/Private/Helpers/PVExportHelper.cpp
@@ -37,8 +37,103 @@
 #include "GeometryCollection/Facades/CollectionConstraintOverrideFacade.h"
 #include "AssetRegistry/AssetData.h"
 
+#define LOCTEXT_NAMESPACE "PVExportHelper"
+
 namespace PV::Export::Internal
 {
+	enum class EExportMeshStatus
+	{
+		InitializingExportMesh,
+		CombiningMeshes,
+		CopyingMeshDataToOutputMesh,
+		ComputingBoneWeights,
+		CopyingSkinWeightProfiles,
+		BuildingNaniteAssemblies,
+		ImportingDynamicWind,
+		BuildingMesh,
+		MAX
+	};
+
+	float CalculateStatusWeight(EExportMeshStatus InStatus, bool bIsBuildingNaniteAssemblies, bool bIsExportingSkeletalMesh)
+	{
+		switch (InStatus)
+		{
+		case PV::Export::Internal::EExportMeshStatus::InitializingExportMesh:
+			return 1.f;
+		case PV::Export::Internal::EExportMeshStatus::CombiningMeshes:
+			return bIsBuildingNaniteAssemblies ? 0.f : 1.f;
+		case PV::Export::Internal::EExportMeshStatus::CopyingMeshDataToOutputMesh:
+			return 1.f;
+		case PV::Export::Internal::EExportMeshStatus::ComputingBoneWeights:
+			return bIsExportingSkeletalMesh ? 2.f : 0.f;
+		case PV::Export::Internal::EExportMeshStatus::CopyingSkinWeightProfiles:
+			return bIsExportingSkeletalMesh ? 0.1f : 0.f;
+		case PV::Export::Internal::EExportMeshStatus::BuildingNaniteAssemblies:
+			return bIsBuildingNaniteAssemblies ? 3.f : 0.f;
+		case PV::Export::Internal::EExportMeshStatus::ImportingDynamicWind:
+			return bIsExportingSkeletalMesh ? 0.1f : 0.f;
+		case PV::Export::Internal::EExportMeshStatus::BuildingMesh:
+			return bIsBuildingNaniteAssemblies ? 3.f : 4.f;
+		default:
+			break;
+		}
+
+		check(false);
+		return 0.f;
+	}
+
+	float CalculateExpectedWork(EExportMeshStatus InStatus, bool bIsBuildingNaniteAssemblies, bool bIsExportingSkeletalMesh)
+	{
+		float TotalWeightSum = 0;
+		for (int32 i = 0; i < (int32)EExportMeshStatus::MAX; ++i)
+		{
+			TotalWeightSum += CalculateStatusWeight(EExportMeshStatus(i), bIsBuildingNaniteAssemblies, bIsExportingSkeletalMesh);
+		}
+
+		float StatusWeight = CalculateStatusWeight(InStatus, bIsBuildingNaniteAssemblies, bIsExportingSkeletalMesh);
+
+		return StatusWeight > 0
+			? StatusWeight / TotalWeightSum
+			: 0.f;
+	}
+
+	FText GetStatusText(EExportMeshStatus InStatus, bool bIsBuildingNaniteAssemblies, bool bIsExportingSkeletalMesh)
+	{
+		switch (InStatus)
+		{
+		case PV::Export::Internal::EExportMeshStatus::InitializingExportMesh:
+			return LOCTEXT("ExportMeshStatus_InitializingExportMesh", "Initializing mesh data");
+		case PV::Export::Internal::EExportMeshStatus::CombiningMeshes:
+			return LOCTEXT("ExportMeshStatus_CombiningMeshes", "Combining foliage meshes");
+		case PV::Export::Internal::EExportMeshStatus::CopyingMeshDataToOutputMesh:
+			return LOCTEXT("ExportMeshStatus_CopyingMeshDataToOutputMesh", "Copying mesh data to output mesh");
+		case PV::Export::Internal::EExportMeshStatus::ComputingBoneWeights:
+			return LOCTEXT("ExportMeshStatus_ComputingBoneWeights", "Computing bone weights");
+		case PV::Export::Internal::EExportMeshStatus::CopyingSkinWeightProfiles:
+			return LOCTEXT("ExportMeshStatus_CopyingSkinWeightProfiles", "Copying skin weight profiles");
+		case PV::Export::Internal::EExportMeshStatus::BuildingNaniteAssemblies:
+			return LOCTEXT("ExportMeshStatus_BuildingNaniteAssemblies", "Building nanite assemblies");
+		case PV::Export::Internal::EExportMeshStatus::ImportingDynamicWind:
+			return LOCTEXT("ExportMeshStatus_ImportingDynamicWind", "Importing dynamic wind");
+		case PV::Export::Internal::EExportMeshStatus::BuildingMesh:
+		{
+			if (!bIsBuildingNaniteAssemblies && bIsExportingSkeletalMesh)
+			{
+				return LOCTEXT("ExportMeshStatus_BuildingMesh_Slow", "Building mesh (this might take a while)");
+			}
+			else
+			{
+				return LOCTEXT("ExportMeshStatus_BuildingMesh", "Building mesh");
+			}
+		}
+		default:
+			break;
+		}
+
+		check(false);
+		return FText::FromString("Error");
+	}
+
 	TObjectPtr<UDynamicMesh> CollectionToDynamicMesh(const FManagedArrayCollection& Collection)
 	{
 		TRACE_CPUPROFILER_EVENT_SCOPE(PV::Export::CollectionToDynamicMesh);
@@ -115,6 +210,8 @@ namespace PV::Export::Internal
 
 	TArray<TObjectPtr<UMaterialInterface>> GetMaterialsFromCollection(const FManagedArrayCollection& Collection)
 	{
+		TRACE_CPUPROFILER_EVENT_SCOPE(PV::Export::GetMaterialsFromCollection);
+
 		TArray<TObjectPtr<UMaterialInterface>> Materials;
 		if (Collection.HasAttribute("MaterialPath", FGeometryCollection::MaterialGroup))
 		{
@@ -130,6 +227,8 @@ namespace PV::Export::Internal
 	static bool SkeletalMeshToDynamicMesh(const USkeletalMesh* SkeletalMesh, const int32 LodIndex, FDynamicMesh3& ToDynamicMesh)
 	{
 #if WITH_EDITOR
+		TRACE_CPUPROFILER_EVENT_SCOPE(PV::Export::SkeletalMeshToDynamicMesh);
+
 		if (SkeletalMesh->HasMeshDescription(LodIndex))
 		{
 			const FMeshDescription* SourceMesh = SkeletalMesh->GetMeshDescription(LodIndex);
@@ -155,6 +254,8 @@ namespace PV::Export::Internal
 		const FMeshInstanceCombined& OnMeshInstanceCombined
 	)
 	{
+		TRACE_CPUPROFILER_EVENT_SCOPE(PV::Export::CombineMeshInstancesToDynamicMesh);
+
 		Facades::FFoliageFacade FoliageFacade(Collection);
 		int32 NumInstances = FoliageFacade.NumFoliageEntries();
 		
@@ -315,6 +416,8 @@ namespace PV::Export::Internal
 		UStaticMesh* StaticMesh
 	)
 	{
+		TRACE_CPUPROFILER_EVENT_SCOPE(PV::Export::BuildNaniteAssemblyData);
+
 		//Create NaniteAssemblyDataBuilder and the main static mesh materials to it
 		FNaniteAssemblyDataBuilder AssemblyBuilder;
 		for (const FStaticMaterial& Material : StaticMesh->GetStaticMaterials())
@@ -402,6 +505,8 @@ namespace PV::Export::Internal
 		USkeletalMesh* SkeletalMesh
 	)
 	{
+		TRACE_CPUPROFILER_EVENT_SCOPE(PV::Export::BuildNaniteAssemblyData);
+
 		//Create NaniteAssemblyDataBuilder and the main static mesh materials to it
 		FNaniteAssemblyDataBuilder AssemblyBuilder;
 		Facades::FFoliageFacade FoliageFacade(Collection);
@@ -486,6 +591,8 @@ namespace PV::Export::Internal
 	void AddNodeToBuilder(FNaniteAssemblyDataBuilder& AssemblyBuilder, const FManagedArrayCollection& Collection,
 		const TMap<FString, int32>& InMeshNamePartMap, const TMap<FString, TArray<TObjectPtr<UMaterialInterface>>>& InMeshMaterialsMap)
 	{
+		TRACE_CPUPROFILER_EVENT_SCOPE(PV::Export::AddNodeToBuilder);
+
 		Facades::FFoliageFacade FoliageFacade(Collection);
 		int32 NumInstances = FoliageFacade.NumFoliageEntries();
 			
@@ -532,6 +639,8 @@ namespace PV::Export::Internal
 
 	void ConvertToDefaultSkeletalMesh(USkeletalMesh* SkeletalMesh, FDynamicMesh3* Mesh,TArray<TObjectPtr<UMaterialInterface>> Materials)
 	{
+		TRACE_CPUPROFILER_EVENT_SCOPE(PV::Export::ConvertToDefaultSkeletalMesh);
+
 		TObjectPtr<UDynamicMesh> DynamicMesh = NewObject<UDynamicMesh>();
 		DynamicMesh->SetMesh(MoveTemp(*Mesh));
 		
@@ -598,58 +707,87 @@ namespace PV::Export::Internal
 		TObjectPtr<UStaticMesh> ExportMesh,
 		const FManagedArrayCollection& Collection,
 		bool bBuildNaniteAssemblies,
-		ENaniteShapePreservation ShapePreservation
+		ENaniteShapePreservation ShapePreservation,
+		const FStatusReportCallback& StatusReportCallback
 	)
 	{
+		TRACE_CPUPROFILER_EVENT_SCOPE(PV::Export::ExportCollectionToStaticMesh);
+
+		const auto UpdateStatus = [&](EExportMeshStatus NewStatus)
+		{
+			StatusReportCallback(CalculateExpectedWork(NewStatus, bBuildNaniteAssemblies, false), GetStatusText(NewStatus, bBuildNaniteAssemblies, false));
+		};
+
 		FManagedArrayCollection ExportedCollection;
-		Collection.CopyTo(&ExportedCollection);
-		
-		TObjectPtr<UDynamicMesh> DynamicMesh = CollectionToDynamicMesh(ExportedCollection);
+		TObjectPtr<UDynamicMesh> DynamicMesh;
+		TArray<TObjectPtr<UMaterialInterface>> Materials;
+		{
+			TRACE_CPUPROFILER_EVENT_SCOPE(InitExportMesh);
 
-		//Get the materials from the collection
-		TArray<TObjectPtr<UMaterialInterface>> Materials = GetMaterialsFromCollection(ExportedCollection);
+			UpdateStatus(EExportMeshStatus::InitializingExportMesh);
 
-		//Get the foliage data from facade, spawn an actor attach the foliage to the actor
-		Facades::FFoliageFacade FoliageFacade(ExportedCollection);
+			Collection.CopyTo(&ExportedCollection);
+
+			DynamicMesh = CollectionToDynamicMesh(ExportedCollection);
+			Materials = GetMaterialsFromCollection(ExportedCollection);
+		}
 		
 		//Combine all the foliage instances into one dynamic mesh if not building nanite assemblies
 		if (!bBuildNaniteAssemblies)
 		{
+			UpdateStatus(EExportMeshStatus::CombiningMeshes);
 			CombineMeshInstancesToDynamicMesh(ExportedCollection, DynamicMesh, Materials,
 				[](FString MeshName, int32 InstanceID, int32 VertexStart, int32 VertexCount){});
 		}
 
-		//Convert the dynamic mesh to static mesh
 		EGeometryScriptOutcomePins Output;
+		{
+			TRACE_CPUPROFILER_EVENT_SCOPE(CopyMeshToStaticMesh);
+
+			UpdateStatus(EExportMeshStatus::CopyingMeshDataToOutputMesh);
+
 #if WITH_EDITORONLY_DATA
-		FGeometryScriptCopyMeshToAssetOptions Options = GetCopyMeshToAssetOptions(Materials, ExportMesh->GetNaniteSettings(), ShapePreservation);
-		
+			FGeometryScriptCopyMeshToAssetOptions Options = GetCopyMeshToAssetOptions(Materials, ExportMesh->GetNaniteSettings(), ShapePreservation);
 #else
-		FGeometryScriptCopyMeshToAssetOptions Options = GetCopyMeshToAssetOptions(Materials);
+			FGeometryScriptCopyMeshToAssetOptions Options = GetCopyMeshToAssetOptions(Materials);
 #endif
-		UGeometryScriptLibrary_StaticMeshFunctions::CopyMeshToStaticMesh(DynamicMesh, ExportMesh, Options, {}, Output);
+			Options.bEmitTransaction = false;
+			Options.bDeferMeshPostEditChange = true;
+
+			UGeometryScriptLibrary_StaticMeshFunctions::CopyMeshToStaticMesh(DynamicMesh, ExportMesh, Options, {}, Output);
+		}
 
 #if WITH_EDITOR
 		if (bBuildNaniteAssemblies)
 		{
+			UpdateStatus(EExportMeshStatus::BuildingNaniteAssemblies);
+
 			if (NaniteAssembliesSupported())
 			{
 				//Build the nanite assemblies
             	BuildNaniteAssemblyData(ExportedCollection, ExportMesh);
-            	ExportMesh->Build(true);
-				FStaticMeshCompilingManager::Get().FinishAllCompilation();
 			}
 			else
 			{
 				UE_LOG(LogProceduralVegetation, Warning, TEXT("Failed to build Nanite Assemblies, because neither Nanite Foliage nor Nanite Assemblies are enabled for the project."));
 			}
 		}
+
+		{
+			TRACE_CPUPROFILER_EVENT_SCOPE(BuildSkinnedMesh);
+			UpdateStatus(EExportMeshStatus::BuildingMesh);
+
+			ExportMesh->Build(true);
+			FStaticMeshCompilingManager::Get().FinishAllCompilation();
+		}
 #endif
 	}
 
 #if WITH_EDITOR
 	void BuildWindImportData(FDynamicWindSkeletalImportData& OutImportData, const FManagedArrayCollection& Collection, const TArray<PV::Facades::FBoneNode>& Bones)
 	{
+		TRACE_CPUPROFILER_EVENT_SCOPE(PV::Export::BuildWindImportData);
+
 		Facades::FBranchFacade BranchFacade(Collection);
 
 		for (const auto& Bone : Bones)
@@ -664,139 +802,175 @@ namespace PV::Export::Internal
 		TObjectPtr<USkeletalMesh> ExportMesh,
 		const FManagedArrayCollection& Collection,
 		ENaniteShapePreservation ShapePreservation,
-		bool bBuildNaniteAssemblies
+		bool bBuildNaniteAssemblies,
+		const FStatusReportCallback& StatusReportCallback
 	)
 	{
-		FManagedArrayCollection ExportedCollection;
-		Collection.CopyTo(&ExportedCollection);
+		TRACE_CPUPROFILER_EVENT_SCOPE(PV::Export::ExportCollectionToSkeletalMesh);
 
-		Facades::FFoliageFacade FoliageFacade(ExportedCollection);
-		Facades::FPointFacade PointFacade(ExportedCollection);
-		Facades::FBranchFacade BranchFacade(ExportedCollection);
-		Facades::FBoneFacade BoneFacade = Facades::FBoneFacade(ExportedCollection);
-		
-		TObjectPtr<UDynamicMesh> DynamicMesh = CollectionToDynamicMesh(ExportedCollection);
-		UE::Geometry::FDynamicMeshAttributeSet* MeshAttributes = DynamicMesh->GetMeshPtr()->Attributes();
-		
-		//Get the materials from the collection
-		TArray<TObjectPtr<UMaterialInterface>> Materials = GetMaterialsFromCollection(ExportedCollection);
+		const auto UpdateStatus = [&](EExportMeshStatus NewStatus)
+		{
+			StatusReportCallback(CalculateExpectedWork(NewStatus, bBuildNaniteAssemblies, true), GetStatusText(NewStatus, bBuildNaniteAssemblies, true));
+		};
 
+		FManagedArrayCollection ExportedCollection;
+		TObjectPtr<UDynamicMesh> DynamicMesh;
+		TArray<TObjectPtr<UMaterialInterface>> Materials;
+		TArray<int32> VertexBoneIDs;
 		TArray<Facades::FBoneNode> BoneNodes;
-		BoneNodes = BoneFacade.GetBoneDataFromCollection();
-
-		//if bones already not created with bone reduction node create bones with full density
-		if (BoneNodes.IsEmpty())
+		int32 FoliageVertexStart = 0;
 		{
-			BoneFacade.CreateBoneData(BoneNodes, 0);
-		}
-		
-		AssignBoneIDsToFoliage(BoneNodes, ExportedCollection);
-		
-		int32 BoneCount = BoneNodes.Num();
-		MeshAttributes->EnableBones(BoneCount);
+			TRACE_CPUPROFILER_EVENT_SCOPE(InitExportMesh);
 
-		TArray<int32> VertexPointIDs;
-		BoneFacade.GetPointIds(VertexPointIDs);
+			UpdateStatus(EExportMeshStatus::InitializingExportMesh);
+			Collection.CopyTo(&ExportedCollection);
 
-		TArray<int32> VertexBoneIDs;
-		//Assign the Bone id to every vertex in the base mesh through VertexPointIDs
-		for (const int32& PointID : VertexPointIDs)
-		{
-			if (Facades::FBoneNode* BoneNode = Facades::FBoneFacade::FindClosestBone(ExportedCollection, BoneNodes, PointID))
+			DynamicMesh = CollectionToDynamicMesh(ExportedCollection);
+			Materials = GetMaterialsFromCollection(ExportedCollection);
+
+			UE::Geometry::FDynamicMeshAttributeSet* MeshAttributes = DynamicMesh->GetMeshPtr()->Attributes();
+
+			Facades::FBoneFacade BoneFacade = Facades::FBoneFacade(ExportedCollection);
+			BoneNodes = BoneFacade.GetBoneDataFromCollection();
+
+			// if bones already not created with bone reduction node create bones with full density
+			if (BoneNodes.IsEmpty())
 			{
-				VertexBoneIDs.Add(BoneNode->BoneIndex);
+				BoneFacade.CreateBoneData(BoneNodes, 0);
 			}
-			else
+
+			AssignBoneIDsToFoliage(BoneNodes, ExportedCollection);
+
+			const int32 BoneCount = BoneNodes.Num();
+			MeshAttributes->EnableBones(BoneCount);
+
+			TArray<int32> VertexPointIDs;
+			BoneFacade.GetPointIds(VertexPointIDs);
+			
+			// Assign the Bone id to every vertex in the base mesh through VertexPointIDs
+			for (const int32& PointID : VertexPointIDs)
 			{
-				VertexBoneIDs.Add(INDEX_NONE);
-				UE_LOG(LogProceduralVegetation, Log, TEXT("Invalid Bone Assigned to PointID %i") , PointID);
+				if (Facades::FBoneNode* BoneNode = Facades::FBoneFacade::FindClosestBone(ExportedCollection, BoneNodes, PointID))
+				{
+					VertexBoneIDs.Add(BoneNode->BoneIndex);
+				}
+				else
+				{
+					VertexBoneIDs.Add(INDEX_NONE);
+					UE_LOG(LogProceduralVegetation, Log, TEXT("Invalid Bone Assigned to PointID %i"), PointID);
+				}
 			}
-		}
-		int FoliageVertexStart = VertexBoneIDs.Num();
 
-		UE::Geometry::FDynamicMeshBoneNameAttribute* BoneNames = MeshAttributes->GetBoneNames();
-		UE::Geometry::FDynamicMeshBoneParentIndexAttribute* BoneParentIndices = MeshAttributes->GetBoneParentIndices();
-		UE::Geometry::FDynamicMeshBonePoseAttribute* BonePoses = MeshAttributes->GetBonePoses();
+			FoliageVertexStart = VertexBoneIDs.Num();
 
-		for (int i = 0; i < BoneCount; i++)
-		{
-			BoneNames->SetValue(i, BoneNodes[i].BoneName);
-			BonePoses->SetValue(i, BoneNodes[i].BoneTransform);
-			BoneParentIndices->SetValue(i, BoneNodes[i].ParentBoneIndex);
+			UE::Geometry::FDynamicMeshBoneNameAttribute* BoneNames = MeshAttributes->GetBoneNames();
+			UE::Geometry::FDynamicMeshBoneParentIndexAttribute* BoneParentIndices = MeshAttributes->GetBoneParentIndices();
+			UE::Geometry::FDynamicMeshBonePoseAttribute* BonePoses = MeshAttributes->GetBonePoses();
+
+			for (int i = 0; i < BoneCount; i++)
+			{
+				BoneNames->SetValue(i, BoneNodes[i].BoneName);
+				BonePoses->SetValue(i, BoneNodes[i].BoneTransform);
+				BoneParentIndices->SetValue(i, BoneNodes[i].ParentBoneIndex);
+			}
 		}
 
 		if (!bBuildNaniteAssemblies)
 		{
-			//Combine all the foliage instances into one dynamic mesh if not building nanite assemblies
+			UpdateStatus(EExportMeshStatus::CombiningMeshes);
+
+			Facades::FFoliageFacade FoliageFacade(ExportedCollection);
+
+			// Combine all the foliage instances into one dynamic mesh if not building nanite assemblies
 			CombineMeshInstancesToDynamicMesh(ExportedCollection, DynamicMesh, Materials,
-			[&](const FString& MeshName, int32 InstanceID, int32 VertexStart, int32 VertexCount)
-			{
-				//Assign parent bone id to foliage vertices
-				int32 FoliageParentBoneID = FoliageFacade.GetParentBoneID(InstanceID);
-					
-				for (int VertexID = VertexStart; VertexID < VertexCount; ++VertexID)
+				[&](const FString& MeshName, int32 InstanceID, int32 VertexStart, int32 VertexCount)
 				{
-					VertexBoneIDs.Add(FoliageParentBoneID);
-					//UE_LOG(LogProceduralVegetation, Warning, TEXT("Adding VertexBoneIDs, VertexId{%i} BoneID{%i}, VertexStart{%i}, VertexCount{%i}, MeshName{%s}, InstanceID{%i}"),
-					//VertexID, FoliageParentBoneID, VertexStart, VertexCount, *MeshName, InstanceID);
-				}
-			});
+					// Assign parent bone id to foliage vertices
+					int32 FoliageParentBoneID = FoliageFacade.GetParentBoneID(InstanceID);
+					
+					for (int VertexID = VertexStart; VertexID < VertexCount; ++VertexID)
+					{
+						VertexBoneIDs.Add(FoliageParentBoneID);
+					}
+				});
 		}
 
-		//Convert the dynamic mesh to static mesh
+		// Convert the dynamic mesh to static mesh
 		EGeometryScriptOutcomePins Output;
-		
+		{
+			TRACE_CPUPROFILER_EVENT_SCOPE(CopyMeshToSkeletalMesh);
+
+			UpdateStatus(EExportMeshStatus::CopyingMeshDataToOutputMesh);
+
 #if WITH_EDITORONLY_DATA
-		FGeometryScriptCopyMeshToAssetOptions Options = GetCopyMeshToAssetOptions(Materials, ExportMesh->GetNaniteSettings(), ShapePreservation);
-		
+			FGeometryScriptCopyMeshToAssetOptions Options = GetCopyMeshToAssetOptions(Materials, ExportMesh->GetNaniteSettings(), ShapePreservation);
 #else
-		FGeometryScriptCopyMeshToAssetOptions Options = GetCopyMeshToAssetOptions(Materials);
+			FGeometryScriptCopyMeshToAssetOptions Options = GetCopyMeshToAssetOptions(Materials);
 #endif
-		Options.BoneHierarchyMismatchHandling = EGeometryScriptBoneHierarchyMismatchHandling::CreateNewReferenceSkeleton;
+			Options.BoneHierarchyMismatchHandling = EGeometryScriptBoneHierarchyMismatchHandling::CreateNewReferenceSkeleton;
+			Options.bEmitTransaction = false;
+			Options.bDeferMeshPostEditChange = true;
 
-		UGeometryScriptLibrary_StaticMeshFunctions::CopyMeshToSkeletalMesh(
-			DynamicMesh,
-			ExportMesh,
-			Options,
-			{},
-			Output
-		);
+			UGeometryScriptLibrary_StaticMeshFunctions::CopyMeshToSkeletalMesh(
+				DynamicMesh,
+				ExportMesh,
+				Options,
+				{},
+				Output
+			);
+		}
 
 		TObjectPtr<USkeleton> Skeleton = ExportMesh->GetSkeleton();
 		Skeleton->RecreateBoneTree(ExportMesh);
 
 		const FGeometryScriptBoneWeightProfile SkinProfile = FGeometryScriptBoneWeightProfile();
-		FGeometryScriptSmoothBoneWeightsOptions SmoothBoneWeightsOptions;
+		const FName ProfileName = SkinProfile.GetProfileName();
 
-		SmoothBoneWeightsOptions.DistanceWeighingType = EGeometryScriptSmoothBoneWeightsType::DirectDistance;
-		SmoothBoneWeightsOptions.MaxInfluences = 3;
-		SmoothBoneWeightsOptions.Stiffness = 0.5;
-		
-		UGeometryScriptLibrary_MeshBoneWeightFunctions::ComputeSmoothBoneWeights(
-			DynamicMesh,
-			Skeleton,
-			SmoothBoneWeightsOptions,
-			SkinProfile
-		);
+		{
+			TRACE_CPUPROFILER_EVENT_SCOPE(ComputeSmoothBoneWeights);
 
-		const FName ProfileName = SkinProfile.GetProfileName();
-		
-		//Here we make the correction to the weights
-		RemoveUnwantedSkinWeights(DynamicMesh, ProfileName, VertexBoneIDs, FoliageVertexStart,BoneNodes);
+			UpdateStatus(EExportMeshStatus::ComputingBoneWeights);
+
+			FGeometryScriptSmoothBoneWeightsOptions SmoothBoneWeightsOptions;
+			SmoothBoneWeightsOptions.DistanceWeighingType = EGeometryScriptSmoothBoneWeightsType::DirectDistance;
+			SmoothBoneWeightsOptions.MaxInfluences = 3;
+			SmoothBoneWeightsOptions.Stiffness = 0.5;
+
+			UGeometryScriptLibrary_MeshBoneWeightFunctions::ComputeSmoothBoneWeights(
+				DynamicMesh,
+				Skeleton,
+				SmoothBoneWeightsOptions,
+				SkinProfile
+			);
+
+			// Here we make the correction to the weights
+			RemoveUnwantedSkinWeights(DynamicMesh, ProfileName, VertexBoneIDs, FoliageVertexStart, BoneNodes);
+		}
 		
-		UGeometryScriptLibrary_StaticMeshFunctions::CopySkinWeightProfileToSkeletalMesh(
-			DynamicMesh,
-			ExportMesh,
-			ProfileName,
-			ProfileName,
-			{},
-			{},
-			Output
-		);
+		{
+			TRACE_CPUPROFILER_EVENT_SCOPE(CopySkinWeightProfileToSkeletalMesh);
+
+			UpdateStatus(EExportMeshStatus::CopyingSkinWeightProfiles);
+
+			FGeometryScriptCopySkinWeightProfileToAssetOptions Options;
+			Options.bEmitTransaction = false;
+			Options.bDeferMeshPostEditChange = true;
+			UGeometryScriptLibrary_StaticMeshFunctions::CopySkinWeightProfileToSkeletalMesh(
+				DynamicMesh,
+				ExportMesh,
+				ProfileName,
+				ProfileName,
+				Options,
+				{},
+				Output
+			);
+		}
 
 #if WITH_EDITOR
 		if (bBuildNaniteAssemblies)
 		{
+			UpdateStatus(EExportMeshStatus::BuildingNaniteAssemblies);
+
 			if (NaniteAssembliesSupported())
 			{
 				//Build the nanite assemblies
@@ -808,16 +982,26 @@ namespace PV::Export::Internal
 			}
 		}
 
-		FDynamicWindSkeletalImportData ImportData;
-		BuildWindImportData(ImportData,ExportedCollection, BoneNodes);
-				
-		if (!DynamicWind::ImportSkeletalData(*ExportMesh, ImportData))
 		{
-			UE_LOG(LogTemp, Warning, TEXT("Failed to build wind data"));
+			TRACE_CPUPROFILER_EVENT_SCOPE(DynamicWind);
+
+			UpdateStatus(EExportMeshStatus::ImportingDynamicWind);
+
+			FDynamicWindSkeletalImportData ImportData;
+			BuildWindImportData(ImportData, ExportedCollection, BoneNodes);
+
+			if (!DynamicWind::ImportSkeletalData(*ExportMesh, ImportData))
+			{
+				UE_LOG(LogProceduralVegetation, Warning, TEXT("Failed to build wind data"));
+			}
+		}
+		
+		{
+			TRACE_CPUPROFILER_EVENT_SCOPE(BuildSkinnedMesh);
+			UpdateStatus(EExportMeshStatus::BuildingMesh);
+			ExportMesh->Build();
+			FSkinnedAssetCompilingManager::Get().FinishAllCompilation();
 		}
-				
-		ExportMesh->Build();
-		FSkinnedAssetCompilingManager::Get().FinishAllCompilation();
 #endif
 	}
 
@@ -834,6 +1018,8 @@ namespace PV::Export::Internal
 
 	void AssignBoneIDsToFoliage(const TArray<Facades::FBoneNode>& Bones, FManagedArrayCollection& Collection)
 	{
+		TRACE_CPUPROFILER_EVENT_SCOPE(PV::Export::AssignBoneIDsToFoliage);
+
 		Facades::FPointFacade PointFacade(Collection);
 		Facades::FFoliageFacade FoliageFacade(Collection);
 		
@@ -892,6 +1078,8 @@ namespace PV::Export::Internal
 
 	void RemoveUnwantedSkinWeights(TObjectPtr<UDynamicMesh> DynamicMesh, const FName ProfileName, const TArray<int32>& VertexBoneIDs, int32 FoliageVertexStart, const TArray<Facades::FBoneNode>& BoneNodes)
 	{
+		TRACE_CPUPROFILER_EVENT_SCOPE(PV::Export::RemoveUnwantedSkinWeights);
+
 		auto InMesh = DynamicMesh->GetMeshPtr();
 		check(InMesh)
 		const int32 NumVertices = VertexBoneIDs.Num();
@@ -960,9 +1148,12 @@ namespace PV::Export
 		const TObjectPtr<UProceduralVegetation> InProceduralVegetation,
 		const FManagedArrayCollection& Collection,
 		const FManagedArrayCollection& FoliageCollection,
-		const FPVExportParams& ExportParams
+		const FPVExportParams& ExportParams,
+		const FStatusReportCallback& StatusReportCallback
 	)
 	{
+		TRACE_CPUPROFILER_EVENT_SCOPE(PV::Export::ExportCollectionAsMesh);
+
 		// Set up the Package path and name from ExportParams
 		const FString PackageName = ExportParams.ContentBrowserFolder.Path;
 		const FString MeshName = ExportParams.MeshName.ToString();
@@ -982,7 +1173,8 @@ namespace PV::Export
 				ExportMesh,
 				Collection,
 				ExportParams.bCreateNaniteFoliage,
-				ExportParams.NaniteShapePreservation
+				ExportParams.NaniteShapePreservation,
+				StatusReportCallback
 			);
 
 			MeshPackage->SetDirtyFlag(true);
@@ -1027,7 +1219,8 @@ namespace PV::Export
 				ExportMesh,
 				Collection,
 				ExportParams.NaniteShapePreservation,
-				ExportParams.bCreateNaniteFoliage
+				ExportParams.bCreateNaniteFoliage,
+				StatusReportCallback
 			);
 
 			MeshPackage->SetDirtyFlag(true);
@@ -1038,3 +1231,5 @@ namespace PV::Export
 		}
 	}
 }
+
+#undef LOCTEXT_NAMESPACE 
diff --git a/Engine/Plugins/Experimental/ProceduralVegetationEditor/Source/ProceduralVegetation/Public/Helpers/PVExportHelper.h b/Engine/Plugins/Experimental/ProceduralVegetationEditor/Source/ProceduralVegetation/Public/Helpers/PVExportHelper.h
--- a/Engine/Plugins/Experimental/ProceduralVegetationEditor/Source/ProceduralVegetation/Public/Helpers/PVExportHelper.h
+++ b/Engine/Plugins/Experimental/ProceduralVegetationEditor/Source/ProceduralVegetation/Public/Helpers/PVExportHelper.h
@@ -129,10 +129,13 @@ namespace PV::Export::Internal
 
 namespace PV::Export
 {
+	typedef TFunction<void(float ExpectedWork, const FText& StatusText)> FStatusReportCallback;
+
 	void PROCEDURALVEGETATION_API ExportCollectionAsMesh(
 		const TObjectPtr<UProceduralVegetation> InProceduralVegetation,
 		const FManagedArrayCollection& Collection,
 		const FManagedArrayCollection& FoliageCollection,
-		const FPVExportParams& ExportParams
+		const FPVExportParams& ExportParams,
+		const FStatusReportCallback& StatusReportCallback
 	);
 }
diff --git a/Engine/Plugins/Experimental/ProceduralVegetationEditor/Source/ProceduralVegetationEditor/Private/PVEditor.cpp b/Engine/Plugins/Experimental/ProceduralVegetationEditor/Source/ProceduralVegetationEditor/Private/PVEditor.cpp
--- a/Engine/Plugins/Experimental/ProceduralVegetationEditor/Source/ProceduralVegetationEditor/Private/PVEditor.cpp
+++ b/Engine/Plugins/Experimental/ProceduralVegetationEditor/Source/ProceduralVegetationEditor/Private/PVEditor.cpp
@@ -11,6 +11,7 @@
 #include "PVEditorSettings.h"
 #include "ToolMenus.h"
 #include "ToolMenuSection.h"
+#include "Misc/UObjectToken.h"
 
 #include "Dataflow/DataflowCollectionSpreadSheetWidget.h"
 
@@ -352,63 +353,93 @@ FReply FPVEditor::ExportInternal()
 		.Title(LOCTEXT("Export Selection", "Choose outputs to export"))
 		.Graph(ProceduralVegetationGraph);
 
-	if (ExportSelectionWidget->ShowModal() == EAppReturnType::Ok)
+	if (ExportSelectionWidget->ShowModal() != EAppReturnType::Ok)
 	{
-		const TArray<UPCGNode*>& Nodes = ProceduralVegetationGraph->GetNodes();
+		return FReply::Handled();
+	}
 
-		TArray<const UPCGNode*> NodesToExport;
-		
-		for (const UPCGNode* Node : Nodes)
+	const TArray<UPCGNode*> NodesToExport = ProceduralVegetationGraph->GetNodes().FilterByPredicate(
+		[](UPCGNode* Node) 
 		{
-			if(const UPVOutputSettings* OutputSettings = Cast<UPVOutputSettings>(Node->GetSettings()))
-			{
-				if(OutputSettings->ExportSettings.bShouldExport)
-				{
-					NodesToExport.Add(Node);
-				}
-			}
+			const UPVOutputSettings* OutputSettings = Cast<UPVOutputSettings>(Node->GetSettings());
+			return OutputSettings && OutputSettings->ExportSettings.bShouldExport;
 		}
+	);
 
-		FScopedSlowTask Progress(NodesToExport.Num(), NSLOCTEXT("ExportAsset", "ExportAsset_SlowTask", "Exporting Asset(s) ..."));
-		Progress.MakeDialog();
+	TArray<TSharedRef<FTokenizedMessage>> ErrorMessages;
+	const auto AddError = [&](const FText& ErrorMsg)->TSharedRef<FTokenizedMessage>
+	{
+		TSharedRef<FTokenizedMessage> Error = FTokenizedMessage::Create(EMessageSeverity::Error, ErrorMsg);
+		ErrorMessages.Add(Error);
+		return Error;
+	};
+
+	FScopedSlowTask Progress(NodesToExport.Num(), LOCTEXT("ExportAsset_SlowTask", "Exporting Asset(s) ..."));
+	Progress.MakeDialog();
+
+	for (int32 i = 0; i < NodesToExport.Num(); ++i)
+	{
+		UPCGNode* Node = NodesToExport[i];
+		const UPVOutputSettings* OutputSettings = Cast<UPVOutputSettings>(Node->GetSettings());
+		check(OutputSettings != nullptr);
 
-		FString ErrorMessage;
-					
-		for(auto Node : NodesToExport)
+		const FText AssetCountText = NodesToExport.Num() > 1
+			? FText::FromString(FString::Printf(TEXT(" [%d/%d]"), i + 1, NodesToExport.Num()))
+			: FText::FromString("");
+
+		Progress.EnterProgressFrame(1.f, FText::Format(LOCTEXT("ExportingAsset", "Exporting \"{0}\"{1}"), FText::FromName(OutputSettings->ExportSettings.MeshName), AssetCountText));
+
+		FManagedArrayCollection OutCollection;
+		FName PinName = NAME_None;
+		if (!GetCollectionFromNode(Node, PinName, OutCollection))
 		{
-			if(const UPVOutputSettings* OutputSettings = Cast<UPVOutputSettings>(Node->GetSettings()))
-			{
-				FManagedArrayCollection OutCollection;
-				FName PinName = NAME_None;
-				if (GetCollectionFromNode(Node, PinName, OutCollection))
-				{
-					FString Reason;
-					UClass* MeshClass = OutputSettings->ExportSettings.ExportMeshType == EPVExportMeshType::StaticMesh ? UStaticMesh::StaticClass() : USkeletalMesh::StaticClass();
-					if (PV::Utilities::ValidateMeshPathAndName(OutputSettings->ExportSettings.MeshName.ToString(), OutputSettings->ExportSettings.ContentBrowserFolder.Path,
-						MeshClass, Reason))
-					{
-						PV::Export::ExportCollectionAsMesh(ProceduralVegetationBeingEdited, OutCollection, OutCollection, OutputSettings->ExportSettings);
-					}
-					else
-					{
-						ErrorMessage += FString::Format( TEXT("Export failed for mesh {0} with error {1}"), {OutputSettings->ExportSettings.MeshName.ToString(), Reason});
-					}
-					Progress.EnterProgressFrame();
-				}
-				else
-				{
-					ErrorMessage += FString::Format( TEXT("Export failed for output node {0}, Invalid output please fix the errors on the node."), { Node->GetAuthoredTitleName().ToString()});
-				}
-			}
+			TSharedRef<FTokenizedMessage> Error = AddError(
+				FText::Format(
+					LOCTEXT("ExportFailed_NodeErrors", "Export failed for output node \"{0}\", Invalid output please fix the errors on the node."), 
+					Node->GetAuthoredTitleLine()
+				)
+			);
+			Error->AddToken(FUObjectToken::Create(Node));
+			continue;
 		}
 
-		if (!ErrorMessage.IsEmpty())
+		FString PathError;
+		UClass* MeshClass = OutputSettings->ExportSettings.ExportMeshType == EPVExportMeshType::StaticMesh ? UStaticMesh::StaticClass() : USkeletalMesh::StaticClass();
+		if (!PV::Utilities::ValidateMeshPathAndName(OutputSettings->ExportSettings.MeshName.ToString(), OutputSettings->ExportSettings.ContentBrowserFolder.Path,
+			MeshClass, PathError))
 		{
-			FMessageLog EditorErrors("EditorErrors");
-			EditorErrors.Error(LOCTEXT("PVE Mesh Export Failed", "PVE Mesh Export Failed"));
-			EditorErrors.Info(FText::FromString(ErrorMessage));
-			EditorErrors.Notify(LOCTEXT("PVE Mesh Export Failed", "PVE Mesh Export Failed"));
+			TSharedRef<FTokenizedMessage> Error = AddError(
+				FText::Format(
+					LOCTEXT("ExportFailed_InvalidPath", "Export failed for mesh \"{0}\" with error \"{1}\""),
+					Node->GetAuthoredTitleLine(),
+					FText::FromString(PathError)
+				)
+			);
+			Error->AddToken(FUObjectToken::Create(Node));
+			continue;
 		}
+		
+		FScopedSlowTask ExportCollectionSlowTask(1.f,LOCTEXT("GeneratingMesh", "Generating Mesh ..."));
+		ExportCollectionSlowTask.MakeDialog();
+		float PrevFrameWork = 0.f;
+		const auto OnStatusUpdated = [&](float ExpectedWorkThisFrame, const FText& StatusText)
+		{
+			ExportCollectionSlowTask.EnterProgressFrame(PrevFrameWork, StatusText);
+
+			// Since our frames are so slow it does not make sense for the progress bar to progress as we start the frame, but rather as the frame finishes,
+			// hence why we repport the previours frame's work.
+			PrevFrameWork = ExpectedWorkThisFrame;
+		};
+		PV::Export::ExportCollectionAsMesh(ProceduralVegetationBeingEdited, OutCollection, OutCollection, OutputSettings->ExportSettings, OnStatusUpdated);
+		
+		ExportCollectionSlowTask.EnterProgressFrame(PrevFrameWork); // Let the slow task know all work has been completed.
+	}
+
+	if (ErrorMessages.Num() > 0)
+	{
+		FMessageLog EditorErrors("EditorErrors");
+		EditorErrors.AddMessages(ErrorMessages);
+		EditorErrors.Notify(LOCTEXT("PVE Mesh Export Failed", "PVE Mesh Export Failed"));
 	}
 	
 	return FReply::Handled();
