[PVE] Add Tooltips for Foliage distributor and output node. Update default scale clamp in foliage distributor

diff --git a/Engine/Plugins/Experimental/ProceduralVegetationEditor/Source/ProceduralVegetation/Private/Implementations/PVMeshBuilder.h b/Engine/Plugins/Experimental/ProceduralVegetationEditor/Source/ProceduralVegetation/Private/Implementations/PVMeshBuilder.h
--- a/Engine/Plugins/Experimental/ProceduralVegetationEditor/Source/ProceduralVegetation/Private/Implementations/PVMeshBuilder.h
+++ b/Engine/Plugins/Experimental/ProceduralVegetationEditor/Source/ProceduralVegetation/Private/Implementations/PVMeshBuilder.h
@@ -137,7 +137,7 @@ struct FPVMeshBuilderParams
 	FIntVector2 DisplacementUVScale = FIntVector2(1.0f, 1.0f);
 
 	UPROPERTY(EditAnywhere, Category="Displacement Settings", DisplayName="Apply Displacement up to generation",
-		meta=(ClampMin=1, ClampMax=10, UIMin=1, UIMax=10, Tooltip="Controls which generations the displacement is applied on \n\nControls which generations the displacement is applied on "))
+		meta=(ClampMin=1, ClampMax=10, UIMin=1, UIMax=10, Tooltip="Controls which generations the displacement is applied on \n\nControls which generations the displacement is applied on."))
 	int32 DisplacementGenerationUpperLimit = 1;
 
 	FPVMeshBuilderParams()
diff --git a/Engine/Plugins/Experimental/ProceduralVegetationEditor/Source/ProceduralVegetation/Private/Nodes/PVBoneReductionSettings.h b/Engine/Plugins/Experimental/ProceduralVegetationEditor/Source/ProceduralVegetation/Private/Nodes/PVBoneReductionSettings.h
--- a/Engine/Plugins/Experimental/ProceduralVegetationEditor/Source/ProceduralVegetation/Private/Nodes/PVBoneReductionSettings.h
+++ b/Engine/Plugins/Experimental/ProceduralVegetationEditor/Source/ProceduralVegetation/Private/Nodes/PVBoneReductionSettings.h
@@ -29,7 +29,7 @@ protected:
 	virtual FPCGElementPtr CreateElement() const override;
 
 private:
-	UPROPERTY(EditAnywhere, Category = "Bone Reduction", meta = (ShowOnlyInnerProperties, ClampMin ="0", ClampMax = "1", UIMin = "0", UIMax = "1"))
+	UPROPERTY(EditAnywhere, Category = "Bone Reduction", meta = (ShowOnlyInnerProperties, ClampMin ="0", ClampMax = "1", UIMin = "0", UIMax = "1", Tooltip="Controls how aggressively bones are reduced.\n\nControls how aggressively bones are reduced. Higher values result in less bones. Reduced bone complexity helps in optimizing performance but can effect the accuracy of wind simulation."))
 	float Strength;
 };
 
diff --git a/Engine/Plugins/Experimental/ProceduralVegetationEditor/Source/ProceduralVegetation/Private/Nodes/PVFoliageDistributorSettings.h b/Engine/Plugins/Experimental/ProceduralVegetationEditor/Source/ProceduralVegetation/Private/Nodes/PVFoliageDistributorSettings.h
--- a/Engine/Plugins/Experimental/ProceduralVegetationEditor/Source/ProceduralVegetation/Private/Nodes/PVFoliageDistributorSettings.h
+++ b/Engine/Plugins/Experimental/ProceduralVegetationEditor/Source/ProceduralVegetation/Private/Nodes/PVFoliageDistributorSettings.h
@@ -40,76 +40,76 @@ protected:
 
 private:
 
-	UPROPERTY(EditAnywhere, Category="Distribution Settings", meta=(ClampMin=0.0f, ClampMax=1.0f, UIMin=0.0f, UIMax=1.0f))
+	UPROPERTY(EditAnywhere, Category="Distribution Settings", meta=(ClampMin=0.0f, ClampMax=1.0f, UIMin=0.0f, UIMax=1.0f, Tooltip="Threshold controlling foliage release based on an ethylene-like signal.\n\nHigher Threshold → more buds/branches retained, making the plant denser and bushier.\nLower Threshold → fewer buds/branches retained, resulting in a sparser, pruned structure."))
 	float EthyleneThreshold = 0.85f;
 
-	UPROPERTY(EditAnywhere, Category="Distribution Settings")
+	UPROPERTY(EditAnywhere, Category="Distribution Settings", meta=(Tooltip="Override default placement with custom distribution rules.\n\nOverides the distribution settings loaded from the procedural vegetation preset."))
 	bool OverrideDistribution = true;
 
-	UPROPERTY(EditAnywhere, Category="Distribution Settings", meta=(EditCondition="OverrideDistribution", ClampMin=0.0f, ClampMax=1.0f, UIMin=0.0f, UIMax=1.0f))
+	UPROPERTY(EditAnywhere, Category="Distribution Settings", meta=(EditCondition="OverrideDistribution", ClampMin=0.0f, ClampMax=1.0f, UIMin=0.0f, UIMax=1.0f, Tooltip="Minimum distance between placed instances.\n\nHigher spacing yields a more open distribution; lower spacing allows denser placement."))
 	float InstanceSpacing = 0.3f;
 
-	UPROPERTY(EditAnywhere, Category="Distribution Settings", meta=(EditCondition="OverrideDistribution", XAxisMin=0.0f, XAxisMax=1.0f, YAxisMin=0.0f, YAxisMax=1.0f))
+	UPROPERTY(EditAnywhere, Category="Distribution Settings", meta=(EditCondition="OverrideDistribution", XAxisMin=0.0f, XAxisMax=1.0f, YAxisMin=0.0f, YAxisMax=1.0f, Tooltip="Ramp that varies instance spacing across the height of the plant\n\nUseful to reduce spacing in some areas based on the height of the plant."))
 	FPVFloatRamp InstanceSpacingRamp;
 
-	UPROPERTY(EditAnywhere, Category="Distribution Settings", meta=(EditCondition="OverrideDistribution", ClampMin=0.0f, ClampMax=1.0f, UIMin=0.0f, UIMax=1.0f))
+	UPROPERTY(EditAnywhere, Category="Distribution Settings", meta=(EditCondition="OverrideDistribution", ClampMin=0.0f, ClampMax=1.0f, UIMin=0.0f, UIMax=1.0f, Tooltip="Strength of the spacing ramp’s influence.\n\nBlends the ramp with the base spacing. Lower values favor the base spacing (subtle effect); higher values apply the ramp more strongly and can amplify spacing differences."))
 	float InstanceSpacingRampEffect = 0.0f;
 
-	UPROPERTY(EditAnywhere, Category="Distribution Settings", meta=(EditCondition="OverrideDistribution", ClampMin=-1))
+	UPROPERTY(EditAnywhere, Category="Distribution Settings", meta=(EditCondition="OverrideDistribution", ClampMin=-1, Tooltip="Caps the number of instances per branch\n\nlimit on how many foliage instances can be placed per branch to avoid overcrowding."))
 	int32 MaxPerBranch = -1;
 
-	UPROPERTY(EditAnywhere, Category="Scale Settings", meta=(EditCondition="OverrideDistribution", ClampMin=0.0f, ClampMax=2.0f, UIMin=0.0f, UIMax=2.0f))
+	UPROPERTY(EditAnywhere, Category="Scale Settings", meta=(EditCondition="OverrideDistribution", ClampMin=0.0f, UIMin=0.0f, UIMax=5.0f, Tooltip="Default size multiplier for placed foliage\n\nSets the base uniform scale applied to each instance before randomness and ramps are considered. Increase to enlarge all instances proportionally; decrease for smaller foliage. Final size still respects Min/Max limits."))
 	float BaseScale = 1.0f;
 
-	UPROPERTY(EditAnywhere, Category="Scale Settings", meta=(EditCondition="OverrideDistribution", ClampMin=0.0f, ClampMax=1.0f, UIMin=0.0f, UIMax=1.0f))
+	UPROPERTY(EditAnywhere, Category="Scale Settings", meta=(EditCondition="OverrideDistribution", ClampMin=0.0f, ClampMax=1.0f, UIMin=0.0f, UIMax=1.0f, Tooltip="Impact of branch size on foliage size\n\nThe size of the foliage being scattered will be chosen in relation to the physical size of the branch it is distributed on. The largest branches will be distributed the largest foliage pieces."))
 	float BranchScaleImpact = 0.0f;
 
-	UPROPERTY(EditAnywhere, Category="Scale Settings", meta=(EditCondition="OverrideDistribution", ClampMin=0.0f, ClampMax=1.0f, UIMin=0.0f, UIMax=1.0f))
+	UPROPERTY(EditAnywhere, Category="Scale Settings", meta=(EditCondition="OverrideDistribution", ClampMin=0.0f, UIMin=0.0f, UIMax=5.0f, Tooltip="Minimum allowed instance size.\n\nClamps the final computed scale so it never goes below this value. Use to prevent tiny, hard-to-see instances."))
 	float MinScale = 0.5f;
 
-	UPROPERTY(EditAnywhere, Category="Scale Settings", meta=(EditCondition="OverrideDistribution", ClampMin=0.0f, ClampMax=1.0f, UIMin=0.0f, UIMax=1.0f))
+	UPROPERTY(EditAnywhere, Category="Scale Settings", meta=(EditCondition="OverrideDistribution", ClampMin=0.0f, UIMin=0.0f, UIMax=5.0f, Tooltip="Maximum allowed instance size.\n\nClamps the final computed scale so it never exceeds this value. Use to cap unusually large instances for visual consistency."))
 	float MaxScale = 1.0f;
 
-	UPROPERTY(EditAnywhere, Category="Scale Settings", meta=(EditCondition="OverrideDistribution", ClampMin=0.0f, ClampMax=1.0f, UIMin=0.0f, UIMax=1.0f))
+	UPROPERTY(EditAnywhere, Category="Scale Settings", meta=(EditCondition="OverrideDistribution", ClampMin=0.0f, UIMin=0.0f, UIMax=5.0f, Tooltip="Lower bound for per-instance size variation.\n\nSets the minimum of the random scale range used to vary instance sizes. Each instance samples within Random Scale Min–Max, then combines with Base Scale and any ramps (respecting Min/Max Scale)."))
 	float RandomScaleMin = 1.0f;
 
-	UPROPERTY(EditAnywhere, Category="Scale Settings", meta=(EditCondition="OverrideDistribution", ClampMin=1.0f, ClampMax=2.0f, UIMin=1.0f, UIMax=2.0f))
+	UPROPERTY(EditAnywhere, Category="Scale Settings", meta=(EditCondition="OverrideDistribution", ClampMin=1.0f, UIMin=1.0f, UIMax=5.0f, Tooltip="Upper bound for per-instance size variation.\n\nSets the maximum of the random scale range used to vary instance sizes. Each instance samples within Random Scale Min–Max, then combines with Base Scale and any ramps (respecting Min/Max Scale)."))
 	float RandomScaleMax = 1.0f;
 
-	UPROPERTY(EditAnywhere, Category="Scale Settings", meta=(EditCondition="OverrideDistribution", XAxisMin=0.0f, XAxisMax=1.0f, YAxisMin=0.0f, YAxisMax=1.0f))
+	UPROPERTY(EditAnywhere, Category="Scale Settings", meta=(EditCondition="OverrideDistribution", XAxisMin=0.0f, XAxisMax=1.0f, YAxisMin=0.0f, YAxisMax=1.0f, Tooltip="Curve that varies scale across the height of the plant\n\nUse to taper foliage toward tips or emphasize size in targeted regions."))
 	FPVFloatRamp ScaleRamp;
 
-	UPROPERTY(EditAnywhere, Category="Vector Settings", meta=(EditCondition="OverrideDistribution"))
+	UPROPERTY(EditAnywhere, Category="Vector Settings", meta=(EditCondition="OverrideDistribution", Tooltip="Use a custom axil angle instead of the default.\n\nEnables a user-defined orientation. When on, the system uses AxilAngle and the AxilAngleRamp settings (and Effect) to set the tilt of foliage relative to its parent branch direction, overriding any automatic/species defaults."))
 	bool OverrideAxilAngle = true;
 
-	UPROPERTY(EditAnywhere, Category="Vector Settings", meta=(EditCondition="OverrideDistribution", ClampMin=0.0f, ClampMax=90.0f, UIMin=0.0f, UIMax=90.0f))
+	UPROPERTY(EditAnywhere, Category="Vector Settings", meta=(EditCondition="OverrideDistribution", ClampMin=0.0f, ClampMax=90.0f, UIMin=0.0f, UIMax=90.0f, Tooltip="Base tilt from the parent axis\n\nSets the base axil angle in degrees—the tilt of the instance relative to the parent branch’s direction. Acts as the lower bound when used with a ramp; higher values tilt foliage further away from the branch axis."))
 	float AxilAngle = 45.0f;
 
-	UPROPERTY(EditAnywhere, Category="Vector Settings", meta=(EditCondition="OverrideDistribution", XAxisMin=0.0f, XAxisMax=1.0f, YAxisMin=0.0f, YAxisMax=1.0f))
+	UPROPERTY(EditAnywhere, Category="Vector Settings", meta=(EditCondition="OverrideDistribution", XAxisMin=0.0f, XAxisMax=1.0f, YAxisMin=0.0f, YAxisMax=1.0f, Tooltip="Ramp that varies the axil angle across the height of the plant\n\nIts output is remapped to an angle range defined by AxilAngle and AxilAngleRampUpperValue, modulated by AxilAngleRampEffect."))
 	FPVFloatRamp AxilAngleRamp;
 
-	UPROPERTY(EditAnywhere, Category="Vector Settings", meta=(EditCondition="OverrideDistribution", ClampMin=0.0f, ClampMax=90.0f, UIMin=0.0f, UIMax=90.0f))
+	UPROPERTY(EditAnywhere, Category="Vector Settings", meta=(EditCondition="OverrideDistribution", ClampMin=0.0f, ClampMax=90.0f, UIMin=0.0f, UIMax=90.0f, Tooltip="Max angle used when the ramp\n\nProvides detailed control over this parameter to balance visual quality and performance. Valid values: 0.0–90.0."))
 	float AxilAngleRampUpperValue = 45.0f;
 
-	UPROPERTY(EditAnywhere, Category="Vector Settings", meta=(EditCondition="OverrideDistribution", ClampMin=0.0f, ClampMax=1.0f, UIMin=0.0f, UIMax=1.0f))
+	UPROPERTY(EditAnywhere, Category="Vector Settings", meta=(EditCondition="OverrideDistribution", ClampMin=0.0f, ClampMax=1.0f, UIMin=0.0f, UIMax=1.0f, Tooltip="Adjusts this parameter for finer control.\n\nDefines the upper axil angle reached when AxilAngleRamp evaluates to 1. Together with AxilAngle, it forms the min–max range for ramp-driven orientation."))
 	float AxilAngleRampEffect = 0.0f;
 
-	UPROPERTY(EditAnywhere, Category="Phyllotaxy Settings", meta=(EditCondition="OverrideDistribution"))
+	UPROPERTY(EditAnywhere, Category="Phyllotaxy Settings", meta=(EditCondition="OverrideDistribution", Tooltip="Use custom phyllotaxy settings instead of defaults.\n\nOveride the phyllotaxy settings loaded from the procedural vegetation preset."))
 	bool OverridePhyllotaxy = false;
 
-	UPROPERTY(EditAnywhere, Category="Phyllotaxy Settings", meta=(EditCondition="OverrideDistribution"))
+	UPROPERTY(EditAnywhere, Category="Phyllotaxy Settings", meta=(EditCondition="OverrideDistribution", Tooltip="Pattern for arranging buds/leaves around the stem.\n\nSelects the phyllotactic pattern used for placement (e.g., Alternate, Opposite, Whorled, Spiral). This defines the sequence and angular spacing of buds/leaves around the stem and along its length."))
 	EPhyllotaxyType PhyllotaxyType = EPhyllotaxyType::Alternate;
 
 	UPROPERTY(EditAnywhere, Category="Phyllotaxy Settings", meta=(EditCondition="OverrideDistribution && PhyllotaxyType == EPhyllotaxyType::Spiral", EditConditionHides))
 	EPhyllotaxyFormation PhyllotaxyFormation = EPhyllotaxyFormation::Distichous;
 
-	UPROPERTY(EditAnywhere, Category="Phyllotaxy Settings", meta=(EditCondition="OverrideDistribution", ClampMin=1, ClampMax=10, UIMin=1, UIMax=10))
+	UPROPERTY(EditAnywhere, Category="Phyllotaxy Settings", meta=(EditCondition="OverrideDistribution", ClampMin=1, ClampMax=10, UIMin=1, UIMax=10, Tooltip="Minimum buds/leaves per node.\n\nSets the lower bound on the number of buds/leaves generated at each node. Guarantees at least this many instances (subject to spacing and other constraints). Should be ≤ MaximumNodeBuds. Applied when Whorled phyllotaxy is selected."))
 	int32 MinimumNodeBuds = 2;
 
-	UPROPERTY(EditAnywhere, Category="Phyllotaxy Settings", meta=(EditCondition="OverrideDistribution", ClampMin=1, ClampMax=10, UIMin=1, UIMax=10))
+	UPROPERTY(EditAnywhere, Category="Phyllotaxy Settings", meta=(EditCondition="OverrideDistribution", ClampMin=1, ClampMax=10, UIMin=1, UIMax=10, Tooltip="Maximum buds/leaves per node.\n\nCaps the number of buds/leaves spawned. Use with MinimumNodeBuds to bound density and keep arrangements predictable. Applied when Whorled phyllotaxy is selected."))
 	int32 MaximumNodeBuds = 3;
 
-	UPROPERTY(EditAnywhere, Category="Phyllotaxy Settings", meta=(EditCondition="OverrideDistribution", ClampMin=0.0f, ClampMax=360.0f, UIMin=0.0f, UIMax=360.0f))
+	UPROPERTY(EditAnywhere, Category="Phyllotaxy Settings", meta=(EditCondition="OverrideDistribution", ClampMin=0.0f, ClampMax=360.0f, UIMin=0.0f, UIMax=360.0f, Tooltip="Extra divergence angle (°) added to the phyllotaxy pattern\n\nAdds an angular offset to the base phyllotactic divergence."))
 	float PhyllotaxyAdditionalAngle = 0.0f;
 
 	UPROPERTY(EditAnywhere, Category="Misc Settings", meta=(EditCondition="OverrideDistribution"))
diff --git a/Engine/Plugins/Experimental/ProceduralVegetationEditor/Source/ProceduralVegetation/Private/Nodes/PVFoliagePaletteSettings.h b/Engine/Plugins/Experimental/ProceduralVegetationEditor/Source/ProceduralVegetation/Private/Nodes/PVFoliagePaletteSettings.h
--- a/Engine/Plugins/Experimental/ProceduralVegetationEditor/Source/ProceduralVegetation/Private/Nodes/PVFoliagePaletteSettings.h
+++ b/Engine/Plugins/Experimental/ProceduralVegetationEditor/Source/ProceduralVegetation/Private/Nodes/PVFoliagePaletteSettings.h
@@ -32,7 +32,7 @@ public:
 	UPROPERTY(EditAnywhere, Category="Meshes", meta=(PinHiddenByDefault, InlineEditConditionToggle))
 	bool bOverrideFoliage = false;
 
-	UPROPERTY(EditAnywhere, Category = "Meshes", Meta = (EditCondition="bOverrideFoliage", AllowedClasses = "/Script/Engine.StaticMesh,/Script/Engine.SkeletalMesh"))
+	UPROPERTY(EditAnywhere, Category = "Meshes", Meta = (EditCondition="bOverrideFoliage", AllowedClasses = "/Script/Engine.StaticMesh,/Script/Engine.SkeletalMesh", Tooltip="List of Foliage meshes used for this plant\n\nList of all foliage meshes applied based on procedural vegetation preset."))
 	TArray<TSoftObjectPtr<UObject>> FoliageMeshes;
 
 };
diff --git a/Engine/Plugins/Experimental/ProceduralVegetationEditor/Source/ProceduralVegetation/Public/Helpers/PVExportHelper.h b/Engine/Plugins/Experimental/ProceduralVegetationEditor/Source/ProceduralVegetation/Public/Helpers/PVExportHelper.h
--- a/Engine/Plugins/Experimental/ProceduralVegetationEditor/Source/ProceduralVegetation/Public/Helpers/PVExportHelper.h
+++ b/Engine/Plugins/Experimental/ProceduralVegetationEditor/Source/ProceduralVegetation/Public/Helpers/PVExportHelper.h
@@ -30,19 +30,19 @@ struct PROCEDURALVEGETATION_API FPVExportParams
 	UPROPERTY()
 	bool bShouldExport = true;
 
-	UPROPERTY(EditAnywhere, Category = "Export", Meta = (ContentDir, DisplayName = "Content Browser Folder"))
+	UPROPERTY(EditAnywhere, Category = "Export", Meta = (ContentDir, DisplayName = "Content Browser Folder", Tooltip="Path to the folder where exported mesh will be saved\n\nPath to the folder where exported mesh will be saved."))
 	FDirectoryPath ContentBrowserFolder;
 
-	UPROPERTY(EditAnywhere, Category = "Export", Meta = (DisplayName = "Mesh Name"))
+	UPROPERTY(EditAnywhere, Category = "Export", Meta = (DisplayName = "Mesh Name", Tooltip="Name of the exported mesh\n\nName of the exported mesh."))
 	FName MeshName;
 
-	UPROPERTY(EditAnywhere, Category = "Export", Meta = (DisplayName = "Export Mesh Type"))
+	UPROPERTY(EditAnywhere, Category = "Export", Meta = (DisplayName = "Export Mesh Type", Tooltip="Static Mesh or Skeletal Mesh."))
 	EPVExportMeshType ExportMeshType = EPVExportMeshType::StaticMesh;
 
-	UPROPERTY(EditAnywhere, Category = "Export", Meta = (DisplayName = "Create Nanite Foliage"))
+	UPROPERTY(EditAnywhere, Category = "Export", Meta = (DisplayName = "Create Nanite Foliage", Tooltip="Add Nanite foliage support for the exported mesh."))
 	bool bCreateNaniteFoliage = true;
 	
-	UPROPERTY(EditAnywhere, Category = "Export", Meta = (DisplayName = "Nanite Shape Preservation Method"))
+	UPROPERTY(EditAnywhere, Category = "Export", Meta = (DisplayName = "Nanite Shape Preservation Method", Tooltip="Shape preservation setting to be used in nanite settings."))
 	ENaniteShapePreservation NaniteShapePreservation = ENaniteShapePreservation::Voxelize;
 
 	void Initialize(const FString& InAssetPath, const FString& InName)
