Fix foliage missing from exported mesh if nanite foliage is disabled in the project settings but enabled on the node

diff --git a/Engine/Plugins/Experimental/ProceduralVegetationEditor/Source/ProceduralVegetation/Private/Helpers/PVExportHelper.cpp b/Engine/Plugins/Experimental/ProceduralVegetationEditor/Source/ProceduralVegetation/Private/Helpers/PVExportHelper.cpp
--- a/Engine/Plugins/Experimental/ProceduralVegetationEditor/Source/ProceduralVegetation/Private/Helpers/PVExportHelper.cpp
+++ b/Engine/Plugins/Experimental/ProceduralVegetationEditor/Source/ProceduralVegetation/Private/Helpers/PVExportHelper.cpp
@@ -1144,7 +1144,7 @@ namespace PV::Export::Internal
 
 namespace PV::Export
 {
-	void ExportCollectionAsMesh(
+		void ExportCollectionAsMesh(
 		const TObjectPtr<UProceduralVegetation> InProceduralVegetation,
 		const FManagedArrayCollection& Collection,
 		const FManagedArrayCollection& FoliageCollection,
@@ -1160,6 +1160,8 @@ namespace PV::Export
 
 		const FString FullPath = FPaths::Combine(PackageName, MeshName);
 
+		const bool bCreateNaniteFoliage = ExportParams.bCreateNaniteFoliage && NaniteAssembliesSupported();
+
 		if (ExportParams.ExportMeshType == EPVExportMeshType::StaticMesh)
 		{
 			UPackage* MeshPackage = CreatePackage(*FullPath);
@@ -1172,7 +1174,7 @@ namespace PV::Export
 			PV::Export::Internal::ExportCollectionToStaticMesh(
 				ExportMesh,
 				Collection,
-				ExportParams.bCreateNaniteFoliage,
+				bCreateNaniteFoliage,
 				ExportParams.NaniteShapePreservation,
 				StatusReportCallback
 			);
@@ -1219,7 +1221,7 @@ namespace PV::Export
 				ExportMesh,
 				Collection,
 				ExportParams.NaniteShapePreservation,
-				ExportParams.bCreateNaniteFoliage,
+				bCreateNaniteFoliage,
 				StatusReportCallback
 			);
 
