[PVE] Merge material slots

diff --git a/Engine/Plugins/Experimental/ProceduralVegetationEditor/Source/ProceduralVegetation/Private/Implementations/PVMeshBuilder.cpp b/Engine/Plugins/Experimental/ProceduralVegetationEditor/Source/ProceduralVegetation/Private/Implementations/PVMeshBuilder.cpp
--- a/Engine/Plugins/Experimental/ProceduralVegetationEditor/Source/ProceduralVegetation/Private/Implementations/PVMeshBuilder.cpp
+++ b/Engine/Plugins/Experimental/ProceduralVegetationEditor/Source/ProceduralVegetation/Private/Implementations/PVMeshBuilder.cpp
@@ -13,6 +13,34 @@
 #include "Facades/PVPlantFacade.h"
 #include "Facades/PVPointFacade.h"
 
+void MergeMaterials(const TArray<FTrunkGenerationMaterialSetup>& InMaterialSetups, TArray<FLocalDynamicMeshData>& MeshDatas)
+{
+	TMap<FString, int32> MaterialPathIDMap;
+	for (auto& MeshData : MeshDatas)
+	{
+		TObjectPtr<UMaterialInterface> Material;
+		if (InMaterialSetups.IsValidIndex(MeshData.MaterialID))
+		{
+			Material = InMaterialSetups[MeshData.MaterialID].Material;
+		}
+
+		FString MaterialPathName = "";
+		if (Material)
+		{
+			MaterialPathName = Material->GetPathName();
+		}
+
+		if (!MaterialPathIDMap.Contains(MaterialPathName))
+		{
+			MaterialPathIDMap.Add(MaterialPathName);
+		}
+		else
+		{
+			MeshData.MaterialID = MaterialPathIDMap[MaterialPathName];
+		}
+	}
+}
+
 void FPVMeshBuilder::GenerateGeometryCollection(const FManagedArrayCollection& InSkeletonCollection, const FPVMeshBuilderParams& MeshBuilderParams,
                                            FGeometryCollection& OutGeometryCollection)
 {
@@ -75,6 +103,9 @@ void FPVMeshBuilder::GenerateGeometryCollection(const FManagedArrayCollection& I
 	int32 FaceIndex = 0;
 	int32 PrevMaterialID = -1;
 	TArray<FSection> Sections;
+
+	//Merge MaterialIds based on material Paths
+	MergeMaterials(MaterialSetups, MeshesDataCollection);
 	
 	//Sorting by material id so sections are created according to material ids
 	MeshesDataCollection.Sort([](const FLocalDynamicMeshData& A, const FLocalDynamicMeshData& B)
