[PVE] Added Wind Settings

#virtualized

diff --git a/Engine/Plugins/Experimental/ProceduralVegetationEditor/Source/ProceduralVegetation/Private/Helpers/PVExportHelper.cpp b/Engine/Plugins/Experimental/ProceduralVegetationEditor/Source/ProceduralVegetation/Private/Helpers/PVExportHelper.cpp
--- a/Engine/Plugins/Experimental/ProceduralVegetationEditor/Source/ProceduralVegetation/Private/Helpers/PVExportHelper.cpp
+++ b/Engine/Plugins/Experimental/ProceduralVegetationEditor/Source/ProceduralVegetation/Private/Helpers/PVExportHelper.cpp
@@ -32,6 +32,7 @@
 #include "Misc/Paths.h"
 #include "UObject/Package.h"
 #include "ProceduralVegetation.h"
+#include "PVWindSettings.h"
 #include "SkinnedAssetCompiler.h"
 #include "StaticMeshCompiler.h"
 #include "GeometryCollection/Facades/CollectionConstraintOverrideFacade.h"
@@ -801,7 +802,8 @@ namespace PV::Export::Internal
 	void ExportCollectionToSkeletalMesh(
 		TObjectPtr<USkeletalMesh> ExportMesh,
 		const FManagedArrayCollection& Collection,
-		ENaniteShapePreservation ShapePreservation,
+		const ENaniteShapePreservation& InShapePreservation,
+		const TObjectPtr<UPVWindSettings> InWindSettings,
 		bool bBuildNaniteAssemblies,
 		const FStatusReportCallback& StatusReportCallback
 	)
@@ -809,9 +811,9 @@ namespace PV::Export::Internal
 		TRACE_CPUPROFILER_EVENT_SCOPE(PV::Export::ExportCollectionToSkeletalMesh);
 
 		const auto UpdateStatus = [&](EExportMeshStatus NewStatus)
-		{
-			StatusReportCallback(CalculateExpectedWork(NewStatus, bBuildNaniteAssemblies, true), GetStatusText(NewStatus, bBuildNaniteAssemblies, true));
-		};
+			{
+				StatusReportCallback(CalculateExpectedWork(NewStatus, bBuildNaniteAssemblies, true), GetStatusText(NewStatus, bBuildNaniteAssemblies, true));
+			};
 
 		FManagedArrayCollection ExportedCollection;
 		TObjectPtr<UDynamicMesh> DynamicMesh;
@@ -903,7 +905,8 @@ namespace PV::Export::Internal
 			UpdateStatus(EExportMeshStatus::CopyingMeshDataToOutputMesh);
 
 #if WITH_EDITORONLY_DATA
-			FGeometryScriptCopyMeshToAssetOptions Options = GetCopyMeshToAssetOptions(Materials, ExportMesh->GetNaniteSettings(), ShapePreservation);
+			FGeometryScriptCopyMeshToAssetOptions Options = GetCopyMeshToAssetOptions(Materials, ExportMesh->GetNaniteSettings(), InShapePreservation);
+		
 #else
 			FGeometryScriptCopyMeshToAssetOptions Options = GetCopyMeshToAssetOptions(Materials);
 #endif
@@ -985,17 +988,34 @@ namespace PV::Export::Internal
 		{
 			TRACE_CPUPROFILER_EVENT_SCOPE(DynamicWind);
 
+			bool bWindDataExist = ExportMesh->HasAssetUserDataOfClass(UDynamicWindSkeletalData::StaticClass());
+			bool bOverwriteWindSettings = InWindSettings ? InWindSettings->bOverwriteExisting : false;
 			UpdateStatus(EExportMeshStatus::ImportingDynamicWind);
-
+			
 			FDynamicWindSkeletalImportData ImportData;
 			BuildWindImportData(ImportData, ExportedCollection, BoneNodes);
 
+			if (InWindSettings)
+			{
+				ImportData.SimulationGroups = InWindSettings->SimulationGroupData;
+			}
+		
+			if (bWindDataExist && !bOverwriteWindSettings)
+			{
+				//Get Existing WindSettings
+				UDynamicWindSkeletalData* WindSkeletalData = Cast<UDynamicWindSkeletalData>(ExportMesh->GetAssetUserDataOfClass(UDynamicWindSkeletalData::StaticClass()));
+				if (WindSkeletalData)
+				{
+					ImportData.SimulationGroups = WindSkeletalData->SimulationGroups;
+				}
+			}
+
 			if (!DynamicWind::ImportSkeletalData(*ExportMesh, ImportData))
 			{
 				UE_LOG(LogProceduralVegetation, Warning, TEXT("Failed to build wind data"));
 			}
 		}
-		
+
 		{
 			TRACE_CPUPROFILER_EVENT_SCOPE(BuildSkinnedMesh);
 			UpdateStatus(EExportMeshStatus::BuildingMesh);
@@ -1221,6 +1241,7 @@ namespace PV::Export
 				ExportMesh,
 				Collection,
 				ExportParams.NaniteShapePreservation,
+				ExportParams.WindSettings,
 				bCreateNaniteFoliage,
 				StatusReportCallback
 			);
diff --git a/Engine/Plugins/Experimental/ProceduralVegetationEditor/Source/ProceduralVegetation/Private/Nodes/PVOutputSettings.cpp b/Engine/Plugins/Experimental/ProceduralVegetationEditor/Source/ProceduralVegetation/Private/Nodes/PVOutputSettings.cpp
--- a/Engine/Plugins/Experimental/ProceduralVegetationEditor/Source/ProceduralVegetation/Private/Nodes/PVOutputSettings.cpp
+++ b/Engine/Plugins/Experimental/ProceduralVegetationEditor/Source/ProceduralVegetation/Private/Nodes/PVOutputSettings.cpp
@@ -40,7 +40,8 @@ UPVOutputSettings::UPVOutputSettings()
 	{
 		static FString DefaultMeshName = "ProceduralVegetationOutput";
 		static FString DefaultAssetPath = "/Game/ProceduralVegetation/";
-		ExportSettings.Initialize(DefaultAssetPath, DefaultMeshName);
+		static FString DefaultWindSettingsPath = "/ProceduralVegetationEditor/SampleAssets/WindSettings/DefaultTreeWindSettings.DefaultTreeWindSettings";
+		ExportSettings.Initialize(DefaultAssetPath, DefaultMeshName, DefaultWindSettingsPath);
 	}
 }
 
diff --git a/Engine/Plugins/Experimental/ProceduralVegetationEditor/Source/ProceduralVegetation/ProceduralVegetation.Build.cs b/Engine/Plugins/Experimental/ProceduralVegetationEditor/Source/ProceduralVegetation/ProceduralVegetation.Build.cs
--- a/Engine/Plugins/Experimental/ProceduralVegetationEditor/Source/ProceduralVegetation/ProceduralVegetation.Build.cs
+++ b/Engine/Plugins/Experimental/ProceduralVegetationEditor/Source/ProceduralVegetation/ProceduralVegetation.Build.cs
@@ -11,6 +11,7 @@ namespace UnrealBuildTool.Rules
 			PublicDependencyModuleNames.AddRange(
 				new string[]
 				{
+					"DynamicWind"
 				});
 
 			PrivateDependencyModuleNames.AddRange(
@@ -32,8 +33,7 @@ namespace UnrealBuildTool.Rules
 					"SkeletalMeshDescription",
 					"StaticMeshDescription",
 					"ModelingComponents",
-					"PCG",
-					"DynamicWind"// ... add private dependencies that you statically link with here ...	
+					"PCG", // ... add private dependencies that you statically link with here ...	
 				}
 			);
 			
diff --git a/Engine/Plugins/Experimental/ProceduralVegetationEditor/Source/ProceduralVegetation/Public/Helpers/PVExportHelper.h b/Engine/Plugins/Experimental/ProceduralVegetationEditor/Source/ProceduralVegetation/Public/Helpers/PVExportHelper.h
--- a/Engine/Plugins/Experimental/ProceduralVegetationEditor/Source/ProceduralVegetation/Public/Helpers/PVExportHelper.h
+++ b/Engine/Plugins/Experimental/ProceduralVegetationEditor/Source/ProceduralVegetation/Public/Helpers/PVExportHelper.h
@@ -7,6 +7,7 @@
 #include "GeometryCollection/ManagedArrayCollection.h"
 #include "Materials/MaterialInterface.h"
 #include "Facades/PVBoneFacade.h"
+#include "PVWindSettings.h"
 #include "PVExportHelper.generated.h"
 
 class UInstancedStaticMeshComponent;
@@ -37,7 +38,7 @@ struct PROCEDURALVEGETATION_API FPVExportParams
 	FName MeshName;
 
 	UPROPERTY(EditAnywhere, Category = "Export", Meta = (DisplayName = "Export Mesh Type", Tooltip="Static Mesh or Skeletal Mesh."))
-	EPVExportMeshType ExportMeshType = EPVExportMeshType::StaticMesh;
+	EPVExportMeshType ExportMeshType = EPVExportMeshType::SkeletalMesh;
 
 	UPROPERTY(EditAnywhere, Category = "Export", Meta = (DisplayName = "Create Nanite Foliage", Tooltip="Add Nanite foliage support for the exported mesh."))
 	bool bCreateNaniteFoliage = true;
@@ -45,12 +46,17 @@ struct PROCEDURALVEGETATION_API FPVExportParams
 	UPROPERTY(EditAnywhere, Category = "Export", Meta = (DisplayName = "Nanite Shape Preservation Method", Tooltip="Shape preservation setting to be used in nanite settings."))
 	ENaniteShapePreservation NaniteShapePreservation = ENaniteShapePreservation::Voxelize;
 
-	void Initialize(const FString& InAssetPath, const FString& InName)
+	UPROPERTY(EditAnywhere, Category = "Export", Meta = (EditCondition = "ExportMeshType == EPVExportMeshType::SkeletalMesh"))
+	TObjectPtr<UPVWindSettings> WindSettings;
+
+	void Initialize(const FString& InAssetPath, const FString& InName, const FString& InDefaultWindSettingsPath)
 	{
 		FString DefaultDirectory = FPaths::GetPath(InAssetPath);
 		ContentBrowserFolder.Path = DefaultDirectory;
 		
 		MeshName = FName(InName);
+
+		WindSettings = LoadObject<UPVWindSettings>(nullptr, InDefaultWindSettingsPath);
 	}
 
 	FString GetOutputObjectPath() const
@@ -109,7 +115,7 @@ namespace PV::Export::Internal
 	void ExportCollectionToSkeletalMesh(
 		TObjectPtr<USkeletalMesh> ExportMesh,
 		const FManagedArrayCollection& Collection,
-		ENaniteShapePreservation ShapePreservation,
+		const FPVExportParams& ExportParams,
 		bool bBuildNaniteAssemblies
 	);
 
diff --git a/Engine/Plugins/Experimental/ProceduralVegetationEditor/Source/ProceduralVegetation/Public/PVWindSettings.h b/Engine/Plugins/Experimental/ProceduralVegetationEditor/Source/ProceduralVegetation/Public/PVWindSettings.h
new file mode 100644
--- /dev/null
+++ b/Engine/Plugins/Experimental/ProceduralVegetationEditor/Source/ProceduralVegetation/Public/PVWindSettings.h
@@ -0,0 +1,20 @@
+// Copyright Epic Games, Inc. All Rights Reserved.
+
+#pragma once
+
+#include "DynamicWindSkeletalData.h"
+#include "Engine/DataAsset.h"
+#include "PVWindSettings.generated.h"
+
+UCLASS(MinimalAPI, BlueprintType)
+class UPVWindSettings : public UDataAsset
+{
+	GENERATED_BODY()
+
+public:
+	UPROPERTY(EditAnywhere, Category="Wind Settings")
+	bool bOverwriteExisting = false;
+	
+	UPROPERTY(EditAnywhere, Category="Wind Settings")
+	TArray<FDynamicWindSimulationGroupData> SimulationGroupData;
+};
\ No newline at end of file
