Add cancel button to PVE Export

- Add progress reporting to CombineMeshInstancesToDynamicMesh
- Add modal warning the user if they are about to overwrite unsaved changes (this is needed as the cancel functionality may need to restore the asset from disk)
- Refactor PVExportHelpers
- Move PVExporHelpers to Editor module
- Formatting fixes
- BugHawk fixes
- Move PVExportParams from PVExportHelpers to PVExportParams.h/.cpp
- Re-write BuildNaniteAssemblyData such that the assets created by that function can be unloaded if the export is canceled

diff --git a/Engine/Plugins/Experimental/ProceduralVegetationEditor/Source/ProceduralVegetation/Private/Helpers/PVExportHelper.cpp b/Engine/Plugins/Experimental/ProceduralVegetationEditor/Source/ProceduralVegetation/Private/Helpers/PVExportHelper.cpp
deleted file mode 100644
diff --git a/Engine/Plugins/Experimental/ProceduralVegetationEditor/Source/ProceduralVegetation/Private/Helpers/PVUtilities.cpp b/Engine/Plugins/Experimental/ProceduralVegetationEditor/Source/ProceduralVegetation/Private/Helpers/PVUtilities.cpp
--- a/Engine/Plugins/Experimental/ProceduralVegetationEditor/Source/ProceduralVegetation/Private/Helpers/PVUtilities.cpp
+++ b/Engine/Plugins/Experimental/ProceduralVegetationEditor/Source/ProceduralVegetation/Private/Helpers/PVUtilities.cpp
@@ -10,7 +10,6 @@
 #include "Components/InstancedStaticMeshComponent.h"
 #include "Components/MeshComponent.h"
 #include "Nodes/PVImporterSettings.h"
-#include "Helpers/PVExportHelper.h"
 #include "Misc/PackageName.h"
 #include "AssetRegistry/AssetData.h"
 #include "Misc/FileHelper.h"
@@ -160,17 +159,7 @@ namespace PV::Utilities
 		return FString();
 	}
 
-	TArray<FAssetData> GetExistingAssets(FString PackageName)
-	{
-		// Get the Asset Registry Module
-		FAssetRegistryModule& AssetRegistryModule = FModuleManager::LoadModuleChecked<FAssetRegistryModule>("AssetRegistry");
-		// Get all asset data in the Content Browser
-		TArray<FAssetData> AssetData;
-		AssetRegistryModule.Get().GetAssetsByPackageName(FName(*PackageName), AssetData);
-		return AssetData;
-	}
-
-	bool ValidateMeshPathAndName(FString MeshName, FString Path, UClass* InClass, FString& OutError)
+	bool ValidateAssetPathAndName(const FString& AssetName, const FString& Path, UClass* InClass, FString& OutError)
 	{
 		if (FPackageName::EErrorCode ErrorCode; !FPackageName::IsValidLongPackageName(Path, false, &ErrorCode))
         {
@@ -178,28 +167,18 @@ namespace PV::Utilities
         	return false;
         }
 
-       if (FText Reason; !IsFileNameValid(FName(MeshName), Reason))
+		if (FText Reason; !IsFileNameValid(FName(AssetName), Reason))
         {
 			OutError = Reason.ToString();
         	return false;
         }
 
-        if (IsInGameThread())
+        const FString PackageName = FPaths::Combine(Path, AssetName);
+        if (DoesConflictingPackageExist(PackageName, InClass))
         {
-        	FString PackageName = FPaths::Combine(Path, MeshName);
-        	bool bPackageExists = GetExistingAssets(PackageName).Num() > 0;
-        	
-        	FAssetData Asset;
-        	if (bPackageExists && !CanOverwriteAsset(PackageName, InClass, Asset, OutError))
-        	{
-        		 return false;
-        	}
+			OutError = FString::Format(TEXT("Create a Unique name for the output as another asset of different type {0} exists with same name {1}"), { *InClass->GetName(), *PackageName});
+        	return false;
         }
-		else
-		{
-			OutError = FString(TEXT("ValidateMeshPathAndName needs to be called from Game Thread. Returning validation failed result."));
-			return false;
-		}
 
 		return true;
 	}
@@ -212,35 +191,33 @@ namespace PV::Utilities
 		return Reason.IsEmpty();
 	}
 
-	bool CanOverwriteAsset(FString PackageName, UClass* InClass, FAssetData& OutAsset, FString& OutError)
+	bool DoesConflictingPackageExist(const FString& LongPackageName, UClass* AssetClass)
 	{
-		bool bCanOverwrite = false;
-		
-		TArray<FAssetData> AssetData = GetExistingAssets(PackageName);
-
-		if (AssetData.Num() > 0)
+		UPackage* Package = FindPackage(nullptr, *LongPackageName);
+		if (Package != nullptr)
 		{
-			FString AssetType = "";
-			// Iterate through the asset data
-			for (FAssetData& Asset : AssetData)
+			const FName AssetName = FPackageName::GetShortFName(Package->GetFName());
+			const bool bPackageIsOfSameType = StaticFindObjectFast(AssetClass, Package, AssetName) != nullptr;
+			return !bPackageIsOfSameType;
+		}
+		else
+		{
+			TArray<FAssetData> OutAssets;
+			IAssetRegistry& AssetRegistry = FAssetRegistryModule::GetRegistry();
+			AssetRegistry.GetAssetsByPackageName(*LongPackageName, OutAssets, true);
+
+			if (OutAssets.Num() == 0)
 			{
-				AssetType = Asset.GetClass()->GetName();
-				if (Asset.IsInstanceOf(InClass))
-				{
-					OutAsset = Asset;
-					bCanOverwrite = true;
-					break;
-				}
+				return false;
 			}
-			// It is an error if we are trying to replace an object of a different class
-			if (!bCanOverwrite)
+
+			if (OutAssets.Num() != 1
+				|| OutAssets[0].AssetClassPath != AssetClass->GetClassPathName())
 			{
-				bCanOverwrite = false;
-				OutError = FString::Format(TEXT("Create a Unique name for the output as another asset of different type {0} exists with same name {1}"), {*AssetType, *PackageName});
-				return false;
+				return true;
 			}
 		}
 
-		return bCanOverwrite;
+		return false;
 	}
 }
diff --git a/Engine/Plugins/Experimental/ProceduralVegetationEditor/Source/ProceduralVegetation/Private/Nodes/PVOutputSettings.cpp b/Engine/Plugins/Experimental/ProceduralVegetationEditor/Source/ProceduralVegetation/Private/Nodes/PVOutputSettings.cpp
--- a/Engine/Plugins/Experimental/ProceduralVegetationEditor/Source/ProceduralVegetation/Private/Nodes/PVOutputSettings.cpp
+++ b/Engine/Plugins/Experimental/ProceduralVegetationEditor/Source/ProceduralVegetation/Private/Nodes/PVOutputSettings.cpp
@@ -76,13 +76,10 @@ bool FPVOutputElement::ExecuteInternal(FPCGContext* InContext) const
 	{
 		if (const UPVData* InputData = Cast<UPVData>(Input.Data))
 		{
-			const FString OutputPath = Settings->ExportSettings.ContentBrowserFolder.Path;
-			const FName MeshName = Settings->ExportSettings.MeshName;
-
-			UClass* MeshClass = Settings->ExportSettings.ExportMeshType == EPVExportMeshType::StaticMesh ? UStaticMesh::StaticClass() : USkeletalMesh::StaticClass();
-			if (FString Reason; !PV::Utilities::ValidateMeshPathAndName(MeshName.ToString(), OutputPath, MeshClass, Reason))
+			FString Error;
+			if (!Settings->ExportSettings.Validate(Error))
 			{
-				PCGLog::LogErrorOnGraph(FText::FromString(Reason), InContext);
+				PCGLog::LogErrorOnGraph(FText::FromString(Error), InContext);
 				return true;
 			}
 
diff --git a/Engine/Plugins/Experimental/ProceduralVegetationEditor/Source/ProceduralVegetation/Private/PVExportParams.cpp b/Engine/Plugins/Experimental/ProceduralVegetationEditor/Source/ProceduralVegetation/Private/PVExportParams.cpp
new file mode 100644
--- /dev/null
+++ b/Engine/Plugins/Experimental/ProceduralVegetationEditor/Source/ProceduralVegetation/Private/PVExportParams.cpp
@@ -0,0 +1,29 @@
+// Copyright Epic Games, Inc. All Rights Reserved.
+
+#include "PVExportParams.h"
+#include "Helpers/PVUtilities.h"
+
+#include "Engine/StaticMesh.h"
+#include "Engine/SkeletalMesh.h"
+#include "Animation/Skeleton.h"
+
+bool FPVExportParams::Validate(FString& OutError) const
+{
+	const FString& OutputPath = ContentBrowserFolder.Path;
+
+	UClass* MeshClass = ExportMeshType == EPVExportMeshType::StaticMesh ? UStaticMesh::StaticClass() : USkeletalMesh::StaticClass();
+	if (!PV::Utilities::ValidateAssetPathAndName(MeshName.ToString(), OutputPath, MeshClass, OutError))
+	{
+		return false;
+	}
+
+	if (ExportMeshType == EPVExportMeshType::SkeletalMesh)
+	{
+		if (!PV::Utilities::ValidateAssetPathAndName(GetOutputSkeletonName(), OutputPath, USkeleton::StaticClass(), OutError))
+		{
+			return false;
+		}
+	}
+
+	return true;
+}
\ No newline at end of file
diff --git a/Engine/Plugins/Experimental/ProceduralVegetationEditor/Source/ProceduralVegetation/ProceduralVegetation.Build.cs b/Engine/Plugins/Experimental/ProceduralVegetationEditor/Source/ProceduralVegetation/ProceduralVegetation.Build.cs
--- a/Engine/Plugins/Experimental/ProceduralVegetationEditor/Source/ProceduralVegetation/ProceduralVegetation.Build.cs
+++ b/Engine/Plugins/Experimental/ProceduralVegetationEditor/Source/ProceduralVegetation/ProceduralVegetation.Build.cs
@@ -7,7 +7,7 @@ namespace UnrealBuildTool.Rules
 		public ProceduralVegetation(ReadOnlyTargetRules Target) : base(Target)
 		{
 			PCHUsage = PCHUsageMode.UseExplicitOrSharedPCHs;
-
+			
 			PublicDependencyModuleNames.AddRange(
 				new string[]
 				{
@@ -24,27 +24,11 @@ namespace UnrealBuildTool.Rules
 					"RenderCore",
 					"GeometryCore",
 					"GeometryFramework",
-					"GeometryScriptingCore",
 					"Json",
 					"JsonUtilities",
-					"MeshDescription",
-					"MeshConversion",
-					"PlanarCut",
-					"SkeletalMeshDescription",
-					"StaticMeshDescription",
-					"ModelingComponents",
-					"PCG", // ... add private dependencies that you statically link with here ...	
+					"PCG",
 				}
 			);
-			
-			if (Target.Type == TargetRules.TargetType.Editor)
-			{
-				PrivateDependencyModuleNames.AddRange(
-					new string[]
-					{
-						"DynamicWindEditor"
-					});
-			}
 		}
 	}
 }
diff --git a/Engine/Plugins/Experimental/ProceduralVegetationEditor/Source/ProceduralVegetation/Public/Helpers/PVExportHelper.h b/Engine/Plugins/Experimental/ProceduralVegetationEditor/Source/ProceduralVegetation/Public/Helpers/PVExportHelper.h
deleted file mode 100644
diff --git a/Engine/Plugins/Experimental/ProceduralVegetationEditor/Source/ProceduralVegetation/Public/Helpers/PVUtilities.h b/Engine/Plugins/Experimental/ProceduralVegetationEditor/Source/ProceduralVegetation/Public/Helpers/PVUtilities.h
--- a/Engine/Plugins/Experimental/ProceduralVegetationEditor/Source/ProceduralVegetation/Public/Helpers/PVUtilities.h
+++ b/Engine/Plugins/Experimental/ProceduralVegetationEditor/Source/ProceduralVegetation/Public/Helpers/PVUtilities.h
@@ -1,5 +1,6 @@
 ﻿// Copyright Epic Games, Inc. All Rights Reserved.
 #pragma once
+
 #include "Facades/PVRenderingFacade.h"
 
 class UInstancedStaticMeshComponent;
@@ -7,20 +8,19 @@ class UMeshComponent;
 class FFoliageFacade;
 class AActor;
 struct FAssetData;
+struct FPVExportParams;
 
 namespace PV::Utilities
 {
 	bool PROCEDURALVEGETATION_API DebugModeEnabled();
 
 	DECLARE_DELEGATE_OneParam(FFoliageComponentCreatedCallback, UMeshComponent*);
-	DECLARE_DELEGATE_OneParam(FFoliageInstanceAddedCallback, int32);
 	
 	void PROCEDURALVEGETATION_API AddFoliageInstancesToActor(const Facades::FFoliageFacade& InFoliageFacade, AActor* InActor, TMap<FString, TObjectPtr<UMeshComponent>>& OutInstancedComponentMap);
 	void PROCEDURALVEGETATION_API AddFoliageInstances(const Facades::FFoliageFacade& InFoliageFacade, const UObject* InParent, FFoliageComponentCreatedCallback InCallback, TMap<FString, TObjectPtr<UMeshComponent>>& OutInstancedComponentMap);
 
-	bool PROCEDURALVEGETATION_API ValidateMeshPathAndName(FString MeshName, FString Path, UClass* InClass, FString& OutError);
-	
-	bool IsFileNameValid(FName FileName, FText& Reason);
+	bool PROCEDURALVEGETATION_API ValidateAssetPathAndName(const FString& MeshName, const FString& Path, UClass* InClass, FString& OutError);
 
-	bool CanOverwriteAsset(FString PackageName, UClass* InClass, FAssetData& OutAsset, FString& OutError);
+	bool IsFileNameValid(FName FileName, FText& Reason);
+	bool PROCEDURALVEGETATION_API DoesConflictingPackageExist(const FString& PackageName, UClass* InClass);
 }
diff --git a/Engine/Plugins/Experimental/ProceduralVegetationEditor/Source/ProceduralVegetation/Public/Nodes/PVOutputSettings.h b/Engine/Plugins/Experimental/ProceduralVegetationEditor/Source/ProceduralVegetation/Public/Nodes/PVOutputSettings.h
--- a/Engine/Plugins/Experimental/ProceduralVegetationEditor/Source/ProceduralVegetation/Public/Nodes/PVOutputSettings.h
+++ b/Engine/Plugins/Experimental/ProceduralVegetationEditor/Source/ProceduralVegetation/Public/Nodes/PVOutputSettings.h
@@ -2,7 +2,7 @@
 #pragma once
 
 #include "CoreMinimal.h"
-#include "Helpers/PVExportHelper.h"
+#include "PVExportParams.h"
 #include "Nodes/PVBaseSettings.h"
 #include "PVOutputSettings.generated.h"
 
diff --git a/Engine/Plugins/Experimental/ProceduralVegetationEditor/Source/ProceduralVegetation/Public/PVExportParams.h b/Engine/Plugins/Experimental/ProceduralVegetationEditor/Source/ProceduralVegetation/Public/PVExportParams.h
new file mode 100644
--- /dev/null
+++ b/Engine/Plugins/Experimental/ProceduralVegetationEditor/Source/ProceduralVegetation/Public/PVExportParams.h
@@ -0,0 +1,73 @@
+// Copyright Epic Games, Inc. All Rights Reserved.
+#pragma once
+
+#include "Engine/EngineTypes.h"
+#include "PVWindSettings.h"
+#include "Misc/Paths.h"
+#include "PVExportParams.generated.h"
+
+UENUM()
+enum class EPVExportMeshType: uint8
+{
+	StaticMesh,
+	SkeletalMesh,
+};
+
+USTRUCT()
+struct PROCEDURALVEGETATION_API FPVExportParams
+{
+	GENERATED_BODY()
+
+	UPROPERTY()
+	bool bShouldExport = true;
+
+	UPROPERTY(EditAnywhere, Category = "Export", Meta = (ContentDir, DisplayName = "Content Browser Folder", Tooltip="Path to the folder where exported mesh will be saved\n\nPath to the folder where exported mesh will be saved."))
+	FDirectoryPath ContentBrowserFolder;
+
+	UPROPERTY(EditAnywhere, Category = "Export", Meta = (DisplayName = "Mesh Name", Tooltip="Name of the exported mesh\n\nName of the exported mesh."))
+	FName MeshName;
+
+	UPROPERTY(EditAnywhere, Category = "Export", Meta = (DisplayName = "Export Mesh Type", Tooltip="Static Mesh or Skeletal Mesh."))
+	EPVExportMeshType ExportMeshType = EPVExportMeshType::SkeletalMesh;
+
+	UPROPERTY(EditAnywhere, Category = "Export", Meta = (DisplayName = "Create Nanite Foliage", Tooltip="Add Nanite foliage support for the exported mesh."))
+	bool bCreateNaniteFoliage = true;
+	
+	UPROPERTY(EditAnywhere, Category = "Export", Meta = (DisplayName = "Nanite Shape Preservation Method", Tooltip="Shape preservation setting to be used in nanite settings."))
+	ENaniteShapePreservation NaniteShapePreservation = ENaniteShapePreservation::Voxelize;
+
+	UPROPERTY(EditAnywhere, Category = "Export", Meta = (EditCondition = "ExportMeshType == EPVExportMeshType::SkeletalMesh"))
+	TObjectPtr<UPVWindSettings> WindSettings;
+
+	void Initialize(const FString& InAssetPath, const FString& InName, const FString& InDefaultWindSettingsPath)
+	{
+		ContentBrowserFolder.Path = FPaths::GetPath(InAssetPath);
+		MeshName = FName(InName);
+
+		WindSettings = LoadObject<UPVWindSettings>(nullptr, InDefaultWindSettingsPath);
+	}
+
+	FString GetOutputObjectPath() const
+	{
+		const FString AssetName = MeshName.ToString();
+		const FString OutputPath = ContentBrowserFolder.Path / AssetName + '.' + AssetName;
+		return OutputPath;
+	}
+
+	FString GetOutputMeshPackagePath() const
+	{
+		return FPaths::Combine(ContentBrowserFolder.Path, MeshName.ToString());
+	}
+
+	FString GetOutputSkeletonName() const
+	{
+		return FString::Printf(TEXT("%s_Skeleton"), *MeshName.ToString());
+	}
+
+	FString GetOutputSkeletonPackagePath() const
+	{
+		return FPaths::Combine(ContentBrowserFolder.Path, GetOutputSkeletonName());
+	}
+
+	bool Validate(FString& OutError) const;
+};
\ No newline at end of file
diff --git a/Engine/Plugins/Experimental/ProceduralVegetationEditor/Source/ProceduralVegetationEditor/Private/Customizations/PVOutputSettingsCustomizations.cpp b/Engine/Plugins/Experimental/ProceduralVegetationEditor/Source/ProceduralVegetationEditor/Private/Customizations/PVOutputSettingsCustomizations.cpp
--- a/Engine/Plugins/Experimental/ProceduralVegetationEditor/Source/ProceduralVegetationEditor/Private/Customizations/PVOutputSettingsCustomizations.cpp
+++ b/Engine/Plugins/Experimental/ProceduralVegetationEditor/Source/ProceduralVegetationEditor/Private/Customizations/PVOutputSettingsCustomizations.cpp
@@ -10,8 +10,6 @@
 
 #include "Framework/Docking/TabManager.h"
 
-#include "Helpers/PVExportHelper.h"
-
 #include "Nodes/PVOutputSettings.h"
 
 #include "Widgets/Input/SButton.h"
diff --git a/Engine/Plugins/Experimental/ProceduralVegetationEditor/Source/ProceduralVegetationEditor/Private/Helpers/PVExportHelper.cpp b/Engine/Plugins/Experimental/ProceduralVegetationEditor/Source/ProceduralVegetationEditor/Private/Helpers/PVExportHelper.cpp
new file mode 100644
--- /dev/null
+++ b/Engine/Plugins/Experimental/ProceduralVegetationEditor/Source/ProceduralVegetationEditor/Private/Helpers/PVExportHelper.cpp
@@ -0,0 +1,1485 @@
+﻿// Copyright Epic Games, Inc. All Rights Reserved.
+
+#include "Helpers/PVExportHelper.h"
+#include "PVExportParams.h"
+#include "DynamicMeshEditor.h"
+#include "DynamicWindImportData.h"
+#include "DynamicWindSkeletalData.h"
+#include "ProceduralVegetation.h"
+#include "ProceduralVegetationLink.h"
+#include "ProceduralVegetationEditorModule.h"
+#include "MeshDescriptionToDynamicMesh.h"
+#include "PlanarCut.h"
+#include "PackageTools.h"
+#include "RenderUtils.h"
+#include "StaticMeshAttributes.h"
+#include "Animation/Skeleton.h"
+#include "AssetRegistry/IAssetRegistry.h"
+#include "Components/InstancedStaticMeshComponent.h"
+#include "ConversionUtils/SceneComponentToDynamicMesh.h"
+#include "DynamicMesh/DynamicBoneAttribute.h"
+#include "DynamicMesh/DynamicVertexSkinWeightsAttribute.h"
+
+#include "UDynamicMesh.h"
+#include "Engine/SkeletalMesh.h"
+#include "Engine/StaticMesh.h"
+#include "Facades/PVBranchFacade.h"
+#include "Facades/PVPointFacade.h"
+#include "GeometryCollection/GeometryCollection.h"
+#include "GeometryCollection/ManagedArrayCollection.h"
+#include "GeometryScript/GeometryScriptTypes.h"
+#include "GeometryScript/MeshAssetFunctions.h"
+#include "GeometryScript/MeshBoneWeightFunctions.h"
+#include "Helpers/PVUtilities.h"
+#include "Materials/MaterialInterface.h"
+#include "Misc/Paths.h"
+#include "UObject/Package.h"
+#include "PVWindSettings.h"
+#include "SkinnedAssetCompiler.h"
+#include "StaticMeshCompiler.h"
+#include "GeometryCollection/Facades/CollectionConstraintOverrideFacade.h"
+#include "AssetRegistry/AssetData.h"
+#include "AssetRegistry/AssetRegistryModule.h"
+
+#define LOCTEXT_NAMESPACE "PVExportHelper"
+
+namespace PV::Export::Internal
+{
+	enum class EExportMeshStage
+	{
+		InitializingExportMesh,
+		CombiningMeshes,
+		CopyingMeshDataToOutputMesh,
+		RecreatingBoneTree,
+		ComputingBoneWeights,
+		CopyingSkinWeightProfiles,
+		BuildingNaniteAssemblies,
+		ImportingDynamicWind,
+		BuildingMesh,
+		Complete,
+		MAX
+	};
+
+	float CalculateMeshExportStageWeight(EExportMeshStage InStage, bool bIsBuildingNaniteAssemblies, bool bIsExportingSkeletalMesh)
+	{
+		switch (InStage)
+		{
+		case EExportMeshStage::InitializingExportMesh:
+			return 1.f;
+		case EExportMeshStage::CombiningMeshes:
+			return bIsBuildingNaniteAssemblies ? 0.f : 1.f;
+		case EExportMeshStage::CopyingMeshDataToOutputMesh:
+			return 1.f;
+		case EExportMeshStage::RecreatingBoneTree:
+			return bIsExportingSkeletalMesh ? 0.2f : 0.f;
+		case EExportMeshStage::ComputingBoneWeights:
+			return bIsExportingSkeletalMesh ? 2.f : 0.f;
+		case EExportMeshStage::CopyingSkinWeightProfiles:
+			return bIsExportingSkeletalMesh ? 0.1f : 0.f;
+		case EExportMeshStage::BuildingNaniteAssemblies:
+			return bIsBuildingNaniteAssemblies ? 3.f : 0.f;
+		case EExportMeshStage::ImportingDynamicWind:
+			return bIsExportingSkeletalMesh ? 0.1f : 0.f;
+		case EExportMeshStage::BuildingMesh:
+			return bIsBuildingNaniteAssemblies ? 3.f : 4.f;
+		case EExportMeshStage::Complete:
+			return 0.5f;
+		default:
+			break;
+		}
+
+		check(false);
+		return 0.f;
+	}
+
+	float CalculateMeshExportProgress(EExportMeshStage InStage, float StageProgress, bool bIsBuildingNaniteAssemblies, bool bIsExportingSkeletalMesh)
+	{
+		ensure(StageProgress >= 0 && StageProgress <= 1.01f);
+		StageProgress = FMath::Clamp(StageProgress, 0.f, 1.f);
+
+		float TotalWeightSum = 0;
+		float ThisStageWeightSum = 0;
+		for (int32 i = 0; i < (int32)EExportMeshStage::MAX; ++i)
+		{
+			const float StatusWeight = CalculateMeshExportStageWeight(EExportMeshStage(i), bIsBuildingNaniteAssemblies, bIsExportingSkeletalMesh);
+			TotalWeightSum += StatusWeight;
+			if (i < int32(InStage))
+			{
+				ThisStageWeightSum += StatusWeight;
+			}
+		}
+
+		const float ThisStatusWeight = CalculateMeshExportStageWeight(InStage, bIsBuildingNaniteAssemblies, bIsExportingSkeletalMesh);
+		ThisStageWeightSum += ThisStatusWeight * StageProgress;
+
+		return TotalWeightSum > 0
+			? ThisStageWeightSum / TotalWeightSum
+			: 0.f;
+	}
+
+	FText GetMeshExportStageText(EExportMeshStage InStage, bool bIsBuildingNaniteAssemblies, bool bIsExportingSkeletalMesh)
+	{
+		switch (InStage)
+		{
+		case EExportMeshStage::InitializingExportMesh:
+			return LOCTEXT("ExportMeshStatus_InitializingExportMesh", "Initializing mesh data");
+		case EExportMeshStage::CombiningMeshes:
+			return LOCTEXT("ExportMeshStatus_CombiningMeshes", "Combining foliage meshes");
+		case EExportMeshStage::CopyingMeshDataToOutputMesh:
+			return LOCTEXT("ExportMeshStatus_CopyingMeshDataToOutputMesh", "Copying mesh data to output mesh");
+		case EExportMeshStage::RecreatingBoneTree:
+			return LOCTEXT("ExportMeshStatus_RecreatingBoneTree", "Recreating bone tree");
+		case EExportMeshStage::ComputingBoneWeights:
+			return LOCTEXT("ExportMeshStatus_ComputingBoneWeights", "Computing bone weights");
+		case EExportMeshStage::CopyingSkinWeightProfiles:
+			return LOCTEXT("ExportMeshStatus_CopyingSkinWeightProfiles", "Copying skin weight profiles");
+		case EExportMeshStage::BuildingNaniteAssemblies:
+			return LOCTEXT("ExportMeshStatus_BuildingNaniteAssemblies", "Building nanite assemblies");
+		case EExportMeshStage::ImportingDynamicWind:
+			return LOCTEXT("ExportMeshStatus_ImportingDynamicWind", "Importing dynamic wind");
+		case EExportMeshStage::BuildingMesh:
+		{
+			if (!bIsBuildingNaniteAssemblies)
+			{
+				return LOCTEXT("ExportMeshStatus_BuildingMesh_Slow", "Building mesh (this might take a while)");
+			}
+			else
+			{
+				return LOCTEXT("ExportMeshStatus_BuildingMesh", "Building mesh");
+			}
+		}
+		case EExportMeshStage::Complete:
+			return LOCTEXT("ExportMeshStatus_Complete", "Complete");
+		default:
+			break;
+		}
+
+		check(false);
+		return FText::FromString("Error");
+	}
+
+	void CleanupPackages(const TArray<UPackage*>& Packages)
+	{
+		FText UnloadError;
+		UPackageTools::UnloadPackages(Packages, UnloadError, true);
+	}
+
+	template<typename T>
+	TObjectPtr<T> GetAssetOfType(const FString& LongPackageName)
+	{
+		UPackage* Package = LoadPackage(nullptr, *LongPackageName, LOAD_None);
+		if (Package == nullptr)
+		{
+			return nullptr;
+		}
+
+		const FName AssetName = FPackageName::GetShortFName(Package->GetFName());
+		return FindObjectFast<T>(Package, AssetName);
+	}
+
+	UPackage* CreateAssetPackage(const FString& LongPackageName)
+	{
+		UPackage* Package = LoadPackage(nullptr, *LongPackageName, LOAD_None);
+
+		if (Package == nullptr)
+		{
+			Package = CreatePackage(*LongPackageName);
+			Package->SetPackageFlags(PKG_NewlyCreated);
+		}
+		
+		return Package;
+	}
+
+	TObjectPtr<UDynamicMesh> CollectionToDynamicMesh(const FManagedArrayCollection& Collection)
+	{
+		TRACE_CPUPROFILER_EVENT_SCOPE(PV::Export::CollectionToDynamicMesh);
+		
+		if (Collection.NumElements(FGeometryCollection::TransformGroup) == 0)
+		{
+			return nullptr;
+		}
+
+		const TUniquePtr<FGeometryCollection> GeomCollection = TUniquePtr<FGeometryCollection>(Collection.NewCopy<FGeometryCollection>());
+		if (!GeomCollection.IsValid())
+		{
+			return nullptr;
+		}
+
+		TObjectPtr<UDynamicMesh> NewMesh = NewObject<UDynamicMesh>();
+
+		const TManagedArray<FTransform3f>& BoneTransforms = Collection.GetAttribute<FTransform3f>(
+			"Transform", FGeometryCollection::TransformGroup);
+
+		TArray<int32> TransformIndices;
+		TransformIndices.AddUninitialized(BoneTransforms.Num());
+
+		int32 Idx = 0;
+		for (int32& TransformIdx : TransformIndices)
+		{
+			TransformIdx = Idx++;
+		}
+
+		FMeshDescription MeshDescription;
+		FStaticMeshAttributes Attributes(MeshDescription);
+		Attributes.Register();
+
+		FTransform TransformOut;
+
+		constexpr bool bCenterPivot = false;
+		ConvertToMeshDescription(MeshDescription, TransformOut, bCenterPivot, *GeomCollection, BoneTransforms, TransformIndices);
+
+		NewMesh->Reset();
+
+		UE::Geometry::FDynamicMesh3& DynMesh = NewMesh->GetMeshRef();
+		{
+			FMeshDescriptionToDynamicMesh ConverterToDynamicMesh;
+			ConverterToDynamicMesh.Convert(&MeshDescription, DynMesh);
+		}
+
+		return NewMesh;
+	}
+
+	FGeometryScriptCopyMeshToAssetOptions GetCopyMeshToAssetOptions(const TArray<TObjectPtr<UMaterialInterface>>& InMaterials, const FMeshNaniteSettings& InNaniteSettings, const ENaniteShapePreservation InShapePreservation)
+	{
+		FGeometryScriptCopyMeshToAssetOptions Options;
+		Options.bEmitTransaction = false;
+		Options.bDeferMeshPostEditChange = true;
+		Options.bEnableRecomputeTangents = true;
+		Options.bReplaceMaterials = true;
+		Options.NewMaterials = InMaterials;
+
+		Options.NewNaniteSettings = InNaniteSettings;
+		Options.NewNaniteSettings.bEnabled = true;
+		Options.NewNaniteSettings.ShapePreservation = InShapePreservation;
+		Options.bApplyNaniteSettings = true;
+
+		return Options;
+	}
+
+	FString GetFoliageName(const Facades::FFoliageFacade& FoliageFacade, int32 Id)
+	{
+		Facades::FFoliageEntryData Data = FoliageFacade.GetFoliageEntry(Id);
+		return FoliageFacade.GetFoliageName(Data.NameId);
+	}
+
+	TArray<TObjectPtr<UMaterialInterface>> GetMaterialsFromCollection(const FManagedArrayCollection& Collection)
+	{
+		TRACE_CPUPROFILER_EVENT_SCOPE(PV::Export::GetMaterialsFromCollection);
+
+		TArray<TObjectPtr<UMaterialInterface>> Materials;
+		if (Collection.HasAttribute("MaterialPath", FGeometryCollection::MaterialGroup))
+		{
+			const TManagedArray<FString>& MaterialPaths = Collection.GetAttribute<FString>("MaterialPath", FGeometryCollection::MaterialGroup);
+			for (const FString& Path : MaterialPaths)
+			{
+				Materials.Add(LoadObject<UMaterialInterface>(nullptr, *Path));
+			}
+		}
+		return Materials;
+	}
+
+	static bool SkeletalMeshToDynamicMesh(const USkeletalMesh* SkeletalMesh, const int32 LodIndex, FDynamicMesh3& ToDynamicMesh)
+	{
+		TRACE_CPUPROFILER_EVENT_SCOPE(PV::Export::SkeletalMeshToDynamicMesh);
+
+		if (SkeletalMesh->HasMeshDescription(LodIndex))
+		{
+			const FMeshDescription* SourceMesh = SkeletalMesh->GetMeshDescription(LodIndex);
+			if (!SourceMesh)
+			{
+				return false;
+			}
+
+			FMeshDescriptionToDynamicMesh Converter;
+			Converter.Convert(SourceMesh, ToDynamicMesh);
+			return true;
+		}
+
+		return false;
+	}
+
+	bool CombineMeshInstancesToDynamicMesh(
+		const FManagedArrayCollection& Collection,
+		UDynamicMesh* DynamicMesh,
+		TArray<TObjectPtr<UMaterialInterface>>& Materials,
+		const FMeshInstanceCombined& OnMeshInstanceCombined,
+		const TFunction<bool(float)>& OnProgressUpdated
+	)
+	{
+		TRACE_CPUPROFILER_EVENT_SCOPE(PV::Export::CombineMeshInstancesToDynamicMesh);
+
+		const Facades::FFoliageFacade FoliageFacade(Collection);
+		const int32 NumInstances = FoliageFacade.NumFoliageEntries();
+		const int32 NumFoliageNames = FoliageFacade.NumFoliageNames();
+		
+		if (NumInstances == 0 || NumFoliageNames == 0)
+		{
+			return true;
+		}
+
+		TMap<FString, FDynamicMesh3> FoliageDynamicMeshes;
+		TMap<FString, TObjectPtr<UStaticMesh>> FoliageStaticMeshes;
+		TMap<FString, TObjectPtr<USkeletalMesh>> FoliageSkeletalMeshes;
+		
+		FoliageDynamicMeshes.Reserve(NumFoliageNames);
+		FoliageStaticMeshes.Reserve(NumFoliageNames);
+		FoliageSkeletalMeshes.Reserve(NumFoliageNames);
+
+		for (int32 FoliageNameIndex = 0; FoliageNameIndex < NumFoliageNames; ++FoliageNameIndex)
+		{
+			const FString FoliageName = FoliageFacade.GetFoliageName(FoliageNameIndex);
+			if (FoliageDynamicMeshes.Contains(FoliageName))
+			{
+				UE_LOG(LogProceduralVegetationEditor, Warning, TEXT("Duplicate foliage names in foliage facade"));
+				continue;
+			}
+
+			if (const TObjectPtr<UStaticMesh> StaticMesh = LoadObject<UStaticMesh>(nullptr, *FoliageName))
+			{
+				FoliageStaticMeshes.Add(FoliageName, StaticMesh);
+					
+				for (auto Material : StaticMesh->GetStaticMaterials())
+				{
+					Materials.AddUnique(Material.MaterialInterface);
+				}
+
+				FDynamicMesh3 FoliageDynamicMesh;
+				FText OutErrorMessage;
+				UE::Conversion::FStaticMeshConversionOptions Options;
+				if (StaticMeshToDynamicMesh(StaticMesh, FoliageDynamicMesh, OutErrorMessage, Options))
+				{
+					FoliageDynamicMeshes.Add(FoliageName, MoveTemp(FoliageDynamicMesh));
+				}
+			}
+			else if (const TObjectPtr<USkeletalMesh> SkeletalMesh = LoadObject<USkeletalMesh>(nullptr, *FoliageName))
+			{
+				FoliageSkeletalMeshes.Add(FoliageName, SkeletalMesh);
+					
+				for (auto Material : SkeletalMesh->GetMaterials())
+				{
+					Materials.AddUnique(Material.MaterialInterface);
+				}
+
+				FDynamicMesh3 FoliageDynamicMesh;
+				if (SkeletalMeshToDynamicMesh(SkeletalMesh, 0, FoliageDynamicMesh))
+				{
+					FoliageDynamicMesh.Attributes()->RemoveAttribute("BoneIndex");
+					FoliageDynamicMesh.Attributes()->RemoveAttribute("BoneParentIndex");
+					FoliageDynamicMeshes.Add(FoliageName, MoveTemp(FoliageDynamicMesh));
+				}
+			}
+
+			if (!OnProgressUpdated((FoliageNameIndex / (float)NumFoliageNames) * 0.5f))
+			{
+				return false;
+			}
+		}
+
+		for (int32 Id = 0; Id < NumInstances; Id++)
+		{
+			const FString FoliageName = GetFoliageName(FoliageFacade, Id);
+
+			FDynamicMesh3* FoliageDynamicMesh = FoliageDynamicMeshes.Find(FoliageName);
+			if (FoliageDynamicMesh == nullptr)
+			{
+				continue;
+			}
+
+			const TObjectPtr<UStaticMesh>* StaticMesh = FoliageStaticMeshes.Find(FoliageName);
+			const TObjectPtr<USkeletalMesh>* SkeletalMesh = FoliageSkeletalMeshes.Find(FoliageName);
+
+			check((StaticMesh != nullptr || SkeletalMesh != nullptr) && (*StaticMesh != nullptr || *SkeletalMesh != nullptr));
+
+			if (SkeletalMesh != nullptr)
+			{
+				const int32 FoliageParentBoneID = FoliageFacade.GetParentBoneID(Id);
+
+				// Reparent root bone according to the instanced transform
+				UE::Geometry::FDynamicMeshAttributeSet* MeshAttributes = FoliageDynamicMesh->Attributes();
+				UE::Geometry::FDynamicMeshBoneParentIndexAttribute* BoneParentIndices = MeshAttributes->GetBoneParentIndices();
+				BoneParentIndices->SetValue(0, FoliageParentBoneID);
+			}
+				
+			const FTransform Transform = FoliageFacade.GetFoliageTransform(Id);
+			UE::Geometry::FTransformSRT3d GeoTransform(Transform);
+
+			DynamicMesh->EditMesh([&](FDynamicMesh3& Mesh)
+			{
+				const int32 VertexStart = Mesh.VertexCount();
+				UE::Geometry::FDynamicMeshEditor Editor = UE::Geometry::FDynamicMeshEditor(&Mesh);
+
+				UE::Geometry::FMeshIndexMappings IndexMap;
+				Editor.AppendMesh(
+					FoliageDynamicMesh,
+					IndexMap,
+					[&](int32, const FVector3d& Position)
+					{
+						return GeoTransform.TransformPosition(Position);
+					},
+					[&](int32, const FVector3d& Normal)
+					{
+						return GeoTransform.TransformVectorNoScale(Normal);
+					}
+				);
+
+				// Remap Material IDs
+				if (Mesh.Attributes()->HasMaterialID())
+				{
+					const UE::Geometry::FIndexMapi& TriangleMap = IndexMap.GetTriangleMap();
+					UE::Geometry::FDynamicMeshMaterialAttribute* MaterialIDs = Mesh.Attributes()->GetMaterialID();
+
+					for (const auto& [OldTid, NewTid] : TriangleMap.GetForwardMap())
+					{
+						const int32 MaterialID = MaterialIDs->GetValue(OldTid);
+						if (StaticMesh)
+						{
+							UMaterialInterface* Material = (*StaticMesh)->GetMaterial(MaterialID);
+							const int32 NewMaterialID = Materials.Find(Material);
+							MaterialIDs->SetValue(NewTid, NewMaterialID);
+						}
+						else if (SkeletalMesh)
+						{
+							UMaterialInterface* Material = (*SkeletalMesh)->GetMaterials()[MaterialID].MaterialInterface;
+							const int32 NewMaterialID = Materials.Find(Material);
+							MaterialIDs->SetValue(NewTid, NewMaterialID);
+						}
+					}
+				}
+
+				if (StaticMesh)
+				{
+					OnMeshInstanceCombined((*StaticMesh)->GetFullName(), Id, VertexStart, Mesh.VertexCount());
+				}
+				else if (SkeletalMesh != nullptr)
+				{
+					OnMeshInstanceCombined((*SkeletalMesh)->GetFullName(), Id, VertexStart, Mesh.VertexCount());
+				}
+			});
+
+			if (!OnProgressUpdated(0.5f + (Id / (float)NumInstances) * 0.5f))
+			{
+				return false;
+			}
+		}
+
+		return true;
+	}
+
+	bool ContainsFoliageName(const Facades::FFoliageFacade& FoliageFacade, const FString& InFoliageName)
+	{
+		int32 NumInstances = FoliageFacade.NumFoliageEntries();
+			
+		for(int32 Id = 0; Id < NumInstances; Id++)
+		{
+			FString FoliageName = GetFoliageName(FoliageFacade, Id);
+			if(FoliageName == InFoliageName)
+			{
+				return true;
+			}
+		}
+
+		return false;
+	}
+
+	TArray<FString> GetUsedFoliage(const Facades::FFoliageFacade& FoliageFacade)
+	{
+		TArray<FString> OutFoliageNames;
+		OutFoliageNames.Reserve(FoliageFacade.NumFoliageNames());
+
+		for(int32 FoliageNameIndex = 0; FoliageNameIndex < FoliageFacade.NumFoliageNames(); ++FoliageNameIndex)
+		{
+			FString FoliageName = FoliageFacade.GetFoliageName(FoliageNameIndex);
+			if (ContainsFoliageName(FoliageFacade, FoliageName))
+			{
+				OutFoliageNames.Add(FoliageName);
+			}
+		}
+
+		return OutFoliageNames;
+	}
+
+	void BuildNaniteAssemblyData(
+		const FManagedArrayCollection& Collection,
+		UStaticMesh* StaticMesh,
+		TArray<UPackage*>& OutModifiedPackages
+	)
+	{
+		TRACE_CPUPROFILER_EVENT_SCOPE(PV::Export::BuildNaniteAssemblyData_StaticMesh);
+
+		//Create NaniteAssemblyDataBuilder and the main static mesh materials to it
+		FNaniteAssemblyDataBuilder AssemblyBuilder;
+		for (const FStaticMaterial& Material : StaticMesh->GetStaticMaterials())
+		{
+			AssemblyBuilder.AddMaterialSlot(Material.MaterialInterface);
+		}
+
+		const Facades::FFoliageFacade FoliageFacade(Collection);
+		const TArray<FString> UsedFoliage = GetUsedFoliage(FoliageFacade);
+
+		TMap<FString, int32> MeshNamePartMap;
+		MeshNamePartMap.Reserve(UsedFoliage.Num());
+		TMap<FString, TArray<TObjectPtr<UMaterialInterface>>> MeshMaterialsMap;
+		MeshMaterialsMap.Reserve(UsedFoliage.Num() * 2);
+
+		const auto AddStaticMeshToAssembly = [&](const FString& FoliageInstanceObjectPath, TObjectPtr<UStaticMesh> InStaticMesh)
+		{
+			const int32 PartIndex = AssemblyBuilder.FindOrAddPart(FoliageInstanceObjectPath);
+			MeshNamePartMap.FindOrAdd(FoliageInstanceObjectPath) = PartIndex;
+
+			TArray<TObjectPtr<UMaterialInterface>>& MeshMaterials = MeshMaterialsMap.FindOrAdd(FoliageInstanceObjectPath);
+			for (const FStaticMaterial& Material : InStaticMesh->GetStaticMaterials())
+			{
+				MeshMaterials.AddUnique(Material.MaterialInterface);
+			}
+		};
+		
+		for (const FString& FoliageInstanceObjectPath : UsedFoliage)
+		{
+			TObjectPtr<UStaticMesh> PartStaticMesh = LoadObject<UStaticMesh>(nullptr, *FoliageInstanceObjectPath);
+			if (PartStaticMesh)
+			{
+				AddStaticMeshToAssembly(FoliageInstanceObjectPath, PartStaticMesh);
+				continue;
+			}
+
+			const FString FoliageInstancePackagePath = FPackageName::ObjectPathToPackageName(FoliageInstanceObjectPath);
+			const FString LongFoliagePackagePath = FPackageName::GetLongPackagePath(FoliageInstancePackagePath);
+			const FString ShortFoliagePackageName = FPackageName::GetShortName(FoliageInstancePackagePath);
+
+			const FString StaticMeshName = "SKM_" + ShortFoliagePackageName;
+			const FString StaticMeshPackagePath = FPaths::Combine(LongFoliagePackagePath, StaticMeshName);
+			const FString StaticMeshObjectPath = FString::Printf(TEXT("%s.%s"), *StaticMeshPackagePath, *StaticMeshName);
+
+			PartStaticMesh = LoadObject<UStaticMesh>(nullptr, StaticMeshObjectPath);
+			if (PartStaticMesh)
+			{
+				AddStaticMeshToAssembly(FoliageInstanceObjectPath, PartStaticMesh);
+				continue;
+			}
+
+			const TObjectPtr<USkeletalMesh> PartSkeletalMesh = LoadObject<USkeletalMesh>(nullptr, *FoliageInstanceObjectPath);
+			if (!PartSkeletalMesh)
+			{
+				UE_LOG(LogProceduralVegetationEditor, Warning, TEXT("Could not find foliage instance %s"), *FoliageInstanceObjectPath);
+				continue;
+			}
+
+			FDynamicMesh3 PartDynamicMesh;
+			if (!SkeletalMeshToDynamicMesh(PartSkeletalMesh, 0, PartDynamicMesh))
+			{
+				UE_LOG(LogProceduralVegetationEditor, Warning, TEXT("Failed to convert foliage instance \"%s\" to dynamic mesh"), *FoliageInstancePackagePath);
+				continue;
+			}
+					
+			UPackage* StaticMeshPackage = Internal::CreateAssetPackage(*StaticMeshPackagePath);
+			OutModifiedPackages.Add(StaticMeshPackage);
+
+			PartStaticMesh = NewObject<UStaticMesh>(StaticMeshPackage, FName(StaticMeshName), RF_Standalone | RF_Public);
+
+			TArray<TObjectPtr<UMaterialInterface>> PartMaterials;
+			for (auto Material : PartSkeletalMesh->GetMaterials())
+			{
+				PartMaterials.Add(Material.MaterialInterface);
+			}
+
+			EGeometryScriptOutcomePins Output;
+			FGeometryScriptCopyMeshToAssetOptions Options = GetCopyMeshToAssetOptions(PartMaterials, PartStaticMesh->GetNaniteSettings(), ENaniteShapePreservation::Voxelize);
+
+			TObjectPtr<UDynamicMesh> DynamicMesh = NewObject<UDynamicMesh>();
+			DynamicMesh->SetMesh(MoveTemp(PartDynamicMesh));
+
+			UGeometryScriptLibrary_StaticMeshFunctions::CopyMeshToStaticMesh(DynamicMesh, PartStaticMesh, Options, {}, Output);
+
+			AddStaticMeshToAssembly(FoliageInstanceObjectPath, PartStaticMesh);
+		}
+
+		AddNodeToBuilder(AssemblyBuilder, Collection, MeshNamePartMap, MeshMaterialsMap);
+		AssemblyBuilder.ApplyToStaticMesh(*StaticMesh);
+	}
+
+	void BuildNaniteAssemblyData(
+		const FManagedArrayCollection& Collection,
+		USkeletalMesh* SkeletalMesh,
+		TArray<UPackage*>& OutModifiedPackages
+	)
+	{
+		TRACE_CPUPROFILER_EVENT_SCOPE(PV::Export::BuildNaniteAssemblyData_SkeletalMesh);
+
+		//Create NaniteAssemblyDataBuilder and the main static mesh materials to it
+		FNaniteAssemblyDataBuilder AssemblyBuilder;
+		for (const FSkeletalMaterial& Material : SkeletalMesh->GetMaterials())
+		{
+			AssemblyBuilder.AddMaterialSlot(Material.MaterialInterface);
+		}
+
+		const Facades::FFoliageFacade FoliageFacade(Collection);
+		const TArray<FString> UsedFoliage = GetUsedFoliage(FoliageFacade);
+
+		TMap<FString, int32> MeshNamePartMap;
+		MeshNamePartMap.Reserve(UsedFoliage.Num());
+		TMap<FString, TArray<TObjectPtr<UMaterialInterface>>> MeshMaterialsMap;
+		MeshMaterialsMap.Reserve(UsedFoliage.Num() * 2);
+
+		const auto AddSkeletalMeshToAssembly = [&](const FString& FoliageInstanceObjectPath, TObjectPtr<USkeletalMesh> InSkeletalMesh)
+		{
+			const int32 PartIndex = AssemblyBuilder.FindOrAddPart(FoliageInstanceObjectPath);
+			MeshNamePartMap.FindOrAdd(FoliageInstanceObjectPath) = PartIndex;
+
+			TArray<TObjectPtr<UMaterialInterface>>& MeshMaterials = MeshMaterialsMap.FindOrAdd(FoliageInstanceObjectPath);
+			for (const FSkeletalMaterial& Material : InSkeletalMesh->GetMaterials())
+			{
+				MeshMaterials.AddUnique(Material.MaterialInterface);
+			}
+		};
+
+		for (const FString& FoliageInstanceObjectPath : UsedFoliage)
+		{
+			TObjectPtr<USkeletalMesh> PartSkeletalMesh = LoadObject<USkeletalMesh>(nullptr, *FoliageInstanceObjectPath);
+			if (PartSkeletalMesh)
+			{
+				AddSkeletalMeshToAssembly(FoliageInstanceObjectPath, PartSkeletalMesh);
+				continue;
+			}
+
+			const FString FoliageInstancePackagePath = FPackageName::ObjectPathToPackageName(FoliageInstanceObjectPath);
+			const FString LongFoliagePackagePath = FPackageName::GetLongPackagePath(FoliageInstancePackagePath);
+			const FString ShortFoliagePackageName = FPackageName::GetShortName(FoliageInstancePackagePath);
+
+			const FString SkeletalMeshName = "SKM_" + ShortFoliagePackageName;
+			const FString SkeletalMeshPackagePath = FPaths::Combine(LongFoliagePackagePath, SkeletalMeshName);
+			const FString SkeletalMeshObjectPath = FString::Printf(TEXT("%s.%s"), *SkeletalMeshPackagePath, *SkeletalMeshName);
+
+			const FString SkeletonName = "SK_" + ShortFoliagePackageName;
+			const FString SkeletonPackagePath = FPaths::Combine(LongFoliagePackagePath, SkeletonName);
+			const FString SkeletonObjectPath = FString::Printf(TEXT("%s.%s"), *SkeletonPackagePath, *SkeletonName);
+
+			PartSkeletalMesh = LoadObject<USkeletalMesh>(nullptr, SkeletalMeshObjectPath);
+			TObjectPtr<USkeleton> Skeleton = LoadObject<USkeleton>(nullptr, SkeletonObjectPath);
+			if (PartSkeletalMesh && Skeleton)
+			{
+				AddSkeletalMeshToAssembly(FoliageInstanceObjectPath, PartSkeletalMesh);
+				continue;
+			}
+
+			if (!Skeleton)
+			{
+				UPackage* SkeletonPackage = Internal::CreateAssetPackage(*SkeletonPackagePath);
+				OutModifiedPackages.Add(SkeletonPackage);
+
+				Skeleton = NewObject<USkeleton>(SkeletonPackage, FName(SkeletonName), RF_Standalone | RF_Public);
+
+				if (PartSkeletalMesh)
+				{
+					PartSkeletalMesh->SetSkeleton(Skeleton);
+					Skeleton->RecreateBoneTree(PartSkeletalMesh);
+				}
+			}
+
+			if (!PartSkeletalMesh)
+			{
+				const TObjectPtr<UStaticMesh> PartStaticMesh = LoadObject<UStaticMesh>(nullptr, *FoliageInstanceObjectPath);
+				if (!PartStaticMesh)
+				{
+					UE_LOG(LogProceduralVegetationEditor, Warning, TEXT("Could not find foliage instance %s"), *FoliageInstanceObjectPath);
+					continue;
+				}
+
+				FDynamicMesh3 PartDynamicMesh;
+				FText OutErrorMessage;
+				UE::Conversion::FStaticMeshConversionOptions Options;
+				if (!StaticMeshToDynamicMesh(PartStaticMesh, PartDynamicMesh, OutErrorMessage, Options))
+				{
+					UE_LOG(LogProceduralVegetationEditor, Warning, TEXT("Failed to convert foliage instance \"%s\" to dynamic mesh"), *FoliageInstancePackagePath);
+					continue;
+				}
+
+				UPackage* SkeletalMeshPackage = Internal::CreateAssetPackage(*SkeletalMeshPackagePath);
+				OutModifiedPackages.Add(SkeletalMeshPackage);
+
+				if (!Skeleton)
+				{
+					UPackage* SkeletonPackage = Internal::CreateAssetPackage(*SkeletonPackagePath);
+					OutModifiedPackages.Add(SkeletonPackage);
+					Skeleton = NewObject<USkeleton>(SkeletonPackage, FName(SkeletonName), RF_Standalone | RF_Public);
+				}
+
+				PartSkeletalMesh = NewObject<USkeletalMesh>(SkeletalMeshPackage, FName(SkeletalMeshName), RF_Standalone | RF_Public);
+				PartSkeletalMesh->SetSkeleton(Skeleton);
+
+				TArray<TObjectPtr<UMaterialInterface>> PartMaterials;
+				for (auto Material : PartStaticMesh->GetStaticMaterials())
+				{
+					PartMaterials.Add(Material.MaterialInterface);
+				}
+
+				ConvertToDefaultSkeletalMesh(PartSkeletalMesh, PartDynamicMesh, PartMaterials);
+			}
+			
+			AddSkeletalMeshToAssembly(FoliageInstanceObjectPath, PartSkeletalMesh);
+		}
+
+		AddNodeToBuilder(AssemblyBuilder, Collection, MeshNamePartMap, MeshMaterialsMap);
+		AssemblyBuilder.ApplyToSkeletalMesh(*SkeletalMesh);
+	}
+
+	void AddNodeToBuilder(FNaniteAssemblyDataBuilder& AssemblyBuilder, const FManagedArrayCollection& Collection,
+		const TMap<FString, int32>& InMeshNamePartMap, const TMap<FString, TArray<TObjectPtr<UMaterialInterface>>>& InMeshMaterialsMap)
+	{
+		TRACE_CPUPROFILER_EVENT_SCOPE(PV::Export::AddNodeToBuilder);
+
+		Facades::FFoliageFacade FoliageFacade(Collection);
+		const int32 NumInstances = FoliageFacade.NumFoliageEntries();
+			
+		for (int32 Id = 0; Id < NumInstances; Id++)
+		{
+			const FString FoliageName = GetFoliageName(FoliageFacade, Id);
+
+			const int32* PartIndex = InMeshNamePartMap.Find(FoliageName);
+			if (PartIndex == nullptr)
+			{
+				continue;
+			}
+				
+			const FTransform Transform = FoliageFacade.GetFoliageTransform(Id);
+			const UE::Geometry::FTransformSRT3d GeoTransform(Transform);
+
+			const int32 FoliageParentBoneID = FoliageFacade.GetParentBoneID(Id);
+				
+			FNaniteAssemblyBoneInfluence Influence;
+			Influence.BoneIndex = FoliageParentBoneID;
+			Influence.BoneWeight = 1.0;
+				
+			AssemblyBuilder.AddNode(*PartIndex, FTransform3f(GeoTransform), ENaniteAssemblyNodeTransformSpace::Local, { Influence });
+
+			//Add the materials for part to the assembly builder
+			auto* FoliageMaterials = InMeshMaterialsMap.Find(FoliageName);
+			if (FoliageMaterials == nullptr)
+			{
+				continue;
+			}
+
+			int32 LocalMaterialIndex = 0;
+			for (auto PartMaterial : *FoliageMaterials)
+			{
+				int32 PartMaterialIndex = AssemblyBuilder.GetMaterialSlots().IndexOfByPredicate(
+					[&] (const auto& Slot) { return Slot.Material == PartMaterial; }
+				);
+				if (PartMaterialIndex == INDEX_NONE)
+				{
+					PartMaterialIndex = AssemblyBuilder.AddMaterialSlot(PartMaterial);
+				}
+
+				AssemblyBuilder.RemapPartMaterial(*PartIndex, LocalMaterialIndex, PartMaterialIndex);
+				LocalMaterialIndex++;
+			}
+		}
+	}
+
+	void ConvertToDefaultSkeletalMesh(USkeletalMesh* SkeletalMesh, FDynamicMesh3& Mesh, const TArray<TObjectPtr<UMaterialInterface>>& Materials)
+	{
+		TRACE_CPUPROFILER_EVENT_SCOPE(PV::Export::ConvertToDefaultSkeletalMesh);
+
+		TObjectPtr<UDynamicMesh> DynamicMesh = NewObject<UDynamicMesh>();
+		DynamicMesh->SetMesh(MoveTemp(Mesh));
+		
+		//Figure out the correct parent bone index
+		int ParentBoneIndex = INDEX_NONE;
+		UE::Geometry::FDynamicMeshAttributeSet* MeshAttributes = DynamicMesh->GetMeshPtr()->Attributes();
+		MeshAttributes->EnableBones(1);
+
+		UE::Geometry::FDynamicMeshBoneNameAttribute* BoneNames = MeshAttributes->GetBoneNames();
+		UE::Geometry::FDynamicMeshBoneParentIndexAttribute* BoneParentIndices = MeshAttributes->GetBoneParentIndices();
+		UE::Geometry::FDynamicMeshBonePoseAttribute* BonePoses = MeshAttributes->GetBonePoses();
+
+		int BoneIndex = 0;
+		BoneNames->SetValue(BoneIndex, "Root");
+		BonePoses->SetValue(BoneIndex, FTransform::Identity);
+		BoneParentIndices->SetValue(BoneIndex, ParentBoneIndex);
+		
+		//Convert the dynamic mesh to static mesh
+		EGeometryScriptOutcomePins Output;
+		FGeometryScriptCopyMeshToAssetOptions Options = GetCopyMeshToAssetOptions(Materials, SkeletalMesh->GetNaniteSettings(), ENaniteShapePreservation::Voxelize);
+		Options.BoneHierarchyMismatchHandling = EGeometryScriptBoneHierarchyMismatchHandling::CreateNewReferenceSkeleton;
+		
+		UGeometryScriptLibrary_StaticMeshFunctions::CopyMeshToSkeletalMesh(
+			DynamicMesh,
+			SkeletalMesh,
+			Options,
+			{},
+			Output
+		);
+
+		TObjectPtr<USkeleton> Skeleton = SkeletalMesh->GetSkeleton();
+		Skeleton->RecreateBoneTree(SkeletalMesh);
+
+		const FGeometryScriptBoneWeightProfile SkinProfile = FGeometryScriptBoneWeightProfile();
+		FGeometryScriptSmoothBoneWeightsOptions SmoothBoneWeightsOptions;
+		SmoothBoneWeightsOptions.DistanceWeighingType = EGeometryScriptSmoothBoneWeightsType::GeodesicVoxel;
+		SmoothBoneWeightsOptions.MaxInfluences = 2;
+		SmoothBoneWeightsOptions.VoxelResolution = 512;
+		UGeometryScriptLibrary_MeshBoneWeightFunctions::ComputeSmoothBoneWeights(
+			DynamicMesh,
+			Skeleton,
+			SmoothBoneWeightsOptions,
+			SkinProfile
+		);
+
+		const FName ProfileName = SkinProfile.GetProfileName();
+		UGeometryScriptLibrary_StaticMeshFunctions::CopySkinWeightProfileToSkeletalMesh(
+			DynamicMesh,
+			SkeletalMesh,
+			ProfileName,
+			ProfileName,
+			{},
+			{},
+			Output
+		);
+	}
+
+	bool ExportCollectionToStaticMesh(
+		TObjectPtr<UStaticMesh> ExportMesh,
+		const FManagedArrayCollection& Collection,
+		ENaniteShapePreservation InShapePreservation,
+		bool bBuildNaniteAssemblies,
+		TArray<UPackage*>& OutModifiedPackages,
+		const FStatusReportCallback& StatusReportCallback
+	)
+	{
+		TRACE_CPUPROFILER_EVENT_SCOPE(PV::Export::ExportCollectionToStaticMesh);
+
+		const auto UpdateStage = [&](EExportMeshStage NewStage, float StageProgress)->bool
+		{
+			return StatusReportCallback(GetMeshExportStageText(NewStage, bBuildNaniteAssemblies, false), CalculateMeshExportProgress(NewStage, StageProgress, bBuildNaniteAssemblies, false));
+		};
+
+		FManagedArrayCollection ExportedCollection;
+		TObjectPtr<UDynamicMesh> DynamicMesh;
+		TArray<TObjectPtr<UMaterialInterface>> Materials;
+		{
+			TRACE_CPUPROFILER_EVENT_SCOPE(InitExportMesh);
+
+			if (!UpdateStage(EExportMeshStage::InitializingExportMesh, 0.f))
+			{
+				return false;
+			}
+
+			Collection.CopyTo(&ExportedCollection);
+
+			DynamicMesh = CollectionToDynamicMesh(ExportedCollection);
+
+			if (!UpdateStage(EExportMeshStage::InitializingExportMesh, 0.5f))
+			{
+				return false;
+			}
+
+			Materials = GetMaterialsFromCollection(ExportedCollection);
+
+			if (!UpdateStage(EExportMeshStage::InitializingExportMesh, 1.f))
+			{
+				return false;
+			}
+		}
+		
+		// Combine all the foliage instances into one dynamic mesh if not building nanite assemblies
+		if (!bBuildNaniteAssemblies)
+		{
+			if (!UpdateStage(EExportMeshStage::CombiningMeshes, 0.f))
+			{
+				return false;
+			}
+			if (!CombineMeshInstancesToDynamicMesh(
+				ExportedCollection,
+				DynamicMesh,
+				Materials,
+				[](FString MeshName, int32 InstanceID, int32 VertexStart, int32 VertexCount) {},
+				[&](float Progress) { return UpdateStage(EExportMeshStage::CombiningMeshes, Progress); }
+			))
+			{
+				return false;
+			}
+		}
+
+		EGeometryScriptOutcomePins Output;
+		{
+			TRACE_CPUPROFILER_EVENT_SCOPE(CopyMeshToStaticMesh);
+
+			if (!UpdateStage(EExportMeshStage::CopyingMeshDataToOutputMesh, 1.f))
+			{
+				return false;
+			}
+
+			FGeometryScriptCopyMeshToAssetOptions Options = GetCopyMeshToAssetOptions(Materials, ExportMesh->GetNaniteSettings(), InShapePreservation);
+			UGeometryScriptLibrary_StaticMeshFunctions::CopyMeshToStaticMesh(DynamicMesh, ExportMesh, Options, {}, Output);
+		}
+
+		if (bBuildNaniteAssemblies)
+		{
+			if (!UpdateStage(EExportMeshStage::BuildingNaniteAssemblies, 1.f))
+			{
+				return false;
+			}
+
+			if (NaniteAssembliesSupported())
+			{
+				//Build the nanite assemblies
+            	BuildNaniteAssemblyData(ExportedCollection, ExportMesh, OutModifiedPackages);
+			}
+			else
+			{
+				UE_LOG(LogProceduralVegetationEditor, Warning, TEXT("Failed to build Nanite Assemblies, because neither Nanite Foliage nor Nanite Assemblies are enabled for the project."));
+			}
+		}
+
+		{
+			TRACE_CPUPROFILER_EVENT_SCOPE(BuildSkinnedMesh);
+			if (!UpdateStage(EExportMeshStage::BuildingMesh, 1.f))
+			{
+				return false;
+			}
+
+			ExportMesh->Build(true);
+			FStaticMeshCompilingManager::Get().FinishAllCompilation();
+		}
+
+		UpdateStage(EExportMeshStage::Complete, 1.f);
+		return true;
+	}
+
+	void BuildWindImportData(FDynamicWindSkeletalImportData& OutImportData, const FManagedArrayCollection& Collection, const TArray<PV::Facades::FBoneNode>& Bones)
+	{
+		TRACE_CPUPROFILER_EVENT_SCOPE(PV::Export::BuildWindImportData);
+
+		Facades::FBranchFacade BranchFacade(Collection);
+
+		for (const auto& Bone : Bones)
+		{
+			int32 SimulationGroupIndex = BranchFacade.GetBranchSimulationGroupIndex(Bone.BranchIndex);
+			OutImportData.Joints.Add({Bone.BoneName, SimulationGroupIndex});
+		}
+	}
+
+	bool ExportCollectionToSkeletalMesh(
+		TObjectPtr<USkeletalMesh> ExportMesh,
+		const FManagedArrayCollection& Collection,
+		ENaniteShapePreservation InShapePreservation,
+		TObjectPtr<const UPVWindSettings> InWindSettings,
+		bool bBuildNaniteAssemblies,
+		TArray<UPackage*>& OutModifiedPackages,
+		const FStatusReportCallback& StatusReportCallback
+	)
+	{
+		TRACE_CPUPROFILER_EVENT_SCOPE(PV::Export::ExportCollectionToSkeletalMesh);
+
+		const auto UpdateStage = [&](EExportMeshStage NewStage, float StageProgress)->bool
+		{
+			return StatusReportCallback(
+				GetMeshExportStageText(NewStage, bBuildNaniteAssemblies, true), 
+				CalculateMeshExportProgress(NewStage, StageProgress, bBuildNaniteAssemblies, true)
+			);
+		};
+
+		FManagedArrayCollection ExportedCollection;
+		TObjectPtr<UDynamicMesh> DynamicMesh;
+		TArray<TObjectPtr<UMaterialInterface>> Materials;
+		TArray<int32> VertexBoneIDs;
+		TArray<Facades::FBoneNode> BoneNodes;
+		int32 FoliageVertexStart = 0;
+		{
+			TRACE_CPUPROFILER_EVENT_SCOPE(InitExportMesh);
+
+			if (!UpdateStage(EExportMeshStage::InitializingExportMesh, 0.f))
+			{
+				return false;
+			}
+
+			Collection.CopyTo(&ExportedCollection);
+
+			DynamicMesh = CollectionToDynamicMesh(ExportedCollection);
+
+			if (!UpdateStage(EExportMeshStage::InitializingExportMesh, 0.5f))
+			{
+				return false;
+			}
+
+			Materials = GetMaterialsFromCollection(ExportedCollection);
+
+			UE::Geometry::FDynamicMeshAttributeSet* MeshAttributes = DynamicMesh->GetMeshPtr()->Attributes();
+
+			Facades::FBoneFacade BoneFacade = Facades::FBoneFacade(ExportedCollection);
+			BoneNodes = BoneFacade.GetBoneDataFromCollection();
+
+			// if bones already not created with bone reduction node create bones with full density
+			if (BoneNodes.IsEmpty())
+			{
+				BoneFacade.CreateBoneData(BoneNodes, 0);
+			}
+
+			AssignBoneIDsToFoliage(BoneNodes, ExportedCollection);
+
+			const int32 BoneCount = BoneNodes.Num();
+			MeshAttributes->EnableBones(BoneCount);
+
+			TArray<int32> VertexPointIDs;
+			BoneFacade.GetPointIds(VertexPointIDs);
+			
+			// Assign the Bone id to every vertex in the base mesh through VertexPointIDs
+			for (const int32& PointID : VertexPointIDs)
+			{
+				if (Facades::FBoneNode* BoneNode = Facades::FBoneFacade::FindClosestBone(ExportedCollection, BoneNodes, PointID))
+				{
+					VertexBoneIDs.Add(BoneNode->BoneIndex);
+				}
+				else
+				{
+					VertexBoneIDs.Add(INDEX_NONE);
+					UE_LOG(LogProceduralVegetationEditor, Log, TEXT("Invalid Bone Assigned to PointID %i"), PointID);
+				}
+			}
+
+			FoliageVertexStart = VertexBoneIDs.Num();
+
+			UE::Geometry::FDynamicMeshBoneNameAttribute* BoneNames = MeshAttributes->GetBoneNames();
+			UE::Geometry::FDynamicMeshBoneParentIndexAttribute* BoneParentIndices = MeshAttributes->GetBoneParentIndices();
+			UE::Geometry::FDynamicMeshBonePoseAttribute* BonePoses = MeshAttributes->GetBonePoses();
+
+			for (int i = 0; i < BoneCount; i++)
+			{
+				BoneNames->SetValue(i, BoneNodes[i].BoneName);
+				BonePoses->SetValue(i, BoneNodes[i].BoneTransform);
+				BoneParentIndices->SetValue(i, BoneNodes[i].ParentBoneIndex);
+			}
+
+			if (!UpdateStage(EExportMeshStage::InitializingExportMesh, 1.f))
+			{
+				return false;
+			}
+		}
+
+		if (!bBuildNaniteAssemblies)
+		{
+			if (!UpdateStage(EExportMeshStage::CombiningMeshes, 0.f))
+			{
+				return false;
+			}
+
+			Facades::FFoliageFacade FoliageFacade(ExportedCollection);
+
+			// Combine all the foliage instances into one dynamic mesh if not building nanite assemblies
+			if (!CombineMeshInstancesToDynamicMesh(
+				ExportedCollection,
+				DynamicMesh,
+				Materials,
+				[&](const FString& MeshName, int32 InstanceID, int32 VertexStart, int32 VertexCount)
+				{
+					// Assign parent bone id to foliage vertices
+					int32 FoliageParentBoneID = FoliageFacade.GetParentBoneID(InstanceID);
+
+					for (int VertexID = VertexStart; VertexID < VertexCount; ++VertexID)
+					{
+						VertexBoneIDs.Add(FoliageParentBoneID);
+					}
+				},
+				[&](float Progress) { return UpdateStage(EExportMeshStage::CombiningMeshes, Progress); }
+			))
+			{
+				return false;
+			}
+		}
+
+		// Convert the dynamic mesh to static mesh
+		EGeometryScriptOutcomePins Output;
+		{
+			TRACE_CPUPROFILER_EVENT_SCOPE(CopyMeshToSkeletalMesh);
+
+			if (!UpdateStage(EExportMeshStage::CopyingMeshDataToOutputMesh, 1.f))
+			{
+				return false;
+			}
+
+			FGeometryScriptCopyMeshToAssetOptions Options = GetCopyMeshToAssetOptions(Materials, ExportMesh->GetNaniteSettings(), InShapePreservation);
+			Options.BoneHierarchyMismatchHandling = EGeometryScriptBoneHierarchyMismatchHandling::CreateNewReferenceSkeleton;
+			UGeometryScriptLibrary_StaticMeshFunctions::CopyMeshToSkeletalMesh(
+				DynamicMesh,
+				ExportMesh,
+				Options,
+				{},
+				Output
+			);
+		}
+
+		TObjectPtr<USkeleton> Skeleton = ExportMesh->GetSkeleton();
+		{
+			TRACE_CPUPROFILER_EVENT_SCOPE(RecreateBoneTree);
+			if (!UpdateStage(EExportMeshStage::RecreatingBoneTree, 1.f))
+			{
+				return false;
+			}
+			Skeleton->RecreateBoneTree(ExportMesh);
+		}
+
+		const FGeometryScriptBoneWeightProfile SkinProfile = FGeometryScriptBoneWeightProfile();
+		const FName ProfileName = SkinProfile.GetProfileName();
+
+		{
+			TRACE_CPUPROFILER_EVENT_SCOPE(ComputeSmoothBoneWeights);
+
+			if (!UpdateStage(EExportMeshStage::ComputingBoneWeights, 1.f))
+			{
+				return false;
+			}
+
+			FGeometryScriptSmoothBoneWeightsOptions SmoothBoneWeightsOptions;
+			SmoothBoneWeightsOptions.DistanceWeighingType = EGeometryScriptSmoothBoneWeightsType::DirectDistance;
+			SmoothBoneWeightsOptions.MaxInfluences = 3;
+			SmoothBoneWeightsOptions.Stiffness = 0.5;
+
+			UGeometryScriptLibrary_MeshBoneWeightFunctions::ComputeSmoothBoneWeights(
+				DynamicMesh,
+				Skeleton,
+				SmoothBoneWeightsOptions,
+				SkinProfile
+			);
+
+			if (!UpdateStage(EExportMeshStage::ComputingBoneWeights, 0.8f))
+			{
+				return false;
+			}
+
+			// Here we make the correction to the weights
+			RemoveUnwantedSkinWeights(DynamicMesh, ProfileName, VertexBoneIDs, FoliageVertexStart, BoneNodes);
+
+			if (!UpdateStage(EExportMeshStage::ComputingBoneWeights, 1.f))
+			{
+				return false;
+			}
+		}
+		
+		{
+			TRACE_CPUPROFILER_EVENT_SCOPE(CopySkinWeightProfileToSkeletalMesh);
+
+			if (!UpdateStage(EExportMeshStage::CopyingSkinWeightProfiles, 1.f))
+			{
+				return false;
+			}
+
+			FGeometryScriptCopySkinWeightProfileToAssetOptions Options;
+			Options.bEmitTransaction = false;
+			Options.bDeferMeshPostEditChange = true;
+			UGeometryScriptLibrary_StaticMeshFunctions::CopySkinWeightProfileToSkeletalMesh(
+				DynamicMesh,
+				ExportMesh,
+				ProfileName,
+				ProfileName,
+				Options,
+				{},
+				Output
+			);
+		}
+
+		if (bBuildNaniteAssemblies)
+		{
+			if (!UpdateStage(EExportMeshStage::BuildingNaniteAssemblies, 1.f))
+			{
+				return false;
+			}
+
+			if (NaniteAssembliesSupported())
+			{
+				//Build the nanite assemblies
+				BuildNaniteAssemblyData(ExportedCollection, ExportMesh, OutModifiedPackages);
+			}
+			else
+			{
+				UE_LOG(LogProceduralVegetationEditor, Warning, TEXT("Failed to build Nanite Assemblies, because neither Nanite Foliage nor Nanite Assemblies are enabled for the project."));
+			}
+		}
+
+		{
+			TRACE_CPUPROFILER_EVENT_SCOPE(DynamicWind);
+
+			if (!UpdateStage(EExportMeshStage::ImportingDynamicWind, 1.f))
+			{
+				return false;
+			}
+
+			const bool bWindDataExist = ExportMesh->HasAssetUserDataOfClass(UDynamicWindSkeletalData::StaticClass());
+			const bool bOverwriteWindSettings = InWindSettings ? InWindSettings->bOverwriteExisting : false;
+			
+			FDynamicWindSkeletalImportData ImportData;
+			BuildWindImportData(ImportData, ExportedCollection, BoneNodes);
+
+			if (InWindSettings)
+			{
+				ImportData.SimulationGroups = InWindSettings->SimulationGroupData;
+			}
+		
+			if (bWindDataExist && !bOverwriteWindSettings)
+			{
+				//Get Existing WindSettings
+				UDynamicWindSkeletalData* WindSkeletalData = Cast<UDynamicWindSkeletalData>(ExportMesh->GetAssetUserDataOfClass(UDynamicWindSkeletalData::StaticClass()));
+				if (WindSkeletalData)
+				{
+					ImportData.SimulationGroups = WindSkeletalData->SimulationGroups;
+				}
+			}
+
+			if (!DynamicWind::ImportSkeletalData(*ExportMesh, ImportData))
+			{
+				UE_LOG(LogProceduralVegetationEditor, Warning, TEXT("Failed to build wind data"));
+			}
+		}
+
+		{
+			TRACE_CPUPROFILER_EVENT_SCOPE(BuildSkinnedMesh);
+			if (!UpdateStage(EExportMeshStage::BuildingMesh, 1.f))
+			{
+				return false;
+			}
+
+			ExportMesh->Build();
+			FSkinnedAssetCompilingManager::Get().FinishAllCompilation();
+		}
+
+		UpdateStage(EExportMeshStage::Complete, 1.f);
+		return true;
+	}
+
+	void AttachProceduralVegetationLink(UObject* InExportedMesh, const TObjectPtr<UProceduralVegetation>& InProceduralVegetation)
+	{
+		if (IInterface_AssetUserData* IAssetUserData = Cast<IInterface_AssetUserData>(InExportedMesh))
+		{
+			UProceduralVegetationLink* Data = NewObject<UProceduralVegetationLink>(InExportedMesh);
+			Data->Source = InProceduralVegetation;
+		
+			IAssetUserData->AddAssetUserData(Data);	
+		}
+	}
+
+	void AssignBoneIDsToFoliage(const TArray<Facades::FBoneNode>& Bones, FManagedArrayCollection& Collection)
+	{
+		TRACE_CPUPROFILER_EVENT_SCOPE(PV::Export::AssignBoneIDsToFoliage);
+
+		Facades::FPointFacade PointFacade(Collection);
+		Facades::FFoliageFacade FoliageFacade(Collection);
+		
+		int32 NumInstances = FoliageFacade.NumFoliageEntries();
+
+		auto FindBoneById = ( [&](const int32 Id)
+		{
+			const Facades::FBoneNode* BoneNode = Bones.FindByPredicate([Id](const Facades::FBoneNode& Node)
+			{
+				return Node.BoneIndex == Id;
+			});
+			return BoneNode;
+		});
+		
+		for(int32 Id = 0; Id < NumInstances; Id++)
+		{
+			Facades::FFoliageEntryData Data = FoliageFacade.GetFoliageEntry(Id);
+			float FoliageLFR = Data.LengthFromRoot;
+
+			TArray<int32> BoneIDs;
+			for (Facades::FBoneNode BoneNode : Bones)
+			{
+				if (Data.BranchId == BoneNode.BranchIndex)
+				{
+					BoneIDs.Add(BoneNode.BoneIndex);
+				}
+			}
+
+			bool bBoneIDAssigned = false;
+			for (const int32& BoneId : BoneIDs)
+			{
+				const Facades::FBoneNode* BoneNode = FindBoneById(BoneId);
+				check(BoneNode);
+				
+				int PointIndex = BoneNode->PointIndex;
+				const Facades::FBoneNode* BoneParentNode = FindBoneById(BoneNode->ParentBoneIndex);
+				
+				int ParentPointIndex = BoneParentNode ? BoneParentNode->PointIndex : INDEX_NONE;
+				float BoneLFR = PointFacade.GetLengthFromRoot(PointIndex);
+				float BoneParentLFR = PointFacade.GetLengthFromRoot(ParentPointIndex);
+
+				if (FoliageLFR <= BoneLFR && FoliageLFR >= BoneParentLFR)
+				{
+					FoliageFacade.SetParentBoneID(Id, BoneId);
+					bBoneIDAssigned = true;
+					break;
+				}
+			}
+
+			if (!bBoneIDAssigned)
+			{
+				UE_LOG(LogProceduralVegetationEditor, Warning, TEXT("No bone assigned for foliage Id {%i} , Branch Bone count {%i}"), Id , BoneIDs.Num());
+			}
+		}
+	}
+
+	void RemoveUnwantedSkinWeights(TObjectPtr<UDynamicMesh> DynamicMesh, const FName ProfileName, const TArray<int32>& VertexBoneIDs, int32 FoliageVertexStart, const TArray<Facades::FBoneNode>& BoneNodes)
+	{
+		TRACE_CPUPROFILER_EVENT_SCOPE(PV::Export::RemoveUnwantedSkinWeights);
+
+		auto InMesh = DynamicMesh->GetMeshPtr();
+		check(InMesh)
+		const int32 NumVertices = VertexBoneIDs.Num();
+		
+		DynamicMesh->EditMesh([&](FDynamicMesh3& EditMesh)
+		{
+			UE::Geometry::FDynamicMeshVertexSkinWeightsAttribute *SkinWeights = EditMesh.Attributes()->GetSkinWeightsAttribute(ProfileName);
+
+			ParallelFor(NumVertices, [&](const int32 VertexIdx)
+			{
+				int32 BoneID = VertexBoneIDs[VertexIdx];
+
+				if (BoneID != INDEX_NONE)
+				{
+					int32 ParentBoneID = BoneNodes[BoneID].ParentBoneIndex;
+				
+					UE::AnimationCore::FBoneWeights Weights;
+					SkinWeights->GetValue(VertexIdx, Weights);
+
+					auto NextBone = BoneNodes.FindByPredicate([&](const Facades::FBoneNode& Node)
+					{
+						return Node.BoneIndex == BoneID + 1 && Node.BranchIndex == BoneNodes[BoneID].BranchIndex;
+					});
+
+					for (int32 WeightID = Weights.Num() - 1 ; WeightID >= 0; WeightID--)
+					{
+						auto Weight = Weights[WeightID];
+
+						bool bValidBone = (Weight.GetBoneIndex() == ParentBoneID);
+						bValidBone |= (Weight.GetBoneIndex() == BoneID);
+						bValidBone |= NextBone && (Weight.GetBoneIndex() == NextBone->BoneIndex);
+						
+						if (!bValidBone || VertexIdx > FoliageVertexStart || BoneID == 0)
+						{
+							UE::AnimationCore::FBoneWeightsSettings Settings;
+							Settings.SetNormalizeType(UE::AnimationCore::EBoneWeightNormalizeType::Always);
+							Weights.RemoveBoneWeight(Weight.GetBoneIndex(),Settings);
+						}
+					}
+
+					if (VertexIdx > FoliageVertexStart || BoneID == 0 || Weights.Num() == 0)
+					{
+						UE::AnimationCore::FBoneWeightsSettings Settings;
+						Settings.SetNormalizeType(UE::AnimationCore::EBoneWeightNormalizeType::Always);
+						Weights.SetBoneWeight(UE::AnimationCore::FBoneWeight(BoneID, 1.0f), Settings);
+					}
+					
+					SkinWeights->SetNewValue(VertexIdx, Weights);
+				}
+				else
+				{
+					UE_LOG(LogProceduralVegetationEditor, Log, TEXT("Invalid bone index found while removing bones for vertex %i"), VertexIdx);
+				}
+			});
+				
+		}, EDynamicMeshChangeType::GeneralEdit, EDynamicMeshAttributeChangeFlags::Unknown, false);
+	}
+}
+
+namespace PV::Export
+{
+	bool ExportCollectionAsMesh(
+		const TObjectPtr<UProceduralVegetation> InProceduralVegetation,
+		const FManagedArrayCollection& Collection,
+		const FPVExportParams& ExportParams,
+		const FStatusReportCallback& StatusReportCallback
+	)
+	{
+		TRACE_CPUPROFILER_EVENT_SCOPE(PV::Export::ExportCollectionAsMesh);
+
+		const FString MeshName = ExportParams.MeshName.ToString();
+		const FString MeshPackageName = ExportParams.GetOutputMeshPackagePath();
+		const bool bCreateNaniteFoliage = ExportParams.bCreateNaniteFoliage && NaniteAssembliesSupported();
+
+		TArray<UPackage*> ModifiedPackages;
+
+		if (ExportParams.ExportMeshType == EPVExportMeshType::StaticMesh)
+		{
+			if (Utilities::DoesConflictingPackageExist(MeshPackageName, UStaticMesh::StaticClass()))
+			{
+				return false;
+			}
+
+			const TObjectPtr<UStaticMesh> ExistingMeshAsset = Internal::GetAssetOfType<UStaticMesh>(MeshPackageName);
+
+			const TObjectPtr<UStaticMesh> ExportMesh = ExistingMeshAsset != nullptr
+				? ExistingMeshAsset.Get()
+				: NewObject<UStaticMesh>(Internal::CreateAssetPackage(MeshPackageName), *MeshName, RF_Standalone | RF_Public);
+
+			ModifiedPackages.Add(ExportMesh->GetOutermost());
+
+			if (!PV::Export::Internal::ExportCollectionToStaticMesh(
+				ExportMesh,
+				Collection,
+				ExportParams.NaniteShapePreservation,
+				bCreateNaniteFoliage,
+				ModifiedPackages,
+				StatusReportCallback
+			))
+			{
+				Internal::CleanupPackages(ModifiedPackages);
+				return false;
+			}
+
+			Internal::AttachProceduralVegetationLink(ExportMesh, InProceduralVegetation);
+		}
+		else if (ExportParams.ExportMeshType == EPVExportMeshType::SkeletalMesh)
+		{
+			if (Utilities::DoesConflictingPackageExist(*MeshPackageName, USkeletalMesh::StaticClass()))
+			{
+				return false;
+			}
+
+			const FString SkeletonName = ExportParams.GetOutputSkeletonName();
+			const FString SkeletonPackageName = ExportParams.GetOutputSkeletonPackagePath();
+			if (Utilities::DoesConflictingPackageExist(SkeletonPackageName, USkeleton::StaticClass()))
+			{
+				return false;
+			}
+
+			const TObjectPtr<USkeletalMesh> ExistingMeshAsset = Internal::GetAssetOfType<USkeletalMesh>(MeshPackageName);
+			const TObjectPtr<USkeleton> ExistingSkeletonAsset = Internal::GetAssetOfType<USkeleton>(SkeletonPackageName);
+
+			const TObjectPtr<USkeletalMesh> ExportMesh = ExistingMeshAsset != nullptr 
+				? ExistingMeshAsset.Get() 
+				: NewObject<USkeletalMesh>(Internal::CreateAssetPackage(MeshPackageName), *MeshName, RF_Standalone | RF_Public);
+			const TObjectPtr<USkeleton> ExportSkeleton = ExistingSkeletonAsset != nullptr 
+				? ExistingSkeletonAsset.Get() 
+				: NewObject<USkeleton>(Internal::CreateAssetPackage(SkeletonPackageName), *SkeletonName, RF_Standalone | RF_Public);
+			
+			ExportMesh->SetSkeleton(ExportSkeleton);
+
+			ModifiedPackages.Add(ExportMesh->GetOutermost());
+			ModifiedPackages.Add(ExportSkeleton->GetOutermost());
+
+			if (!PV::Export::Internal::ExportCollectionToSkeletalMesh(
+				ExportMesh,
+				Collection,
+				ExportParams.NaniteShapePreservation,
+				ExportParams.WindSettings,
+				bCreateNaniteFoliage,
+				ModifiedPackages,
+				StatusReportCallback
+			))
+			{
+				Internal::CleanupPackages(ModifiedPackages);
+				return false;
+			}
+
+			Internal::AttachProceduralVegetationLink(ExportMesh, InProceduralVegetation);
+			Internal::AttachProceduralVegetationLink(ExportSkeleton, InProceduralVegetation);
+		}
+
+		for (UPackage* Package : ModifiedPackages)
+		{
+			Package->MarkPackageDirty();
+			IAssetRegistry::Get()->AssetCreated(Package);
+		}
+
+		return true;
+	}
+}
+
+#undef LOCTEXT_NAMESPACE 
diff --git a/Engine/Plugins/Experimental/ProceduralVegetationEditor/Source/ProceduralVegetationEditor/Private/Helpers/PVExportHelper.h b/Engine/Plugins/Experimental/ProceduralVegetationEditor/Source/ProceduralVegetationEditor/Private/Helpers/PVExportHelper.h
new file mode 100644
--- /dev/null
+++ b/Engine/Plugins/Experimental/ProceduralVegetationEditor/Source/ProceduralVegetationEditor/Private/Helpers/PVExportHelper.h
@@ -0,0 +1,98 @@
+﻿// Copyright Epic Games, Inc. All Rights Reserved.
+#pragma once
+
+#include "NaniteAssemblyDataBuilder.h"
+#include "Facades/PVBoneFacade.h"
+#include "Engine/EngineTypes.h"
+
+class UInstancedStaticMeshComponent;
+class UStaticMesh;
+class USkeletalMesh;
+class UPackage;
+class UProceduralVegetation;
+class UMaterialInterface;
+class UPVWindSettings;
+struct FGeometryScriptCopyMeshToAssetOptions;
+struct FManagedArrayCollection;
+struct FPVExportParams;
+
+class UDynamicMesh;
+namespace UE { namespace Geometry { class FDynamicMesh3; } }
+using UE::Geometry::FDynamicMesh3;
+
+namespace PV::Export
+{
+	typedef TFunction<bool(const FText& Stage, float Progress)> FStatusReportCallback;
+
+	bool ExportCollectionAsMesh(
+		const TObjectPtr<UProceduralVegetation> InProceduralVegetation,
+		const FManagedArrayCollection& Collection,
+		const FPVExportParams& ExportParams,
+		const FStatusReportCallback& StatusReportCallback
+	);
+}
+
+namespace PV::Export::Internal
+{
+	typedef TFunction<void(const FString& MeshName, int32 InstanceID, int32 VertexStart, int32 VertexCount)> FMeshInstanceCombined;
+	
+	TObjectPtr<UDynamicMesh> CollectionToDynamicMesh(const FManagedArrayCollection& Collection);
+
+	FGeometryScriptCopyMeshToAssetOptions GetCopyMeshToAssetOptions(
+		const TArray<TObjectPtr<UMaterialInterface>>& InMaterials,
+		const FMeshNaniteSettings& InNaniteSettings,
+		const ENaniteShapePreservation InShapePreservation
+	);
+	
+	bool CombineMeshInstancesToDynamicMesh(
+		const FManagedArrayCollection& Collection,
+		UDynamicMesh* DynamicMesh,
+		TArray<TObjectPtr<UMaterialInterface>>& Materials,
+		const FMeshInstanceCombined& OnMeshInstanceCombined,
+		const TFunction<bool(float)>& OnProgressUpdated
+	);
+
+	void BuildNaniteAssemblyData(
+		const FManagedArrayCollection& Collection,
+		UStaticMesh* StaticMesh,
+		TArray<UPackage*>& OutModifiedPackages
+	);
+
+	void AddNodeToBuilder(
+		FNaniteAssemblyDataBuilder& AssemblyBuilder,
+		const FManagedArrayCollection& Collection,
+		const TMap<FString, int32>& InMeshNamePartMap,
+		const TMap<FString, TArray<TObjectPtr<UMaterialInterface>>>& InMeshMaterialsMap
+	);
+	
+	bool ExportCollectionToStaticMesh(
+		TObjectPtr<UStaticMesh> ExportMesh,
+		const FManagedArrayCollection& Collection,
+		ENaniteShapePreservation InShapePreservation,
+		bool bBuildNaniteAssemblies,
+		TArray<UPackage*>& OutModifiedPackages,
+		const FStatusReportCallback& StatusReportCallback
+	);
+
+	bool ExportCollectionToSkeletalMesh(
+		TObjectPtr<USkeletalMesh> ExportMesh,
+		const FManagedArrayCollection& Collection,
+		ENaniteShapePreservation InShapePreservation,
+		TObjectPtr<const UPVWindSettings> InWindSettings,
+		bool bBuildNaniteAssemblies,
+		TArray<UPackage*>& OutModifiedPackages,
+		const FStatusReportCallback& StatusReportCallback
+	);
+
+	void AttachProceduralVegetationLink(UObject* InExportedMesh, const TObjectPtr<UProceduralVegetation>& InProceduralVegetation);
+	
+	void AssignBoneIDsToFoliage(const TArray<Facades::FBoneNode>& Bones, FManagedArrayCollection& Collection);
+
+	void RemoveUnwantedSkinWeights(TObjectPtr<UDynamicMesh> DynamicMesh, const FName ProfileName, const TArray<int32>& VertexBoneIDs,
+		int32 FoliageVertexStart, const TArray<Facades::FBoneNode>& BoneNodes);
+
+	void ConvertToDefaultSkeletalMesh(USkeletalMesh* SkeletalMesh,
+		FDynamicMesh3& Mesh,
+		const TArray<TObjectPtr<UMaterialInterface>>& Materials
+	);
+}
\ No newline at end of file
diff --git a/Engine/Plugins/Experimental/ProceduralVegetationEditor/Source/ProceduralVegetationEditor/Private/PVEditor.cpp b/Engine/Plugins/Experimental/ProceduralVegetationEditor/Source/ProceduralVegetationEditor/Private/PVEditor.cpp
--- a/Engine/Plugins/Experimental/ProceduralVegetationEditor/Source/ProceduralVegetationEditor/Private/PVEditor.cpp
+++ b/Engine/Plugins/Experimental/ProceduralVegetationEditor/Source/ProceduralVegetationEditor/Private/PVEditor.cpp
@@ -9,9 +9,11 @@
 #include "PVEditorCommands.h"
 #include "PVEditorSchema.h"
 #include "PVEditorSettings.h"
+#include "Helpers/PVExportHelper.h"
 #include "ToolMenus.h"
 #include "ToolMenuSection.h"
 #include "Misc/UObjectToken.h"
+#include "FileHelpers.h"
 
 #include "Dataflow/DataflowCollectionSpreadSheetWidget.h"
 
@@ -358,88 +360,227 @@ FReply FPVEditor::ExportInternal()
 		return FReply::Handled();
 	}
 
-	const TArray<UPCGNode*> NodesToExport = ProceduralVegetationGraph->GetNodes().FilterByPredicate(
-		[](UPCGNode* Node) 
+	struct FScopedErrorMessages
+	{
+		TArray<TSharedRef<FTokenizedMessage>> ErrorMessages;
+
+		TSharedRef<FTokenizedMessage> AddError(const FText& ErrorMsg)
 		{
-			const UPVOutputSettings* OutputSettings = Cast<UPVOutputSettings>(Node->GetSettings());
-			return OutputSettings && OutputSettings->ExportSettings.bShouldExport;
+			TSharedRef<FTokenizedMessage> Error = FTokenizedMessage::Create(EMessageSeverity::Error, ErrorMsg);
+			ErrorMessages.Add(Error);
+			return Error;
 		}
-	);
 
-	TArray<TSharedRef<FTokenizedMessage>> ErrorMessages;
-	const auto AddError = [&](const FText& ErrorMsg)->TSharedRef<FTokenizedMessage>
-	{
-		TSharedRef<FTokenizedMessage> Error = FTokenizedMessage::Create(EMessageSeverity::Error, ErrorMsg);
-		ErrorMessages.Add(Error);
-		return Error;
+		~FScopedErrorMessages()
+		{
+			if (ErrorMessages.Num() > 0)
+			{
+				FMessageLog EditorErrors("EditorErrors");
+				EditorErrors.AddMessages(ErrorMessages);
+				EditorErrors.Notify(LOCTEXT("PVEMeshExportFailed", "PVE Mesh Export Failed"));
+			}
+		}
 	};
 
-	FScopedSlowTask Progress(NodesToExport.Num(), LOCTEXT("ExportAsset_SlowTask", "Exporting Asset(s) ..."));
-	Progress.MakeDialog();
+	FScopedErrorMessages ErrorMessages;
 
-	for (int32 i = 0; i < NodesToExport.Num(); ++i)
+	const static auto ValidateExportSettings = [](const TArray<UPCGNode*>& NodesToExport, FScopedErrorMessages& ErrorMessages)->bool
 	{
-		UPCGNode* Node = NodesToExport[i];
-		const UPVOutputSettings* OutputSettings = Cast<UPVOutputSettings>(Node->GetSettings());
-		check(OutputSettings != nullptr);
+		// Check that the paths are valid
+		bool bHasInvalidPath = false;
+		for (const UPCGNode* Node : NodesToExport)
+		{
+			const UPVOutputSettings* OutputSettings = CastChecked<const UPVOutputSettings>(Node->GetSettings());
 
-		const FText AssetCountText = NodesToExport.Num() > 1
-			? FText::FromString(FString::Printf(TEXT(" [%d/%d]"), i + 1, NodesToExport.Num()))
-			: FText::FromString("");
+			FString PathError;
+			if (!OutputSettings->ExportSettings.Validate(PathError))
+			{
+				TSharedRef<FTokenizedMessage> Error = ErrorMessages.AddError(
+					FText::Format(
+						LOCTEXT("ExportFailed_InvalidPath", "Invalid export path for node \"{0}\": \"{1}\""),
+						Node->GetDefaultTitle(),
+						FText::FromString(PathError)
+					)
+				);
+				Error->AddToken(FUObjectToken::Create(OutputSettings));
+				bHasInvalidPath = true;
+			}
+		}
 
-		Progress.EnterProgressFrame(1.f, FText::Format(LOCTEXT("ExportingAsset", "Exporting \"{0}\"{1}"), FText::FromName(OutputSettings->ExportSettings.MeshName), AssetCountText));
+		return !bHasInvalidPath;
+	};
+
+	const static auto ShouldOverwriteOutstandingChanges = [](const TArray<UPCGNode*>& NodesToExport)->bool 
+	{
+		TArray<FString> ExportPaths;
+		ExportPaths.Reserve(NodesToExport.Num() * 2);
 
-		FManagedArrayCollection OutCollection;
-		FName PinName = NAME_None;
-		if (!GetCollectionFromNode(Node, PinName, OutCollection))
+		for (const UPCGNode* Node : NodesToExport)
 		{
-			TSharedRef<FTokenizedMessage> Error = AddError(
-				FText::Format(
-					LOCTEXT("ExportFailed_NodeErrors", "Export failed for output node \"{0}\", Invalid output please fix the errors on the node."), 
-					Node->GetAuthoredTitleLine()
-				)
-			);
-			Error->AddToken(FUObjectToken::Create(Node));
-			continue;
+			const UPVOutputSettings* OutputSettings = CastChecked<const UPVOutputSettings>(Node->GetSettings());
+			const FPVExportParams& ExportParams = OutputSettings->ExportSettings;
+			ExportPaths.Add(ExportParams.GetOutputMeshPackagePath());
+			if (ExportParams.ExportMeshType == EPVExportMeshType::SkeletalMesh)
+			{
+				ExportPaths.Add(ExportParams.GetOutputSkeletonPackagePath());
+			}
 		}
 
-		FString PathError;
-		UClass* MeshClass = OutputSettings->ExportSettings.ExportMeshType == EPVExportMeshType::StaticMesh ? UStaticMesh::StaticClass() : USkeletalMesh::StaticClass();
-		if (!PV::Utilities::ValidateMeshPathAndName(OutputSettings->ExportSettings.MeshName.ToString(), OutputSettings->ExportSettings.ContentBrowserFolder.Path,
-			MeshClass, PathError))
+		TArray<UPackage*> DirtyPackages;
+		DirtyPackages.Reserve(ExportPaths.Num());
+		for (const FString& AssetPath : ExportPaths)
 		{
-			TSharedRef<FTokenizedMessage> Error = AddError(
-				FText::Format(
-					LOCTEXT("ExportFailed_InvalidPath", "Export failed for mesh \"{0}\" with error \"{1}\""),
-					Node->GetAuthoredTitleLine(),
-					FText::FromString(PathError)
-				)
-			);
-			Error->AddToken(FUObjectToken::Create(Node));
-			continue;
+			UPackage* ExistingPackage = FindPackage(nullptr, *AssetPath);
+			if (ExistingPackage && ExistingPackage->IsDirty())
+			{
+				DirtyPackages.Add(ExistingPackage);
+			}
+		}
+
+		if (DirtyPackages.Num() == 0)
+		{
+			return true;
+		}
+
+		FEditorFileUtils::FPromptForCheckoutAndSaveParams SaveParams;
+		SaveParams.bCheckDirty = false;
+		SaveParams.bPromptToSave = true;
+		SaveParams.bAlreadyCheckedOut = false;
+		SaveParams.bCanBeDeclined = true;
+		SaveParams.bIsExplicitSave = false;
+		SaveParams.Title = LOCTEXT("PackagesNeedResave_Title", "WARNING: Unsaved changes may be lost!");
+		SaveParams.Message = LOCTEXT("PackagesNeedResave_Message", "Assets have unsaved changes which may be lost during the export process. Please save the assets, or cancel, if you do not want to risk loosing changes to the assets.");
+
+		const FEditorFileUtils::EPromptReturnCode Return = FEditorFileUtils::PromptForCheckoutAndSave(DirtyPackages, SaveParams);
+		if (Return == FEditorFileUtils::EPromptReturnCode::PR_Cancelled)
+		{
+			return false;
+		}
+
+		return true;
+	};
+
+	const auto GetOutputNodeCollections = [this](TArray<FManagedArrayCollection>& OutCollections, const TArray<UPCGNode*>& NodesToExport, FScopedErrorMessages& ErrorMessages)->bool
+	{
+		OutCollections.Reset(NodesToExport.Num());
+		for (UPCGNode* Node : NodesToExport)
+		{
+			FManagedArrayCollection& OutCollection = OutCollections.AddDefaulted_GetRef();
+			FName PinName = NAME_None;
+			if (!GetCollectionFromNode(Node, PinName, OutCollection))
+			{
+				TSharedRef<FTokenizedMessage> Error = ErrorMessages.AddError(
+					FText::Format(
+						LOCTEXT("ExportFailed_NodeErrors", "Export failed for output node \"{0}\", Invalid output please fix the errors on the node."),
+						Node->GetDefaultTitle()
+					)
+				);
+				Error->AddToken(FUObjectToken::Create(Node));
+				return false;
+			}
 		}
 		
-		FScopedSlowTask ExportCollectionSlowTask(1.f,LOCTEXT("GeneratingMesh", "Generating Mesh ..."));
-		ExportCollectionSlowTask.MakeDialog();
+		return true;
+	};
+
+	const TArray<UPCGNode*> NodesToExport = ProceduralVegetationGraph->GetNodes().FilterByPredicate([](UPCGNode* InNode)->bool {
+			const UPVOutputSettings* OutputSettings = Cast<const UPVOutputSettings>(InNode->GetSettings());
+			return OutputSettings && OutputSettings->ExportSettings.bShouldExport;
+		});
+
+	if (!ValidateExportSettings(NodesToExport, ErrorMessages))
+	{
+		return FReply::Handled();
+	}
+
+	if (!ShouldOverwriteOutstandingChanges(NodesToExport))
+	{
+		return FReply::Handled();
+	}
+	
+	TArray<FManagedArrayCollection> CollectionsToExport;
+	if (!GetOutputNodeCollections(CollectionsToExport, NodesToExport, ErrorMessages))
+	{
+		return FReply::Handled();
+	}
+
+	FScopedSlowTask ProgressSlowTask(NodesToExport.Num(), LOCTEXT("ExportAsset_SlowTask", "Exporting Asset(s) ..."));
+	ProgressSlowTask.MakeDialog(true);
+
+	for (int32 i = 0; i < NodesToExport.Num(); ++i)
+	{
+		if (ProgressSlowTask.ShouldCancel())
+		{
+			break;
+		}
+
+		const UPCGNode* Node = NodesToExport[i];
+		const UPVOutputSettings* OutputSettings = CastChecked<const UPVOutputSettings>(Node->GetSettings());
+		check(OutputSettings != nullptr);
+
+		const FText AssetCountText = NodesToExport.Num() > 1
+			? FText::FromString(FString::Printf(TEXT(" [%d/%d]"), i + 1, NodesToExport.Num()))
+			: FText::FromString("");
+
+		ProgressSlowTask.EnterProgressFrame(1.f, FText::Format(LOCTEXT("ExportingAsset", "Exporting \"{0}\"{1}"), FText::FromName(OutputSettings->ExportSettings.MeshName), AssetCountText));
+
+		FScopedSlowTask ExportCollectionSlowTask(1.f, LOCTEXT("GeneratingMesh", "Generating Mesh ..."));
 		float PrevFrameWork = 0.f;
-		const auto OnStatusUpdated = [&](float ExpectedWorkThisFrame, const FText& StatusText)
+		float PrevProgress = 0.f;
+		const auto OnStatusUpdated = [&](const FText& Stage, float Progress)->bool
 		{
-			ExportCollectionSlowTask.EnterProgressFrame(PrevFrameWork, StatusText);
+			const float WorkThisFrame = Progress - PrevProgress;
+			PrevProgress = Progress;
+
+			ExportCollectionSlowTask.EnterProgressFrame(PrevFrameWork, Stage);
 
 			// Since our frames are so slow it does not make sense for the progress bar to progress as we start the frame, but rather as the frame finishes,
 			// hence why we repport the previours frame's work.
-			PrevFrameWork = ExpectedWorkThisFrame;
+			PrevFrameWork = WorkThisFrame;
+
+			return !ProgressSlowTask.ShouldCancel();
 		};
-		PV::Export::ExportCollectionAsMesh(ProceduralVegetationBeingEdited, OutCollection, OutCollection, OutputSettings->ExportSettings, OnStatusUpdated);
-		
-		ExportCollectionSlowTask.EnterProgressFrame(PrevFrameWork); // Let the slow task know all work has been completed.
-	}
 
-	if (ErrorMessages.Num() > 0)
-	{
-		FMessageLog EditorErrors("EditorErrors");
-		EditorErrors.AddMessages(ErrorMessages);
-		EditorErrors.Notify(LOCTEXT("PVE Mesh Export Failed", "PVE Mesh Export Failed"));
+		const bool bExportSucceeded = PV::Export::ExportCollectionAsMesh(
+			ProceduralVegetationBeingEdited, 
+			CollectionsToExport[i], 
+			OutputSettings->ExportSettings, 
+			OnStatusUpdated
+		);
+
+		if (!bExportSucceeded)
+		{
+			if (ProgressSlowTask.ShouldCancel())
+			{
+				FString MeshNames = "{ ";
+				for (int32 j = i; j < NodesToExport.Num(); j++)
+				{
+					MeshNames += CastChecked<const UPVOutputSettings>(NodesToExport[j]->GetSettings())->ExportSettings.MeshName.ToString();
+					if (j != NodesToExport.Num())
+					{
+						MeshNames += ", ";
+					}
+				}
+				MeshNames += " }";
+
+				ErrorMessages.AddError(
+					FText::Format(
+						LOCTEXT("ExportFailed_ExportCanceled", "Failed to export meshes \"{0}\", operation canceled by user."),
+						FText::FromString(MeshNames)
+					));
+			}
+			else
+			{
+				ErrorMessages.AddError(FText::Format(
+					LOCTEXT("ExportFailed_UnknownError", "Failed to export mesh \"{0}\" due to unkown error, export aborted."),
+					FText::FromName(OutputSettings->ExportSettings.MeshName)
+				));
+			}
+			
+			break;
+		}
+
+		ExportCollectionSlowTask.EnterProgressFrame(PrevFrameWork); // Let the slow task know all work has been completed.
 	}
 	
 	return FReply::Handled();
diff --git a/Engine/Plugins/Experimental/ProceduralVegetationEditor/Source/ProceduralVegetationEditor/Private/ProceduralVegetationEditorModule.cpp b/Engine/Plugins/Experimental/ProceduralVegetationEditor/Source/ProceduralVegetationEditor/Private/ProceduralVegetationEditorModule.cpp
--- a/Engine/Plugins/Experimental/ProceduralVegetationEditor/Source/ProceduralVegetationEditor/Private/ProceduralVegetationEditorModule.cpp
+++ b/Engine/Plugins/Experimental/ProceduralVegetationEditor/Source/ProceduralVegetationEditor/Private/ProceduralVegetationEditorModule.cpp
@@ -29,6 +29,7 @@
 #include "DataTypes/PVFoliageMeshData.h"
 #include "DataTypes/PVGrowthData.h"
 #include "DataTypes/PVMeshData.h"
+#include "PVExportParams.h"
 
 #include "Engine/SkeletalMesh.h"
 #include "Engine/StaticMesh.h"
diff --git a/Engine/Plugins/Experimental/ProceduralVegetationEditor/Source/ProceduralVegetationEditor/ProceduralVegetationEditor.Build.cs b/Engine/Plugins/Experimental/ProceduralVegetationEditor/Source/ProceduralVegetationEditor/ProceduralVegetationEditor.Build.cs
--- a/Engine/Plugins/Experimental/ProceduralVegetationEditor/Source/ProceduralVegetationEditor/ProceduralVegetationEditor.Build.cs
+++ b/Engine/Plugins/Experimental/ProceduralVegetationEditor/Source/ProceduralVegetationEditor/ProceduralVegetationEditor.Build.cs
@@ -8,13 +8,6 @@ public class ProceduralVegetationEditor : ModuleRules
 	{
 		PCHUsage = PCHUsageMode.UseExplicitOrSharedPCHs;
 
-		PublicDependencyModuleNames.AddRange(
-			new string[]
-			{
-				// ... add other public dependencies that you statically link with here ...
-			}
-		);
-
 		PrivateDependencyModuleNames.AddRange(
 			new string[]
 			{
@@ -32,22 +25,31 @@ public class ProceduralVegetationEditor : ModuleRules
 				"GeometryCollectionEngine",
 				"GeometryCore",
 				"GeometryFramework",
+				"GeometryScriptingCore",
 				"InputCore",
 				"Json",
 				"JsonUtilities",
 				"LevelEditor",
 				"MeshConversionEngineTypes",
+				"MeshDescription",
+				"MeshConversion",
+				"ModelingComponents",
 				"PCG",
 				"PCGEditor",
 				"ProceduralVegetation",
 				"Projects",
 				"PropertyEditor",
+				"PlanarCut",
 				"RenderCore",
 				"Slate",
 				"SlateCore",
+				"SkeletalMeshDescription",
+				"StaticMeshDescription",
 				"ToolMenus",
 				"ToolWidgets",
-				"UnrealEd"
+				"UnrealEd",
+				"DynamicWind",
+				"DynamicWindEditor"
 			}
 		);
 	}
