Address minor BugHawk issues

diff --git a/Engine/Plugins/Experimental/ProceduralVegetationEditor/Source/ProceduralVegetationEditor/Private/Helpers/PVExportHelper.cpp b/Engine/Plugins/Experimental/ProceduralVegetationEditor/Source/ProceduralVegetationEditor/Private/Helpers/PVExportHelper.cpp
--- a/Engine/Plugins/Experimental/ProceduralVegetationEditor/Source/ProceduralVegetationEditor/Private/Helpers/PVExportHelper.cpp
+++ b/Engine/Plugins/Experimental/ProceduralVegetationEditor/Source/ProceduralVegetationEditor/Private/Helpers/PVExportHelper.cpp
@@ -388,7 +388,7 @@ namespace PV::Export::Internal
 			const TObjectPtr<UStaticMesh>* StaticMesh = FoliageStaticMeshes.Find(FoliageName);
 			const TObjectPtr<USkeletalMesh>* SkeletalMesh = FoliageSkeletalMeshes.Find(FoliageName);
 
-			check((StaticMesh != nullptr || SkeletalMesh != nullptr) && (*StaticMesh != nullptr || *SkeletalMesh != nullptr));
+			check((StaticMesh && *StaticMesh) || (SkeletalMesh && *SkeletalMesh));
 
 			if (SkeletalMesh != nullptr)
 			{
@@ -1475,7 +1475,16 @@ namespace PV::Export
 		for (UPackage* Package : ModifiedPackages)
 		{
 			Package->MarkPackageDirty();
-			IAssetRegistry::Get()->AssetCreated(Package);
+
+			if (Package->HasAnyPackageFlags(PKG_NewlyCreated))
+			{
+				const FName AssetName = FPackageName::GetShortFName(Package->GetFName());
+				UObject* Asset = FindObjectFast<UObject>(Package, AssetName);
+				if (Asset)
+				{
+					IAssetRegistry::Get()->AssetCreated(Asset);
+				}
+			}
 		}
 
 		return true;
