[PVE] Added viewport and settings toolbar to display manniquen

diff --git a/Engine/Plugins/Experimental/ProceduralVegetationEditor/Source/ProceduralVegetationEditor/Private/PVEditor.cpp b/Engine/Plugins/Experimental/ProceduralVegetationEditor/Source/ProceduralVegetationEditor/Private/PVEditor.cpp
--- a/Engine/Plugins/Experimental/ProceduralVegetationEditor/Source/ProceduralVegetationEditor/Private/PVEditor.cpp
+++ b/Engine/Plugins/Experimental/ProceduralVegetationEditor/Source/ProceduralVegetationEditor/Private/PVEditor.cpp
@@ -10,6 +10,7 @@
 #include "PVEditorSchema.h"
 #include "PVEditorSettings.h"
 #include "Helpers/PVExportHelper.h"
+#include "SPVEditorViewport.h"
 #include "ToolMenus.h"
 #include "ToolMenuSection.h"
 #include "Misc/UObjectToken.h"
@@ -586,6 +587,11 @@ FReply FPVEditor::ExportInternal()
 	return FReply::Handled();
 }
 
+TSharedRef<SPCGEditorViewport> FPVEditor::CreateViewportWidget()
+{
+	return SNew(SPVEditorViewport);
+}
+
 bool FPVEditor::GetCollectionFromNode(const UPCGNode* InNode, FName& OutPinName, FManagedArrayCollection& OutCollection) const
 {
 	const UPCGPin* Pin = nullptr;
diff --git a/Engine/Plugins/Experimental/ProceduralVegetationEditor/Source/ProceduralVegetationEditor/Private/PVEditor.h b/Engine/Plugins/Experimental/ProceduralVegetationEditor/Source/ProceduralVegetationEditor/Private/PVEditor.h
--- a/Engine/Plugins/Experimental/ProceduralVegetationEditor/Source/ProceduralVegetationEditor/Private/PVEditor.h
+++ b/Engine/Plugins/Experimental/ProceduralVegetationEditor/Source/ProceduralVegetationEditor/Private/PVEditor.h
@@ -67,4 +67,7 @@ private:
 	bool GetCollectionFromNode(const UPCGNode* InNode, FName& OutPinName, FManagedArrayCollection& OutCollection) const;
 	void OnExport_Clicked();
 	FReply ExportInternal();
+
+protected:
+	virtual TSharedRef<SPCGEditorViewport> CreateViewportWidget() override;
 };
diff --git a/Engine/Plugins/Experimental/ProceduralVegetationEditor/Source/ProceduralVegetationEditor/Private/SPVEditorViewport.cpp b/Engine/Plugins/Experimental/ProceduralVegetationEditor/Source/ProceduralVegetationEditor/Private/SPVEditorViewport.cpp
new file mode 100644
--- /dev/null
+++ b/Engine/Plugins/Experimental/ProceduralVegetationEditor/Source/ProceduralVegetationEditor/Private/SPVEditorViewport.cpp
@@ -0,0 +1,153 @@
+// Copyright Epic Games, Inc. All Rights Reserved.
+
+#include "SPVEditorViewport.h"
+#include "AdvancedPreviewScene.h"
+#include "AdvancedPreviewSceneMenus.h"
+#include "PVEditorCommands.h"
+#include "PVEditorSettings.h"
+#include "ViewportToolbar/UnrealEdViewportToolbar.h"
+#include "ToolMenu.h"
+#include "ToolMenus.h"
+#include "ToolMenuSection.h"
+#include "Engine/StaticMesh.h"
+#include "Components/StaticMeshComponent.h"
+
+#define LOCTEXT_NAMESPACE "PVEditorViewportToolbarSections"
+
+FToolMenuEntry SPVEditorViewport::CreateSettingsToolbarMenu()
+{
+	return FToolMenuEntry::InitDynamicEntry(TEXT("PVE Settings"), FNewToolMenuSectionDelegate::CreateLambda([](FToolMenuSection& Section)
+		{
+			FToolMenuEntry& Entry = Section.AddSubMenu(
+				/*Name=*/TEXT("PVESettingsSubmenu"),
+				FText::GetEmpty(),
+				FText::GetEmpty(),
+				FNewToolMenuDelegate::CreateLambda([](UToolMenu* Submenu)
+				{
+					// Add preference toggles to this section.
+					FToolMenuSection& PreferencesSection = Submenu->FindOrAddSection(TEXT("PVEMannequinSettings"), LOCTEXT("PVEMannequinSettingsSectionLabel", "Mannequin Settings"));
+					PreferencesSection.AddMenuEntry(FPVEditorCommands::Get().ShowMannequin);
+				}),
+				/*bInOpenSubMenuOnClick=*/false,
+				FSlateIcon(FAppStyle::GetAppStyleSetName(), "LevelEditor.GameSettings"));
+		}));
+}
+
+TSharedPtr<SWidget> SPVEditorViewport::BuildViewportToolbar()
+{
+	check(AdvancedPreviewScene.IsValid());
+
+	const FName ToolbarName = "PVE.ViewportToolbar";
+	
+	if (!UToolMenus::Get()->IsMenuRegistered(ToolbarName))
+	{
+		UToolMenu* Menu = UToolMenus::Get()->RegisterMenu(ToolbarName, /*Parent=*/NAME_None, EMultiBoxType::SlimHorizontalToolBar);
+		Menu->StyleName = "ViewportToolbar";
+		Menu->AddSection("Left");
+		
+		FToolMenuSection& RightSection = Menu->AddSection("Right");
+		RightSection.Alignment = EToolMenuSectionAlign::Last;
+		RightSection.AddEntry(UE::UnrealEd::CreateCameraSubmenu(UE::UnrealEd::FViewportCameraMenuOptions().ShowAll()));
+		RightSection.AddEntry(UE::UnrealEd::CreateViewModesSubmenu());
+		RightSection.AddEntry(UE::UnrealEd::CreateDefaultShowSubmenu());
+		RightSection.AddEntry(UE::UnrealEd::CreateAssetViewerProfileSubmenu());
+		UE::AdvancedPreviewScene::Menus::ExtendAdvancedPreviewSceneSettings(
+			"PVE.ViewportToolbar.AssetViewerProfile",
+			UE::AdvancedPreviewScene::Menus::FSettingsOptions().ShowToggleGrid(false)
+		);
+
+		RightSection.AddEntry(CreateSettingsToolbarMenu());
+	}
+
+	FToolMenuContext Context;
+	{
+		Context.AppendCommandList(AdvancedPreviewScene->GetCommandList());
+		Context.AppendCommandList(GetCommandList());
+		Context.AddExtender(GetExtenders());
+		Context.AddObject(UE::UnrealEd::CreateViewportToolbarDefaultContext(GetViewportWidget()));
+	}
+	
+	return UToolMenus::Get()->GenerateWidget(ToolbarName, Context);
+}
+
+void SPVEditorViewport::BindCommands()
+{
+	SPCGEditorViewport::BindCommands();
+
+	const FPVEditorCommands& Commands = FPVEditorCommands::Get();
+
+	CommandList->MapAction(
+		Commands.ShowMannequin,
+		FExecuteAction::CreateSP(this, &SPVEditorViewport::ToggleShowManniquen),
+		FCanExecuteAction(),
+		FIsActionChecked::CreateSP(this, &SPVEditorViewport::IsShowManniquenChecked));
+}
+
+void SPVEditorViewport::SetMannequinState(bool InEnable)
+{
+	UStaticMesh* MannequinMesh = LoadObject<UStaticMesh>(
+		nullptr,
+		TEXT("/ProceduralVegetationEditor/Mannequin/Viewport_Mannequin_T_Pose.Viewport_Mannequin_T_Pose")
+	);
+
+	if (!MannequinMesh)
+	{
+		return;
+	}
+	
+	if (InEnable)
+	{
+		UStaticMeshComponent* MannequinComponent = NewObject<UStaticMeshComponent>(GetTransientPackage(), NAME_None, RF_Transient);
+		MannequinComponent->SetStaticMesh(MannequinMesh);
+		FBoxSphereBounds MannequinBounds = MannequinComponent->CalcLocalBounds();
+
+		ManagedResources.Add(MannequinComponent);
+		AdvancedPreviewScene->AddComponent(
+			MannequinComponent,
+			FTransform(FVector((-MannequinBounds.BoxExtent.X - MannequinBounds.BoxExtent.X) * 1.5f, 0, 0))
+		);
+	}
+	else
+	{
+		for (TObjectPtr<UObject> ManagedResource : ManagedResources)
+		{
+			if (UStaticMeshComponent* ActorComponent = Cast<UStaticMeshComponent>(ManagedResource))
+			{
+				if (ActorComponent->GetStaticMesh()->GetFName() == MannequinMesh->GetFName())
+				{
+					AdvancedPreviewScene->RemoveComponent(ActorComponent);
+				}
+			}
+		}
+	}
+
+	Client->Invalidate();
+}
+
+void SPVEditorViewport::ToggleShowManniquen()
+{
+	UPVEditorSettings* EditorSettings = GetMutableDefault<UPVEditorSettings>();
+	EditorSettings->bShowMannequin = !EditorSettings->bShowMannequin;
+
+	SetMannequinState(EditorSettings->bShowMannequin);
+}
+
+bool SPVEditorViewport::IsShowManniquenChecked() const
+{
+	return GetDefault<UPVEditorSettings>()->bShowMannequin;
+}
+
+void SPVEditorViewport::OnSetupScene()
+{
+	SPCGEditorViewport::OnSetupScene();
+
+	SetMannequinState(IsShowManniquenChecked());
+}
+
+void SPVEditorViewport::OnResetScene()
+{
+	SPCGEditorViewport::OnResetScene();
+
+	SetMannequinState(IsShowManniquenChecked());
+}
+#undef LOCTEXT_NAMESPACE
diff --git a/Engine/Plugins/Experimental/ProceduralVegetationEditor/Source/ProceduralVegetationEditor/Private/SPVEditorViewport.h b/Engine/Plugins/Experimental/ProceduralVegetationEditor/Source/ProceduralVegetationEditor/Private/SPVEditorViewport.h
new file mode 100644
--- /dev/null
+++ b/Engine/Plugins/Experimental/ProceduralVegetationEditor/Source/ProceduralVegetationEditor/Private/SPVEditorViewport.h
@@ -0,0 +1,28 @@
+// Copyright Epic Games, Inc. All Rights Reserved.
+
+#pragma once
+
+#include "PVEditorSettings.h"
+#include "Widgets/AssetEditorViewport/SPCGEditorViewport.h"
+
+struct FToolMenuEntry;
+class UPVEditorSettings;
+
+class SPVEditorViewport : public SPCGEditorViewport
+{
+public:
+	FToolMenuEntry CreateSettingsToolbarMenu();
+
+protected:
+	virtual TSharedPtr<SWidget> BuildViewportToolbar() override;
+	virtual void BindCommands() override;
+
+private:
+	void ToggleShowManniquen();
+	bool IsShowManniquenChecked() const;
+	void SetMannequinState(bool InEnable);
+
+protected:
+	virtual void OnSetupScene() override;
+	virtual void OnResetScene() override;
+};
diff --git a/Engine/Plugins/Experimental/ProceduralVegetationEditor/Source/ProceduralVegetationEditor/Private/Visualizations/PVDataVisualization.cpp b/Engine/Plugins/Experimental/ProceduralVegetationEditor/Source/ProceduralVegetationEditor/Private/Visualizations/PVDataVisualization.cpp
--- a/Engine/Plugins/Experimental/ProceduralVegetationEditor/Source/ProceduralVegetationEditor/Private/Visualizations/PVDataVisualization.cpp
+++ b/Engine/Plugins/Experimental/ProceduralVegetationEditor/Source/ProceduralVegetationEditor/Private/Visualizations/PVDataVisualization.cpp
@@ -107,23 +107,6 @@ FPCGSetupSceneFunc FPVDataVisualization::GetViewportSetupFunc(const UPCGSettings
 			}
 		}
 
-		if (GetDefault<UPVEditorSettings>()->bShowMannequin)
-		{
-			UStaticMesh* MannequinMesh = LoadObject<UStaticMesh>(
-				nullptr,
-				TEXT("/ProceduralVegetationEditor/Mannequin/Viewport_Mannequin_T_Pose.Viewport_Mannequin_T_Pose")
-			);
-			UStaticMeshComponent* MannequinComponent = NewObject<UStaticMeshComponent>(GetTransientPackage(), NAME_None, RF_Transient);
-			MannequinComponent->SetStaticMesh(MannequinMesh);
-			FBoxSphereBounds MannequinBounds = MannequinComponent->CalcLocalBounds();
-
-			InOutParams.ManagedResources.Add(MannequinComponent);
-			InOutParams.Scene->AddComponent(
-				MannequinComponent,
-				FTransform(FVector((-Bounds.BoxExtent.X - MannequinBounds.BoxExtent.X) * 1.5f, 0, 0))
-			);
-		}
-
 		if (PV::Utilities::DebugModeEnabled())
 		{
 			if (WeakSettings.IsValid() && WeakSettings->bDebug)
