[PVE] Foliage palette grid visualization

#virtualized

diff --git a/Engine/Plugins/Experimental/ProceduralVegetationEditor/Source/ProceduralVegetation/Private/Nodes/PVFoliagePaletteSettings.h b/Engine/Plugins/Experimental/ProceduralVegetationEditor/Source/ProceduralVegetation/Private/Nodes/PVFoliagePaletteSettings.h
--- a/Engine/Plugins/Experimental/ProceduralVegetationEditor/Source/ProceduralVegetation/Private/Nodes/PVFoliagePaletteSettings.h
+++ b/Engine/Plugins/Experimental/ProceduralVegetationEditor/Source/ProceduralVegetation/Private/Nodes/PVFoliagePaletteSettings.h
@@ -18,7 +18,7 @@ public:
 	virtual FText GetNodeTooltipText() const override;
 	virtual FLinearColor GetNodeTitleColor() const override { return FLinearColor::Green; }
 
-	virtual TArray<EPVRenderType> GetRenderType() const override { return TArray{EPVRenderType::Mesh, EPVRenderType::Foliage}; }
+	virtual TArray<EPVRenderType> GetRenderType() const override { return TArray{EPVRenderType::FoliageGrid}; }
 
 	virtual void PostEditChangeProperty(FPropertyChangedEvent& PropertyChangedEvent) override;
 #endif
diff --git a/Engine/Plugins/Experimental/ProceduralVegetationEditor/Source/ProceduralVegetation/Public/DataTypes/PVData.h b/Engine/Plugins/Experimental/ProceduralVegetationEditor/Source/ProceduralVegetation/Public/DataTypes/PVData.h
--- a/Engine/Plugins/Experimental/ProceduralVegetationEditor/Source/ProceduralVegetation/Public/DataTypes/PVData.h
+++ b/Engine/Plugins/Experimental/ProceduralVegetationEditor/Source/ProceduralVegetation/Public/DataTypes/PVData.h
@@ -17,6 +17,7 @@ enum class EPVRenderType : uint8
 	Mesh 		= 1 << 1,
 	Foliage 	= 1 << 2,
 	Bones 		= 1 << 3,
+	FoliageGrid = 1 << 4,
 
 	// Add new renter type above
 	Count	
diff --git a/Engine/Plugins/Experimental/ProceduralVegetationEditor/Source/ProceduralVegetationEditor/Private/Visualizations/PVDataVisualization.cpp b/Engine/Plugins/Experimental/ProceduralVegetationEditor/Source/ProceduralVegetationEditor/Private/Visualizations/PVDataVisualization.cpp
--- a/Engine/Plugins/Experimental/ProceduralVegetationEditor/Source/ProceduralVegetationEditor/Private/Visualizations/PVDataVisualization.cpp
+++ b/Engine/Plugins/Experimental/ProceduralVegetationEditor/Source/ProceduralVegetationEditor/Private/Visualizations/PVDataVisualization.cpp
@@ -11,25 +11,22 @@
 #include "PCGEditorSettings.h"
 #include "PlanarCut.h"
 #include "PVBoneComponent.h"
-#include "PVEditorSettings.h"
-
 #include "Components/DynamicMeshComponent.h"
 #include "Components/LineBatchComponent.h"
 #include "Components/TextRenderComponent.h"
-
 #include "DataTypes/PVData.h"
-
 #include "DebugVisualization/PVDebugVisualizer.h"
-
 #include "Facades/PVFoliageFacade.h"
-
 #include "GeometryCollection/GeometryCollectionObject.h"
-
 #include "Helpers/PVUtilities.h"
-
 #include "Nodes/PVBaseSettings.h"
-
+#include "Rendering/SkeletalMeshRenderData.h"
 #include "Visualizations/PVSkeletonVisualizerComponent.h"
+#include "Engine/StaticMesh.h"
+#include "Engine/SkeletalMesh.h"
+#include "Components/StaticMeshComponent.h"
+#include "Components/SkeletalMeshComponent.h"
+#include "Materials/Material.h"
 
 #define LOCTEXT_NAMESPACE "PCGManagedArrayCollectionDataVisualization"
 
@@ -48,6 +45,7 @@ void FPVDataVisualization::RegisterVisualizations()
 	RenderMap.Add(EPVRenderType::Mesh, FPVRenderCallback::CreateRaw(this, &FPVDataVisualization::MeshRenderer));
 	RenderMap.Add(EPVRenderType::Foliage, FPVRenderCallback::CreateRaw(this, &FPVDataVisualization::FoliageRenderer));
 	RenderMap.Add(EPVRenderType::Bones, FPVRenderCallback::CreateRaw(this, &FPVDataVisualization::BonesRenderer));
+	RenderMap.Add(EPVRenderType::FoliageGrid, FPVRenderCallback::CreateRaw(this, &FPVDataVisualization::FoliageGridRenderer));
 }
 
 void FPVDataVisualization::ExecuteDebugDisplay(FPCGContext* Context, const UPCGSettingsInterface* SettingsInterface, const UPCGData* Data,
@@ -226,4 +224,183 @@ void FPVDataVisualization::BonesRenderer(const FManagedArrayCollection& InCollec
 	OutBounds = OutBounds + BoneComponent->CalcLocalBounds();
 }
 
+int32 GetMeshTriangles(const FString FoliagePath)
+{
+	UStaticMesh* StaticFoliageMesh = LoadObject<UStaticMesh>(nullptr, *FoliagePath);
+	USkeletalMesh* SkeletalFoliageMesh = LoadObject<USkeletalMesh>(nullptr, *FoliagePath);
+
+	if (StaticFoliageMesh)
+	{
+		return StaticFoliageMesh->GetNumTriangles(0);
+	}
+	else
+	{
+		int32 NumTriangles = 0;
+		FSkeletalMeshRenderData* SkelMeshRenderData = SkeletalFoliageMesh->GetResourceForRendering();
+		if (SkelMeshRenderData && SkelMeshRenderData->LODRenderData.Num() > 0)
+		{
+			const FSkeletalMeshLODRenderData& LODData = SkelMeshRenderData->LODRenderData[0];
+			NumTriangles = LODData.GetTotalFaces();
+		}
+
+		return NumTriangles;
+	}
+	
+}
+
+void FPVDataVisualization::FoliageGridRenderer(const FManagedArrayCollection& InCollection, FPCGSceneSetupParams& InOutParams,
+	FBoxSphereBounds& OutBounds)
+{
+	PV::Facades::FFoliageFacade FoliageFacade(InCollection);
+
+	const int FoliageNum = FoliageFacade.NumFoliageNames();
+
+	FBox MaxFoliageBounds;
+	float MaxFoliageVolume = std::numeric_limits<double>::lowest();
+	int ValidFoliageEntries = 0;
+	
+	for (int32 FoliageNameIndex = 0; FoliageNameIndex < FoliageNum; ++FoliageNameIndex)
+	{
+		FString FoliageName = FoliageFacade.GetFoliageName(FoliageNameIndex);
+
+		if (UStaticMesh* StaticFoliageMesh = LoadObject<UStaticMesh>(nullptr, *FoliageName))
+		{
+			FBox BoundingBox = StaticFoliageMesh->GetBoundingBox();
+
+			if (BoundingBox.IsValid)
+			{
+				if (BoundingBox.GetVolume() > MaxFoliageVolume)
+				{
+					MaxFoliageVolume = BoundingBox.GetVolume();
+					MaxFoliageBounds = BoundingBox;
+				}
+			}
+
+			ValidFoliageEntries++;
+		}
+		else if (USkeletalMesh* SkeletalFoliageMesh = LoadObject<USkeletalMesh>(nullptr, *FoliageName))
+		{
+			FBox BoundingBox = SkeletalFoliageMesh->GetBounds().GetBox();
+
+			if (BoundingBox.IsValid)
+			{
+				if (BoundingBox.GetVolume() > MaxFoliageVolume)
+				{
+					MaxFoliageVolume = BoundingBox.GetVolume();
+					MaxFoliageBounds = BoundingBox;
+				}
+			}
+
+			ValidFoliageEntries++;
+		}
+	}
+
+	FName FoliageName = NAME_None;
+	const FVector FoliageEntryBounds = MaxFoliageBounds.GetSize();
+	const float FoliageHorizontalPadding = 100;
+	const float FoliageVericalPadding = 50;
+	const float TextItemPadding = 1.1f;
+	const float BackgroundHorizontalPadding = 0.1f;
+	const float BackgroundVerticalPadding = 0.05f;
+	
+	FVector FoliageHorizontalSize = (FoliageEntryBounds + FoliageHorizontalPadding) * FVector::RightVector ;
+	FVector FoliageVerticalSize = (FoliageEntryBounds + FoliageVericalPadding) * FVector::UpVector;
+	FVector TotalHorizontalSize = ValidFoliageEntries/2.0f * FoliageHorizontalSize;
+
+	FVector FoliagePosOffset = FVector::ZeroVector - TotalHorizontalSize;
+	FVector FoliagePos = FoliagePosOffset;
+
+	UStaticMesh* TextBackgroundMesh = LoadObject<UStaticMesh>(nullptr,TEXT("/Engine/BasicShapes/Plane"));
+	UMaterial* BackgroundMat = LoadObject<UMaterial>(nullptr, TEXT("/ProceduralVegetationEditor/Materials/TextBackgroundMat.TextBackgroundMat"));
+
+	check(BackgroundMat && TextBackgroundMesh);
+	
+	for (int32 FoliageNameIndex = 0; FoliageNameIndex < FoliageNum; ++FoliageNameIndex)
+	{
+		FString FoliagePath = FoliageFacade.GetFoliageName(FoliageNameIndex);
+
+		UStaticMesh* StaticFoliageMesh = LoadObject<UStaticMesh>(nullptr, *FoliagePath);
+		USkeletalMesh* SkeletalFoliageMesh = LoadObject<USkeletalMesh>(nullptr, *FoliagePath);
+
+		const bool bHasValidMesh = StaticFoliageMesh || SkeletalFoliageMesh; 
+		
+		if (bHasValidMesh)
+		{
+			FoliagePos += FoliageHorizontalSize;
+			FoliageName = StaticFoliageMesh != nullptr ? StaticFoliageMesh->GetFName() : SkeletalFoliageMesh->GetFName();
+			
+			FText TriCount = FText::AsNumber(GetMeshTriangles(FoliagePath));
+			FString InTextToDraw = FString::Format(TEXT("Name: {0}\nTriangles:{1}"),{FoliageName.ToString(), TriCount.ToString()});
+
+			TObjectPtr<UTextRenderComponent> TextRenderComponent = NewObject<UTextRenderComponent>(GetTransientPackage(), NAME_None, RF_Transient);
+			TextRenderComponent->SetText(FText::FromString(InTextToDraw));
+			TextRenderComponent->SetTextRenderColor(FColor::Black);
+			TextRenderComponent->SetWorldSize(10);
+			TextRenderComponent->SetGenerateOverlapEvents(false);
+
+			FVector TextPos = FoliagePos;
+			TextPos += FoliageVerticalSize * TextItemPadding;
+			TextPos += FoliageHorizontalSize/2.0f - FoliageHorizontalSize/4.0f;
+
+			FTransform TextTransform = FTransform(FVector(TextPos));
+			
+			InOutParams.ManagedResources.Add(TextRenderComponent);
+			InOutParams.Scene->AddComponent(TextRenderComponent, TextTransform);
+
+			OutBounds = OutBounds + TextRenderComponent->CalcBounds(TextTransform);
+
+			if (TextBackgroundMesh && BackgroundMat)
+			{
+				UStaticMeshComponent* TextBackgroundMeshComponent = NewObject<UStaticMeshComponent>(GetTransientPackage(), NAME_None, RF_Transient);
+				TextBackgroundMeshComponent->SetStaticMesh(TextBackgroundMesh);
+				TextBackgroundMeshComponent->SetMaterial(0, BackgroundMat);
+
+				auto TextBackgroundBounds = TextBackgroundMeshComponent->CalcBounds(FTransform::Identity);
+
+				float TextBackgroundScale = FoliageHorizontalSize.Y/TextBackgroundBounds.BoxExtent.Y;
+				TextBackgroundScale /= 2;
+
+				FVector BackgroundPos = TextPos + FVector::BackwardVector * 0.1;
+				BackgroundPos -= FoliageHorizontalSize/2.0f;
+				BackgroundPos += FoliageHorizontalSize * BackgroundHorizontalPadding;
+				BackgroundPos += FoliageVerticalSize * BackgroundVerticalPadding;
+			
+				FTransform TextBackgroundTransform = FTransform(BackgroundPos);
+				FRotator TextRotator = FRotator(270,0,0);
+				TextBackgroundTransform.SetRotation(TextRotator.Quaternion());
+				TextBackgroundTransform.SetScale3D(FVector(0.4, TextBackgroundScale, 1));
+			
+				InOutParams.ManagedResources.Add(TextBackgroundMeshComponent);
+				InOutParams.Scene->AddComponent(
+					TextBackgroundMeshComponent,
+					FTransform(TextBackgroundTransform)
+				);
+			}
+		}
+		
+		if(StaticFoliageMesh)
+		{
+			TObjectPtr<UStaticMeshComponent> MeshComponent = NewObject<UStaticMeshComponent>(GetTransientPackage(), NAME_None, RF_Transient);
+			MeshComponent->SetStaticMesh(StaticFoliageMesh);
+			MeshComponent->UpdateBounds();
+
+			InOutParams.ManagedResources.Add(MeshComponent);
+			InOutParams.Scene->AddComponent(MeshComponent, FTransform(FoliagePos));
+
+			OutBounds = OutBounds + MeshComponent->CalcBounds(FTransform(FoliagePos));
+		}
+		else if(SkeletalFoliageMesh)
+		{	
+			TObjectPtr<USkeletalMeshComponent> MeshComponent = NewObject<USkeletalMeshComponent>(GetTransientPackage(), NAME_None, RF_Transient);
+			MeshComponent->SetSkeletalMesh(SkeletalFoliageMesh);
+			MeshComponent->UpdateBounds();
+
+			InOutParams.ManagedResources.Add(MeshComponent);
+			InOutParams.Scene->AddComponent(MeshComponent, FTransform(FoliagePos));
+
+			OutBounds = OutBounds + MeshComponent->CalcBounds(FTransform(FoliagePos));
+		}
+	}
+}
+
 #undef LOCTEXT_NAMESPACE
diff --git a/Engine/Plugins/Experimental/ProceduralVegetationEditor/Source/ProceduralVegetationEditor/Private/Visualizations/PVDataVisualization.h b/Engine/Plugins/Experimental/ProceduralVegetationEditor/Source/ProceduralVegetationEditor/Private/Visualizations/PVDataVisualization.h
--- a/Engine/Plugins/Experimental/ProceduralVegetationEditor/Source/ProceduralVegetationEditor/Private/Visualizations/PVDataVisualization.h
+++ b/Engine/Plugins/Experimental/ProceduralVegetationEditor/Source/ProceduralVegetationEditor/Private/Visualizations/PVDataVisualization.h
@@ -31,6 +31,7 @@ protected:
 	void SkeletonRenderer(const FManagedArrayCollection& InCollection, FPCGSceneSetupParams& InOutParams, FBoxSphereBounds& OutBounds);
 	void FoliageRenderer(const FManagedArrayCollection& InCollection, FPCGSceneSetupParams& InOutParams, FBoxSphereBounds& OutBounds);
 	void BonesRenderer(const FManagedArrayCollection& InCollection, FPCGSceneSetupParams& InOutParams, FBoxSphereBounds& OutBounds);
+	void FoliageGridRenderer(const FManagedArrayCollection& InCollection, FPCGSceneSetupParams& InOutParams, FBoxSphereBounds& OutBounds);
 };
 
 DECLARE_LOG_CATEGORY_EXTERN(LogProceduralVegetationDataVisualization, Log, All);
\ No newline at end of file
