[PVE] Removed attributes view

diff --git a/Engine/Plugins/Experimental/ProceduralVegetationEditor/Source/ProceduralVegetationEditor/Private/PVEditor.cpp b/Engine/Plugins/Experimental/ProceduralVegetationEditor/Source/ProceduralVegetationEditor/Private/PVEditor.cpp
--- a/Engine/Plugins/Experimental/ProceduralVegetationEditor/Source/ProceduralVegetationEditor/Private/PVEditor.cpp
+++ b/Engine/Plugins/Experimental/ProceduralVegetationEditor/Source/ProceduralVegetationEditor/Private/PVEditor.cpp
@@ -42,6 +42,10 @@ FPVEditor::FPVEditor()
 
 void FPVEditor::Initialize(const EToolkitMode::Type InMode, const TSharedPtr<class IToolkitHost>& InToolkitHost, UProceduralVegetation* InProceduralVegetation)
 {
+	PRAGMA_DISABLE_DEPRECATION_WARNINGS
+	bForceRefreshAttributeEvenIfClosed = true;
+	PRAGMA_ENABLE_DEPRECATION_WARNINGS
+	
 	check(InProceduralVegetation);
 
 	ProceduralVegetationBeingEdited = InProceduralVegetation;
@@ -200,7 +204,7 @@ TAttribute<FGraphAppearanceInfo> FPVEditor::GetAppearanceInfo() const
 
 TSharedRef<FTabManager::FLayout> FPVEditor::GetDefaultLayout() const
 {
-	return FTabManager::NewLayout("Standalone_PVGraphEditor_DefaultLayout_v1.0.4")
+	return FTabManager::NewLayout("Standalone_PVGraphEditor_DefaultLayout_v1.0.5")
 	->AddArea // Main PCG Graph Editor Area
 	(
 		FTabManager::NewPrimaryArea()->SetOrientation(Orient_Vertical)
@@ -220,13 +224,6 @@ TSharedRef<FTabManager::FLayout> FPVEditor::GetDefaultLayout() const
 					->SetForegroundTab(GetPanelID(EPCGEditorPanel::Viewport1))
 					->SetHideTabWell(true)
 				)
-				->Split // Bottom Section - ALV
-				(
-					FTabManager::NewStack()
-					->SetSizeCoefficient(0.05f)
-					->SetHideTabWell(false)
-					->AddTab(GetPanelID(EPCGEditorPanel::Attributes1), ETabState::ClosedTab)
-				)
 			)
 			->Split // Graph and details panel
 			(
@@ -259,8 +256,7 @@ bool FPVEditor::IsPanelAvailable(const FName PanelID) const
 		GetPanelID(EPCGEditorPanel::Log),
 		GetPanelID(EPCGEditorPanel::PropertyDetails1),
 		GetPanelID(EPCGEditorPanel::UserParams),
-		GetPanelID(EPCGEditorPanel::Viewport1),
-		GetPanelID(EPCGEditorPanel::Attributes1)
+		GetPanelID(EPCGEditorPanel::Viewport1)
 	};
 	return SupportedIDs.Contains(PanelID);
 }
