Fix nanite foliage export breaking due to PVExportHelpers refactor

diff --git a/Engine/Plugins/Experimental/ProceduralVegetationEditor/Source/ProceduralVegetationEditor/Private/Helpers/PVExportHelper.cpp b/Engine/Plugins/Experimental/ProceduralVegetationEditor/Source/ProceduralVegetationEditor/Private/Helpers/PVExportHelper.cpp
--- a/Engine/Plugins/Experimental/ProceduralVegetationEditor/Source/ProceduralVegetationEditor/Private/Helpers/PVExportHelper.cpp
+++ b/Engine/Plugins/Experimental/ProceduralVegetationEditor/Source/ProceduralVegetationEditor/Private/Helpers/PVExportHelper.cpp
@@ -523,7 +523,7 @@ namespace PV::Export::Internal
 
 		const auto AddStaticMeshToAssembly = [&](const FString& FoliageInstanceObjectPath, TObjectPtr<UStaticMesh> InStaticMesh)
 		{
-			const int32 PartIndex = AssemblyBuilder.FindOrAddPart(FoliageInstanceObjectPath);
+			const int32 PartIndex = AssemblyBuilder.FindOrAddPart(InStaticMesh.GetPath());
 			MeshNamePartMap.FindOrAdd(FoliageInstanceObjectPath) = PartIndex;
 
 			TArray<TObjectPtr<UMaterialInterface>>& MeshMaterials = MeshMaterialsMap.FindOrAdd(FoliageInstanceObjectPath);
@@ -622,7 +622,7 @@ namespace PV::Export::Internal
 
 		const auto AddSkeletalMeshToAssembly = [&](const FString& FoliageInstanceObjectPath, TObjectPtr<USkeletalMesh> InSkeletalMesh)
 		{
-			const int32 PartIndex = AssemblyBuilder.FindOrAddPart(FoliageInstanceObjectPath);
+			const int32 PartIndex = AssemblyBuilder.FindOrAddPart(InSkeletalMesh.GetPath());
 			MeshNamePartMap.FindOrAdd(FoliageInstanceObjectPath) = PartIndex;
 
 			TArray<TObjectPtr<UMaterialInterface>>& MeshMaterials = MeshMaterialsMap.FindOrAdd(FoliageInstanceObjectPath);
