[PVE] Default asset name now references template name, if selected

diff --git a/Engine/Plugins/Experimental/ProceduralVegetationEditor/Source/ProceduralVegetationEditor/Private/ProceduralVegetationFactory.cpp b/Engine/Plugins/Experimental/ProceduralVegetationEditor/Source/ProceduralVegetationEditor/Private/ProceduralVegetationFactory.cpp
--- a/Engine/Plugins/Experimental/ProceduralVegetationEditor/Source/ProceduralVegetationEditor/Private/ProceduralVegetationFactory.cpp
+++ b/Engine/Plugins/Experimental/ProceduralVegetationEditor/Source/ProceduralVegetationEditor/Private/ProceduralVegetationFactory.cpp
@@ -1,9 +1,11 @@
 // Copyright Epic Games, Inc. All Rights Reserved.
 
 #include "ProceduralVegetationFactory.h"
-#include "Widgets/SPVCreateDialog.h"
+
 #include "ProceduralVegetation.h"
 
+#include "Widgets/SPVCreateDialog.h"
+
 #define LOCTEXT_NAMESPACE "ProceduralVegetationFactory"
 
 UProceduralVegetationFactory::UProceduralVegetationFactory(const FObjectInitializer& ObjectInitializer)
@@ -20,11 +22,10 @@ UObject* UProceduralVegetationFactory::FactoryCreateNew(UClass* InClass, UObject
 	UProceduralVegetation* NewProceduralVegetation;
 	if (SampleProceduralVegetation)
 	{
-		NewProceduralVegetation = NewObject<UProceduralVegetation>(InParent, InClass, InName, InFlags, SampleProceduralVegetation);
+		NewProceduralVegetation = NewObject<UProceduralVegetation>(InParent, InClass, SampleProceduralVegetation->GetFName(), InFlags, SampleProceduralVegetation);
 
 		TObjectPtr<UProceduralVegetationGraph> SampleGraph = SampleProceduralVegetation->GetGraph();
 		NewProceduralVegetation->CreateGraph(SampleGraph);
-		
 	}
 	else
 	{
@@ -43,7 +44,22 @@ bool UProceduralVegetationFactory::ConfigureProperties()
 		LOCTEXT("ProceduralVegetationFactoryConfigureProperties", "Create Procedural Vegetation"),
 		SampleProceduralVegetation
 	);
+	
+	if (SampleProceduralVegetation)
+	{
+		SampleName = SampleProceduralVegetation->GetName();
+	}
+	
 	return bPressedOk;
 }
 
+FString UProceduralVegetationFactory::GetDefaultNewAssetName() const
+{
+	if (SampleName.IsSet())
+	{
+		return SampleName.GetValue();
+	}
+	return Super::GetDefaultNewAssetName();
+}
+
 #undef LOCTEXT_NAMESPACE
diff --git a/Engine/Plugins/Experimental/ProceduralVegetationEditor/Source/ProceduralVegetationEditor/Private/ProceduralVegetationFactory.h b/Engine/Plugins/Experimental/ProceduralVegetationEditor/Source/ProceduralVegetationEditor/Private/ProceduralVegetationFactory.h
--- a/Engine/Plugins/Experimental/ProceduralVegetationEditor/Source/ProceduralVegetationEditor/Private/ProceduralVegetationFactory.h
+++ b/Engine/Plugins/Experimental/ProceduralVegetationEditor/Source/ProceduralVegetationEditor/Private/ProceduralVegetationFactory.h
@@ -3,6 +3,7 @@
 #pragma once
 
 #include "Factories/Factory.h"
+
 #include "ProceduralVegetationFactory.generated.h"
 
 class UProceduralVegetation;
@@ -19,8 +20,12 @@ public:
 	virtual UObject* FactoryCreateNew(UClass* InClass, UObject* InParent, FName InName, EObjectFlags InFlags, UObject* InContext, FFeedbackContext* Warn) override;
 	virtual bool ShouldShowInNewMenu() const override { return true; }
 	virtual bool ConfigureProperties() override;
+	virtual FString GetDefaultNewAssetName() const override;
 
 private:
 	UPROPERTY()
 	TObjectPtr<UProceduralVegetation> SampleProceduralVegetation = nullptr;
+
+	UPROPERTY()
+	TOptional<FString> SampleName;
 };
\ No newline at end of file
