[PVE] Multiple fixes/optimizations
- Optimize Skeletal Mesh export
- Fix overwriting assets on export keeps older settings
- Fix a crash with Single mesh export (Nanite Foliage disbaled)
- Fix modifying ramp or setting ramp preset doesn't invalidate PCG graph
- Add Scale Visualization
- Add mannequin enable/disable, mannequin offset and scale vis. enable/disable to viewport toolbar
- Add warning on overwriting assets on export
- Template assets now reference original asest's name

diff --git a/Engine/Plugins/Experimental/ProceduralVegetationEditor/Source/ProceduralVegetationEditor/Private/Customizations/PVFloatRampCustomization.cpp b/Engine/Plugins/Experimental/ProceduralVegetationEditor/Source/ProceduralVegetationEditor/Private/Customizations/PVFloatRampCustomization.cpp
--- a/Engine/Plugins/Experimental/ProceduralVegetationEditor/Source/ProceduralVegetationEditor/Private/Customizations/PVFloatRampCustomization.cpp
+++ b/Engine/Plugins/Experimental/ProceduralVegetationEditor/Source/ProceduralVegetationEditor/Private/Customizations/PVFloatRampCustomization.cpp
@@ -106,6 +106,16 @@ void FPVFloatRampOwner::ModifyOwner()
 	}
 }
 
+UObject* FPVFloatRampOwner::GetOwner()
+{
+	return CurveOwner;
+}
+
+const UObject* FPVFloatRampOwner::GetOwner() const
+{
+	return CurveOwner;
+}
+
 TArray<const UObject*> FPVFloatRampOwner::GetOwners() const
 {
 	if (CurveOwner)
@@ -240,6 +250,8 @@ void FPVFloatRampCustomization::CustomizeHeader(
 					if (YAxisMin != -FLT_MAX || YAxisMax != FLT_MAX)
 						ClampKeyValue(Curve, *KeyIter, YAxisMin, YAxisMax);
 				}
+
+				FloatRampOwner->GetOwner()->PostEditChange();
 			}
 		);
 
@@ -664,6 +676,8 @@ void FPVFloatRampCustomization::OnRampPresetSelected(const FString& Selection) c
 		Curve.SetKeyInterpMode(Curve.AddKey(0.5f, 0.0f), RCIM_Cubic);
 		Curve.SetKeyInterpMode(Curve.AddKey(1.0f, 1.0f), RCIM_Cubic);
 	}
+
+	FloatRampOwner->OnCurveChanged(FloatRampOwner->GetCurves());
 }
 
 FReply FPVFloatRampCustomization::OnRampPreviewDoubleClick(const FGeometry& InMyGeometry, const FPointerEvent& InMouseEvent)
diff --git a/Engine/Plugins/Experimental/ProceduralVegetationEditor/Source/ProceduralVegetationEditor/Private/Customizations/PVFloatRampCustomization.h b/Engine/Plugins/Experimental/ProceduralVegetationEditor/Source/ProceduralVegetationEditor/Private/Customizations/PVFloatRampCustomization.h
--- a/Engine/Plugins/Experimental/ProceduralVegetationEditor/Source/ProceduralVegetationEditor/Private/Customizations/PVFloatRampCustomization.h
+++ b/Engine/Plugins/Experimental/ProceduralVegetationEditor/Source/ProceduralVegetationEditor/Private/Customizations/PVFloatRampCustomization.h
@@ -36,6 +36,8 @@ public:
 	const FRichCurve& GetRichCurve() const;
 
 	void SetOwner(UObject* InOwner);
+	UObject* GetOwner();
+	const UObject* GetOwner() const;
 	virtual TArray<const UObject*> GetOwners() const override;
 	virtual void ModifyOwner() override;
 	virtual void MakeTransactional() override;
diff --git a/Engine/Plugins/Experimental/ProceduralVegetationEditor/Source/ProceduralVegetationEditor/Private/Helpers/PVExportHelper.cpp b/Engine/Plugins/Experimental/ProceduralVegetationEditor/Source/ProceduralVegetationEditor/Private/Helpers/PVExportHelper.cpp
--- a/Engine/Plugins/Experimental/ProceduralVegetationEditor/Source/ProceduralVegetationEditor/Private/Helpers/PVExportHelper.cpp
+++ b/Engine/Plugins/Experimental/ProceduralVegetationEditor/Source/ProceduralVegetationEditor/Private/Helpers/PVExportHelper.cpp
@@ -1,45 +1,54 @@
 ﻿// Copyright Epic Games, Inc. All Rights Reserved.
 
+
 #include "Helpers/PVExportHelper.h"
-#include "PVExportParams.h"
+
 #include "DynamicMeshEditor.h"
+#if WITH_EDITOR
 #include "DynamicWindImportData.h"
+#endif
 #include "DynamicWindSkeletalData.h"
-#include "ProceduralVegetation.h"
-#include "ProceduralVegetationLink.h"
-#include "ProceduralVegetationEditorModule.h"
 #include "MeshDescriptionToDynamicMesh.h"
-#include "PlanarCut.h"
 #include "PackageTools.h"
+#include "PlanarCut.h"
+#include "ProceduralVegetation.h"
+#include "ProceduralVegetationEditorModule.h"
+#include "ProceduralVegetationLink.h"
+#include "PVExportParams.h"
+#include "PVWindSettings.h"
 #include "RenderUtils.h"
-#include "StaticMeshAttributes.h"
+#include "SkinnedAssetCompiler.h"
+#include "StaticMeshCompiler.h"
+#include "UDynamicMesh.h"
+
 #include "Animation/Skeleton.h"
+
+#include "AssetRegistry/AssetData.h"
 #include "AssetRegistry/IAssetRegistry.h"
-#include "Components/InstancedStaticMeshComponent.h"
+
 #include "ConversionUtils/SceneComponentToDynamicMesh.h"
+
 #include "DynamicMesh/DynamicBoneAttribute.h"
 #include "DynamicMesh/DynamicVertexSkinWeightsAttribute.h"
 
-#include "UDynamicMesh.h"
 #include "Engine/SkeletalMesh.h"
 #include "Engine/StaticMesh.h"
+
 #include "Facades/PVBranchFacade.h"
 #include "Facades/PVPointFacade.h"
+
 #include "GeometryCollection/GeometryCollection.h"
-#include "GeometryCollection/ManagedArrayCollection.h"
+#include "GeometryCollection/Facades/CollectionConstraintOverrideFacade.h"
+
 #include "GeometryScript/GeometryScriptTypes.h"
 #include "GeometryScript/MeshAssetFunctions.h"
 #include "GeometryScript/MeshBoneWeightFunctions.h"
+
 #include "Helpers/PVUtilities.h"
-#include "Materials/MaterialInterface.h"
+
 #include "Misc/Paths.h"
+
 #include "UObject/Package.h"
-#include "PVWindSettings.h"
-#include "SkinnedAssetCompiler.h"
-#include "StaticMeshCompiler.h"
-#include "GeometryCollection/Facades/CollectionConstraintOverrideFacade.h"
-#include "AssetRegistry/AssetData.h"
-#include "AssetRegistry/AssetRegistryModule.h"
 
 #define LOCTEXT_NAMESPACE "PVExportHelper"
 
@@ -99,11 +108,11 @@ namespace PV::Export::Internal
 
 		float TotalWeightSum = 0;
 		float ThisStageWeightSum = 0;
-		for (int32 i = 0; i < (int32)EExportMeshStage::MAX; ++i)
+		for (int32 i = 0; i < static_cast<int32>(EExportMeshStage::MAX); ++i)
 		{
-			const float StatusWeight = CalculateMeshExportStageWeight(EExportMeshStage(i), bIsBuildingNaniteAssemblies, bIsExportingSkeletalMesh);
+			const float StatusWeight = CalculateMeshExportStageWeight(static_cast<EExportMeshStage>(i), bIsBuildingNaniteAssemblies, bIsExportingSkeletalMesh);
 			TotalWeightSum += StatusWeight;
-			if (i < int32(InStage))
+			if (i < static_cast<int32>(InStage))
 			{
 				ThisStageWeightSum += StatusWeight;
 			}
@@ -194,46 +203,32 @@ namespace PV::Export::Internal
 	{
 		TRACE_CPUPROFILER_EVENT_SCOPE(PV::Export::CollectionToDynamicMesh);
 		
-		if (Collection.NumElements(FGeometryCollection::TransformGroup) == 0)
-		{
-			return nullptr;
-		}
-
-		const TUniquePtr<FGeometryCollection> GeomCollection = TUniquePtr<FGeometryCollection>(Collection.NewCopy<FGeometryCollection>());
-		if (!GeomCollection.IsValid())
-		{
-			return nullptr;
-		}
-
 		TObjectPtr<UDynamicMesh> NewMesh = NewObject<UDynamicMesh>();
 
-		const TManagedArray<FTransform3f>& BoneTransforms = Collection.GetAttribute<FTransform3f>(
-			"Transform", FGeometryCollection::TransformGroup);
-
-		TArray<int32> TransformIndices;
-		TransformIndices.AddUninitialized(BoneTransforms.Num());
-
-		int32 Idx = 0;
-		for (int32& TransformIdx : TransformIndices)
+		if (Collection.NumElements(FGeometryCollection::TransformGroup) > 0)
 		{
-			TransformIdx = Idx++;
-		}
-
-		FMeshDescription MeshDescription;
-		FStaticMeshAttributes Attributes(MeshDescription);
-		Attributes.Register();
+			if (const TUniquePtr<FGeometryCollection> GeomCollection = TUniquePtr<FGeometryCollection>(Collection.NewCopy<FGeometryCollection>()))
+			{
+				const TManagedArray<FTransform3f>& BoneTransforms = Collection.GetAttribute<FTransform3f>(
+					"Transform", FGeometryCollection::TransformGroup);
 
-		FTransform TransformOut;
+				TArray<int32> TransformIndices;
+				TransformIndices.AddUninitialized(BoneTransforms.Num());
 
-		constexpr bool bCenterPivot = false;
-		ConvertToMeshDescription(MeshDescription, TransformOut, bCenterPivot, *GeomCollection, BoneTransforms, TransformIndices);
+				int32 Idx = 0;
+				for (int32& TransformIdx : TransformIndices)
+				{
+					TransformIdx = Idx++;
+				}
 
-		NewMesh->Reset();
+				constexpr bool bCenterPivot = false;
+				FTransform TransformOut;
+				
+				FDynamicMesh3 CombinedMesh;
+				ConvertGeometryCollectionToDynamicMesh(CombinedMesh, TransformOut, bCenterPivot, *GeomCollection, true, BoneTransforms.GetConstArray(), true, TransformIndices);
 
-		UE::Geometry::FDynamicMesh3& DynMesh = NewMesh->GetMeshRef();
-		{
-			FMeshDescriptionToDynamicMesh ConverterToDynamicMesh;
-			ConverterToDynamicMesh.Convert(&MeshDescription, DynMesh);
+				NewMesh->SetMesh(CombinedMesh);
+			}
 		}
 
 		return NewMesh;
@@ -369,7 +364,7 @@ namespace PV::Export::Internal
 				}
 			}
 
-			if (!OnProgressUpdated((FoliageNameIndex / (float)NumFoliageNames) * 0.5f))
+			if (!OnProgressUpdated((FoliageNameIndex / static_cast<float>(NumFoliageNames)) * 0.5f))
 			{
 				return false;
 			}
@@ -384,11 +379,11 @@ namespace PV::Export::Internal
 			{
 				continue;
 			}
+			
+			const TObjectPtr<UStaticMesh> StaticMesh = FoliageStaticMeshes.FindRef(FoliageName, nullptr);
+			const TObjectPtr<USkeletalMesh> SkeletalMesh = FoliageSkeletalMeshes.FindRef(FoliageName, nullptr);
 
-			const TObjectPtr<UStaticMesh>* StaticMesh = FoliageStaticMeshes.Find(FoliageName);
-			const TObjectPtr<USkeletalMesh>* SkeletalMesh = FoliageSkeletalMeshes.Find(FoliageName);
-
-			check((StaticMesh && *StaticMesh) || (SkeletalMesh && *SkeletalMesh));
+			check(StaticMesh || SkeletalMesh);
 
 			if (SkeletalMesh != nullptr)
 			{
@@ -433,13 +428,13 @@ namespace PV::Export::Internal
 						const int32 MaterialID = MaterialIDs->GetValue(OldTid);
 						if (StaticMesh)
 						{
-							UMaterialInterface* Material = (*StaticMesh)->GetMaterial(MaterialID);
+							UMaterialInterface* Material = StaticMesh->GetMaterial(MaterialID);
 							const int32 NewMaterialID = Materials.Find(Material);
 							MaterialIDs->SetValue(NewTid, NewMaterialID);
 						}
 						else if (SkeletalMesh)
 						{
-							UMaterialInterface* Material = (*SkeletalMesh)->GetMaterials()[MaterialID].MaterialInterface;
+							UMaterialInterface* Material = SkeletalMesh->GetMaterials()[MaterialID].MaterialInterface;
 							const int32 NewMaterialID = Materials.Find(Material);
 							MaterialIDs->SetValue(NewTid, NewMaterialID);
 						}
@@ -448,15 +443,15 @@ namespace PV::Export::Internal
 
 				if (StaticMesh)
 				{
-					OnMeshInstanceCombined((*StaticMesh)->GetFullName(), Id, VertexStart, Mesh.VertexCount());
+					OnMeshInstanceCombined(StaticMesh->GetFullName(), Id, VertexStart, Mesh.VertexCount());
 				}
 				else if (SkeletalMesh != nullptr)
 				{
-					OnMeshInstanceCombined((*SkeletalMesh)->GetFullName(), Id, VertexStart, Mesh.VertexCount());
+					OnMeshInstanceCombined(SkeletalMesh->GetFullName(), Id, VertexStart, Mesh.VertexCount());
 				}
 			});
 
-			if (!OnProgressUpdated(0.5f + (Id / (float)NumInstances) * 0.5f))
+			if (!OnProgressUpdated(0.5f + (Id / static_cast<float>(NumInstances)) * 0.5f))
 			{
 				return false;
 			}
@@ -583,7 +578,7 @@ namespace PV::Export::Internal
 			}
 
 			EGeometryScriptOutcomePins Output;
-			FGeometryScriptCopyMeshToAssetOptions Options = GetCopyMeshToAssetOptions(PartMaterials, PartStaticMesh->GetNaniteSettings(), ENaniteShapePreservation::Voxelize);
+			FGeometryScriptCopyMeshToAssetOptions Options = GetCopyMeshToAssetOptions(PartMaterials, PartStaticMesh->GetNaniteSettings(), NaniteVoxelsSupported() ? ENaniteShapePreservation::Voxelize : ENaniteShapePreservation::PreserveArea);
 
 			TObjectPtr<UDynamicMesh> DynamicMesh = NewObject<UDynamicMesh>();
 			DynamicMesh->SetMesh(MoveTemp(PartDynamicMesh));
@@ -798,7 +793,7 @@ namespace PV::Export::Internal
 		
 		//Convert the dynamic mesh to static mesh
 		EGeometryScriptOutcomePins Output;
-		FGeometryScriptCopyMeshToAssetOptions Options = GetCopyMeshToAssetOptions(Materials, SkeletalMesh->GetNaniteSettings(), ENaniteShapePreservation::Voxelize);
+		FGeometryScriptCopyMeshToAssetOptions Options = GetCopyMeshToAssetOptions(Materials, SkeletalMesh->GetNaniteSettings(), NaniteVoxelsSupported() ? ENaniteShapePreservation::Voxelize : ENaniteShapePreservation::PreserveArea);
 		Options.BoneHierarchyMismatchHandling = EGeometryScriptBoneHierarchyMismatchHandling::CreateNewReferenceSkeleton;
 		
 		UGeometryScriptLibrary_StaticMeshFunctions::CopyMeshToSkeletalMesh(
@@ -1402,11 +1397,13 @@ namespace PV::Export
 				return false;
 			}
 
-			const TObjectPtr<UStaticMesh> ExistingMeshAsset = Internal::GetAssetOfType<UStaticMesh>(MeshPackageName);
+			if (TObjectPtr<UStaticMesh> ExistingMeshAsset = Internal::GetAssetOfType<UStaticMesh>(MeshPackageName))
+			{
+				ExistingMeshAsset->SetExternalPackage(nullptr);
+				ExistingMeshAsset->MarkAsGarbage();
+			}
 
-			const TObjectPtr<UStaticMesh> ExportMesh = ExistingMeshAsset != nullptr
-				? ExistingMeshAsset.Get()
-				: NewObject<UStaticMesh>(Internal::CreateAssetPackage(MeshPackageName), *MeshName, RF_Standalone | RF_Public);
+			const TObjectPtr<UStaticMesh> ExportMesh = NewObject<UStaticMesh>(Internal::CreateAssetPackage(MeshPackageName), *MeshName, RF_Standalone | RF_Public);
 
 			ModifiedPackages.Add(ExportMesh->GetOutermost());
 
@@ -1439,15 +1436,20 @@ namespace PV::Export
 				return false;
 			}
 
-			const TObjectPtr<USkeletalMesh> ExistingMeshAsset = Internal::GetAssetOfType<USkeletalMesh>(MeshPackageName);
-			const TObjectPtr<USkeleton> ExistingSkeletonAsset = Internal::GetAssetOfType<USkeleton>(SkeletonPackageName);
+			if (TObjectPtr<USkeletalMesh> ExistingMeshAsset = Internal::GetAssetOfType<USkeletalMesh>(MeshPackageName))
+			{
+				ExistingMeshAsset->SetExternalPackage(nullptr);
+				ExistingMeshAsset->MarkAsGarbage();
+			}
+			
+			if (TObjectPtr<USkeleton> ExistingSkeletonAsset = Internal::GetAssetOfType<USkeleton>(SkeletonPackageName))
+			{
+				ExistingSkeletonAsset->SetExternalPackage(nullptr);
+				ExistingSkeletonAsset->MarkAsGarbage();
+			}
 
-			const TObjectPtr<USkeletalMesh> ExportMesh = ExistingMeshAsset != nullptr 
-				? ExistingMeshAsset.Get() 
-				: NewObject<USkeletalMesh>(Internal::CreateAssetPackage(MeshPackageName), *MeshName, RF_Standalone | RF_Public);
-			const TObjectPtr<USkeleton> ExportSkeleton = ExistingSkeletonAsset != nullptr 
-				? ExistingSkeletonAsset.Get() 
-				: NewObject<USkeleton>(Internal::CreateAssetPackage(SkeletonPackageName), *SkeletonName, RF_Standalone | RF_Public);
+			const TObjectPtr<USkeletalMesh> ExportMesh = NewObject<USkeletalMesh>(Internal::CreateAssetPackage(MeshPackageName), *MeshName, RF_Standalone | RF_Public);
+			const TObjectPtr<USkeleton> ExportSkeleton = NewObject<USkeleton>(Internal::CreateAssetPackage(SkeletonPackageName), *SkeletonName, RF_Standalone | RF_Public);
 			
 			ExportMesh->SetSkeleton(ExportSkeleton);
 
diff --git a/Engine/Plugins/Experimental/ProceduralVegetationEditor/Source/ProceduralVegetationEditor/Private/PVEditor.cpp b/Engine/Plugins/Experimental/ProceduralVegetationEditor/Source/ProceduralVegetationEditor/Private/PVEditor.cpp
--- a/Engine/Plugins/Experimental/ProceduralVegetationEditor/Source/ProceduralVegetationEditor/Private/PVEditor.cpp
+++ b/Engine/Plugins/Experimental/ProceduralVegetationEditor/Source/ProceduralVegetationEditor/Private/PVEditor.cpp
@@ -8,20 +8,20 @@
 #include "ProceduralVegetation.h"
 #include "PVEditorCommands.h"
 #include "PVEditorSchema.h"
-#include "PVEditorSettings.h"
-#include "Helpers/PVExportHelper.h"
 #include "SPVEditorViewport.h"
 #include "ToolMenus.h"
 #include "ToolMenuSection.h"
-#include "Misc/UObjectToken.h"
-#include "FileHelpers.h"
 
 #include "Dataflow/DataflowCollectionSpreadSheetWidget.h"
 
 #include "DataTypes/PVData.h"
 
+#include "Helpers/PVExportHelper.h"
 #include "Helpers/PVUtilities.h"
 
+#include "Misc/MessageDialog.h"
+#include "Misc/UObjectToken.h"
+
 #include "Nodes/PCGEditorGraphNodeBase.h"
 #include "Nodes/PVOutputSettings.h"
 
@@ -272,14 +272,6 @@ void FPVEditor::RegisterToolbarInternal(FToolMenuSection& PCGSection) const
 			TAttribute<FText>(),
 			TAttribute<FText>(),
 			FSlateIcon(FAppStyle::GetAppStyleSetName(), "Icons.Toolbar.Export")));
-
-	PCGSection.AddSeparator("Debug Settings");
-
-	PCGSection.AddEntry(FToolMenuEntry::InitToolBarButton(
-		Commands.ShowMannequin,
-		TAttribute<FText>(),
-		TAttribute<FText>(),
-		FSlateIcon(FAppStyle::GetAppStyleSetName(), "ClassIcon.PoseAsset")));
 }
 
 void FPVEditor::BindCommands()
@@ -291,24 +283,6 @@ void FPVEditor::BindCommands()
 	ToolkitCommands->MapAction(
 		EditorCommands.Export,
 		FExecuteAction::CreateSP(this, &FPVEditor::OnExport_Clicked));
-
-	ToolkitCommands->MapAction(
-		EditorCommands.ShowMannequin,
-		FExecuteAction::CreateLambda(
-			[this]
-				{
-					UPVEditorSettings* Settings = GetMutableDefault<UPVEditorSettings>();
-					Settings->bShowMannequin = !Settings->bShowMannequin;
-					GetPCGEditorGraph()->GetPCGGraph()->ForceNotificationForEditor(EPCGChangeType::Settings);
-				}
-		),
-		FCanExecuteAction(),
-		FIsActionChecked::CreateLambda([]
-			{
-				return GetDefault<UPVEditorSettings>()->bShowMannequin;
-			}
-		)
-	);
 }
 
 void FPVEditor::OnSelectedNodesChanged(const TSet<UObject*>& InNewSelection)
@@ -407,7 +381,7 @@ FReply FPVEditor::ExportInternal()
 		return !bHasInvalidPath;
 	};
 
-	const static auto ShouldOverwriteOutstandingChanges = [](const TArray<UPCGNode*>& NodesToExport)->bool 
+	const static auto ShouldOverwriteAssets = [](const TArray<UPCGNode*>& NodesToExport)->bool 
 	{
 		TArray<FString> ExportPaths;
 		ExportPaths.Reserve(NodesToExport.Num() * 2);
@@ -423,37 +397,23 @@ FReply FPVEditor::ExportInternal()
 			}
 		}
 
-		TArray<UPackage*> DirtyPackages;
-		DirtyPackages.Reserve(ExportPaths.Num());
+		int32 NumExistingPackages = 0;
 		for (const FString& AssetPath : ExportPaths)
 		{
-			UPackage* ExistingPackage = FindPackage(nullptr, *AssetPath);
-			if (ExistingPackage && ExistingPackage->IsDirty())
-			{
-				DirtyPackages.Add(ExistingPackage);
-			}
+			NumExistingPackages += FPackageName::DoesPackageExist(AssetPath);
 		}
 
-		if (DirtyPackages.Num() == 0)
+		if (NumExistingPackages == 0)
 		{
 			return true;
 		}
 
-		FEditorFileUtils::FPromptForCheckoutAndSaveParams SaveParams;
-		SaveParams.bCheckDirty = false;
-		SaveParams.bPromptToSave = true;
-		SaveParams.bAlreadyCheckedOut = false;
-		SaveParams.bCanBeDeclined = true;
-		SaveParams.bIsExplicitSave = false;
-		SaveParams.Title = LOCTEXT("PackagesNeedResave_Title", "WARNING: Unsaved changes may be lost!");
-		SaveParams.Message = LOCTEXT("PackagesNeedResave_Message", "Assets have unsaved changes which may be lost during the export process. Please save the assets, or cancel, if you do not want to risk loosing changes to the assets.");
-
-		const FEditorFileUtils::EPromptReturnCode Return = FEditorFileUtils::PromptForCheckoutAndSave(DirtyPackages, SaveParams);
-		if (Return == FEditorFileUtils::EPromptReturnCode::PR_Cancelled)
+		if (FMessageDialog::Open(EAppMsgType::YesNo,
+			FText::Format(LOCTEXT("PVE_OverwritePackagesWarning", "WARNING: {0} assets will be overwritten on export.\nDo you want to continue?"),
+				NumExistingPackages)) == EAppReturnType::No)
 		{
 			return false;
 		}
-
 		return true;
 	};
 
@@ -490,7 +450,7 @@ FReply FPVEditor::ExportInternal()
 		return FReply::Handled();
 	}
 
-	if (!ShouldOverwriteOutstandingChanges(NodesToExport))
+	if (!ShouldOverwriteAssets(NodesToExport))
 	{
 		return FReply::Handled();
 	}
diff --git a/Engine/Plugins/Experimental/ProceduralVegetationEditor/Source/ProceduralVegetationEditor/Private/PVEditorCommands.cpp b/Engine/Plugins/Experimental/ProceduralVegetationEditor/Source/ProceduralVegetationEditor/Private/PVEditorCommands.cpp
--- a/Engine/Plugins/Experimental/ProceduralVegetationEditor/Source/ProceduralVegetationEditor/Private/PVEditorCommands.cpp
+++ b/Engine/Plugins/Experimental/ProceduralVegetationEditor/Source/ProceduralVegetationEditor/Private/PVEditorCommands.cpp
@@ -19,6 +19,7 @@ void FPVEditorCommands::RegisterCommands()
 {
 	UI_COMMAND(Export, "Export", "Export outputs.", EUserInterfaceActionType::Button, FInputChord(EModifierKey::Control, EKeys::E));
 	UI_COMMAND(ShowMannequin, "Mannequin", "Show/Hide Mannequin", EUserInterfaceActionType::ToggleButton, FInputChord(EModifierKey::Control, EKeys::M));
+	UI_COMMAND(ShowScaleVisualization, "Scale Visualization", "Show/Hide Scale Visualization", EUserInterfaceActionType::ToggleButton, FInputChord(EModifierKey::Control, EKeys::V));
 }
 
 #undef LOCTEXT_NAMESPACE
diff --git a/Engine/Plugins/Experimental/ProceduralVegetationEditor/Source/ProceduralVegetationEditor/Private/PVEditorCommands.h b/Engine/Plugins/Experimental/ProceduralVegetationEditor/Source/ProceduralVegetationEditor/Private/PVEditorCommands.h
--- a/Engine/Plugins/Experimental/ProceduralVegetationEditor/Source/ProceduralVegetationEditor/Private/PVEditorCommands.h
+++ b/Engine/Plugins/Experimental/ProceduralVegetationEditor/Source/ProceduralVegetationEditor/Private/PVEditorCommands.h
@@ -13,7 +13,10 @@ public:
 	virtual void RegisterCommands() override;
 	// ~End TCommands<> interface
 
+	//Editor Commands
 	TSharedPtr<FUICommandInfo> Export;
-	
+
+	//Viewport Commands
 	TSharedPtr<FUICommandInfo> ShowMannequin;
+	TSharedPtr<FUICommandInfo> ShowScaleVisualization;
 };
\ No newline at end of file
diff --git a/Engine/Plugins/Experimental/ProceduralVegetationEditor/Source/ProceduralVegetationEditor/Private/ProceduralVegetationFactory.cpp b/Engine/Plugins/Experimental/ProceduralVegetationEditor/Source/ProceduralVegetationEditor/Private/ProceduralVegetationFactory.cpp
--- a/Engine/Plugins/Experimental/ProceduralVegetationEditor/Source/ProceduralVegetationEditor/Private/ProceduralVegetationFactory.cpp
+++ b/Engine/Plugins/Experimental/ProceduralVegetationEditor/Source/ProceduralVegetationEditor/Private/ProceduralVegetationFactory.cpp
@@ -22,7 +22,7 @@ UObject* UProceduralVegetationFactory::FactoryCreateNew(UClass* InClass, UObject
 	UProceduralVegetation* NewProceduralVegetation;
 	if (SampleProceduralVegetation)
 	{
-		NewProceduralVegetation = NewObject<UProceduralVegetation>(InParent, InClass, SampleProceduralVegetation->GetFName(), InFlags, SampleProceduralVegetation);
+		NewProceduralVegetation = NewObject<UProceduralVegetation>(InParent, InClass, InName, InFlags, SampleProceduralVegetation);
 
 		TObjectPtr<UProceduralVegetationGraph> SampleGraph = SampleProceduralVegetation->GetGraph();
 		NewProceduralVegetation->CreateGraph(SampleGraph);
diff --git a/Engine/Plugins/Experimental/ProceduralVegetationEditor/Source/ProceduralVegetationEditor/Private/SPVEditorViewport.cpp b/Engine/Plugins/Experimental/ProceduralVegetationEditor/Source/ProceduralVegetationEditor/Private/SPVEditorViewport.cpp
--- a/Engine/Plugins/Experimental/ProceduralVegetationEditor/Source/ProceduralVegetationEditor/Private/SPVEditorViewport.cpp
+++ b/Engine/Plugins/Experimental/ProceduralVegetationEditor/Source/ProceduralVegetationEditor/Private/SPVEditorViewport.cpp
@@ -11,6 +11,11 @@
 #include "ToolMenuSection.h"
 #include "Engine/StaticMesh.h"
 #include "Components/StaticMeshComponent.h"
+#include "Components/TextRenderComponent.h"
+
+#include "Visualizations/PVScaleVisualizationComponent.h"
+
+#include "Widgets/Input/SNumericEntryBox.h"
 
 #define LOCTEXT_NAMESPACE "PVEditorViewportToolbarSections"
 
@@ -23,11 +28,28 @@ FToolMenuEntry SPVEditorViewport::CreateSettingsToolbarMenu()
 				FText::GetEmpty(),
 				FText::GetEmpty(),
 				FNewToolMenuDelegate::CreateLambda([](UToolMenu* Submenu)
-				{
-					// Add preference toggles to this section.
-					FToolMenuSection& PreferencesSection = Submenu->FindOrAddSection(TEXT("PVEMannequinSettings"), LOCTEXT("PVEMannequinSettingsSectionLabel", "Mannequin Settings"));
-					PreferencesSection.AddMenuEntry(FPVEditorCommands::Get().ShowMannequin);
-				}),
+					{
+						const UPVMannequinWidgetContext* Context = Submenu->FindContext<UPVMannequinWidgetContext>();
+						
+						// Add preference toggles to this section.
+						FToolMenuSection& MannequinSection = Submenu->FindOrAddSection(TEXT("PVE_MannequinSettings"),
+							LOCTEXT("PVE_MannequinSettingsSectionLabel", "Mannequin Settings"));
+						MannequinSection.AddMenuEntry(FPVEditorCommands::Get().ShowMannequin);
+						MannequinSection.AddEntry(
+							FToolMenuEntry::InitWidget(
+								"PVE_MannequinOffset",
+								Context->MannequinOffsetWidget.ToSharedRef(),
+								LOCTEXT("PVE_MannequinOffsetLabel", "Mannequin Offset"),
+								false,
+								true,
+								false,
+								LOCTEXT("PVE_MannequinOffsetTooltip", "Manually Offset the Mannequin")
+							)
+						);
+						FToolMenuSection& ScaleVisSection = Submenu->FindOrAddSection(TEXT("PVE_ScaleVisSettings"),
+							LOCTEXT("PVE_ScaleVisSettingsSectionLabel", "Scale Vis. Settings"));
+						ScaleVisSection.AddMenuEntry(FPVEditorCommands::Get().ShowScaleVisualization);
+					}),
 				/*bInOpenSubMenuOnClick=*/false,
 				FSlateIcon(FAppStyle::GetAppStyleSetName(), "LevelEditor.GameSettings"));
 		}));
@@ -65,6 +87,7 @@ TSharedPtr<SWidget> SPVEditorViewport::BuildViewportToolbar()
 		Context.AppendCommandList(GetCommandList());
 		Context.AddExtender(GetExtenders());
 		Context.AddObject(UE::UnrealEd::CreateViewportToolbarDefaultContext(GetViewportWidget()));
+		Context.AddObject(CreateMannequinWidgetContext(CreateMannequinOffsetWidget()));
 	}
 	
 	return UToolMenus::Get()->GenerateWidget(ToolbarName, Context);
@@ -74,80 +97,189 @@ void SPVEditorViewport::BindCommands()
 {
 	SPCGEditorViewport::BindCommands();
 
-	const FPVEditorCommands& Commands = FPVEditorCommands::Get();
+	FPVEditorCommands& Commands = FPVEditorCommands::Get();
 
 	CommandList->MapAction(
 		Commands.ShowMannequin,
-		FExecuteAction::CreateSP(this, &SPVEditorViewport::ToggleShowManniquen),
+		FExecuteAction::CreateSP(this, &SPVEditorViewport::ToggleShowMannequin),
+		FCanExecuteAction(),
+		FIsActionChecked::CreateSP(this, &SPVEditorViewport::IsShowMannequinChecked));
+
+	CommandList->MapAction(
+		Commands.ShowScaleVisualization,
+		FExecuteAction::CreateSP(this, &SPVEditorViewport::ToggleShowScaleVis),
 		FCanExecuteAction(),
-		FIsActionChecked::CreateSP(this, &SPVEditorViewport::IsShowManniquenChecked));
+		FIsActionChecked::CreateSP(this, &SPVEditorViewport::IsShowScaleVisChecked));
 }
 
-void SPVEditorViewport::SetMannequinState(bool InEnable)
+void SPVEditorViewport::InitVisualizationScene()
 {
-	UStaticMesh* MannequinMesh = LoadObject<UStaticMesh>(
-		nullptr,
-		TEXT("/ProceduralVegetationEditor/Mannequin/Viewport_Mannequin_T_Pose.Viewport_Mannequin_T_Pose")
-	);
-
-	if (!MannequinMesh)
-	{
-		return;
-	}
-	
-	if (InEnable)
+	const FString MannequinPath = TEXT("/ProceduralVegetationEditor/Mannequin/Viewport_Mannequin_T_Pose.Viewport_Mannequin_T_Pose");
+	if (UStaticMesh* MannequinMesh = LoadObject<UStaticMesh>(nullptr, MannequinPath))
 	{
-		UStaticMeshComponent* MannequinComponent = NewObject<UStaticMeshComponent>(GetTransientPackage(), NAME_None, RF_Transient);
+		MannequinComponent = NewObject<UStaticMeshComponent>(GetTransientPackage(), NAME_None, RF_Transient);
 		MannequinComponent->SetStaticMesh(MannequinMesh);
-		FBoxSphereBounds MannequinBounds = MannequinComponent->CalcLocalBounds();
+		const FBoxSphereBounds MannequinBounds = MannequinComponent->CalcLocalBounds();
 
 		ManagedResources.Add(MannequinComponent);
-		AdvancedPreviewScene->AddComponent(
-			MannequinComponent,
-			FTransform(FVector((-MannequinBounds.BoxExtent.X - MannequinBounds.BoxExtent.X) * 1.5f, 0, 0))
+		AdvancedPreviewScene->AddComponent(MannequinComponent, FTransform::Identity);
+		MannequinComponent->SetRelativeLocation(
+			FVector((-FocusBounds.BoxExtent.X - MannequinBounds.BoxExtent.X) * 1.5f - GetMannequinOffset(), 0.0, 0.0)
 		);
 	}
-	else
+
+	ScaleVisualizationComponent = NewObject<UPVScaleVisualizationComponent>(GetTransientPackage(), NAME_None, RF_Transient);
+	ScaleVisualizationComponent->SetScaleBounds(FocusBounds);
+
+	ManagedResources.Add(ScaleVisualizationComponent);
+	AdvancedPreviewScene->AddComponent(ScaleVisualizationComponent, FTransform::Identity);
+	for (const TObjectPtr<UTextRenderComponent>& TextRenderComponent : ScaleVisualizationComponent->GetTextRenderComponents())
 	{
-		for (TObjectPtr<UObject> ManagedResource : ManagedResources)
+		ManagedResources.Add(TextRenderComponent);
+		AdvancedPreviewScene->AddComponent(TextRenderComponent, FTransform::Identity);
+	}
+}
+
+void SPVEditorViewport::ResetVisualizationScene()
+{
+	if (MannequinComponent)
+	{
+		AdvancedPreviewScene->RemoveComponent(MannequinComponent);
+		ManagedResources.Remove(MannequinComponent);
+		MannequinComponent->MarkAsGarbage();
+		MannequinComponent = nullptr;
+	}
+
+	if (ScaleVisualizationComponent)
+	{
+		for (const TObjectPtr<UTextRenderComponent>& TextRenderComponent : ScaleVisualizationComponent->GetTextRenderComponents())
 		{
-			if (UStaticMeshComponent* ActorComponent = Cast<UStaticMeshComponent>(ManagedResource))
+			if (TextRenderComponent)
 			{
-				if (ActorComponent->GetStaticMesh()->GetFName() == MannequinMesh->GetFName())
-				{
-					AdvancedPreviewScene->RemoveComponent(ActorComponent);
-				}
+				AdvancedPreviewScene->RemoveComponent(TextRenderComponent);
+				ManagedResources.Remove(TextRenderComponent);
 			}
 		}
+		AdvancedPreviewScene->RemoveComponent(ScaleVisualizationComponent);
+		ManagedResources.Remove(ScaleVisualizationComponent);
+		ScaleVisualizationComponent->MarkAsGarbage();
+		ScaleVisualizationComponent = nullptr;
 	}
+}
 
-	Client->Invalidate();
+TSharedRef<SWidget> SPVEditorViewport::CreateMannequinOffsetWidget() const
+{
+	return SNew(SNumericEntryBox<float>)
+		.AllowSpin(true)
+		.AllowWheel(true)
+		.Delta(5.0f)
+		.MinValue(-100.0f)
+		.MaxValue(500.0f)
+		.MinSliderValue(-100.0f)
+		.MaxSliderValue(500.0f)
+		.Value_Lambda([this]{ return GetMannequinOffset(); })
+		.OnValueChanged_Lambda([this](const float NewValue)
+			{
+				SetMannequinOffset(NewValue, false);
+			})
+		.OnValueCommitted_Lambda([this](const float NewValue, ETextCommit::Type)
+			{
+				SetMannequinOffset(NewValue, true);
+			});
 }
 
-void SPVEditorViewport::ToggleShowManniquen()
+void SPVEditorViewport::ToggleShowMannequin()
 {
 	UPVEditorSettings* EditorSettings = GetMutableDefault<UPVEditorSettings>();
 	EditorSettings->bShowMannequin = !EditorSettings->bShowMannequin;
+	EditorSettings->SaveConfig();
 
 	SetMannequinState(EditorSettings->bShowMannequin);
 }
 
-bool SPVEditorViewport::IsShowManniquenChecked() const
+bool SPVEditorViewport::IsShowMannequinChecked() const
 {
 	return GetDefault<UPVEditorSettings>()->bShowMannequin;
 }
 
+void SPVEditorViewport::SetMannequinState(bool InEnable)
+{
+	if (MannequinComponent)
+	{
+		MannequinComponent->SetVisibility(InEnable, true);
+		Client->Invalidate();
+	}
+}
+
+void SPVEditorViewport::ToggleShowScaleVis()
+{
+	UPVEditorSettings* EditorSettings = GetMutableDefault<UPVEditorSettings>();
+	EditorSettings->bShowScaleVisualization = !EditorSettings->bShowScaleVisualization;
+	EditorSettings->SaveConfig();
+	
+	SetScaleVisState(EditorSettings->bShowScaleVisualization);
+}
+
+bool SPVEditorViewport::IsShowScaleVisChecked() const
+{
+	return GetMutableDefault<UPVEditorSettings>()->bShowScaleVisualization;
+}
+
+void SPVEditorViewport::SetScaleVisState(bool InEnable)
+{
+	if (MannequinComponent)
+	{
+		ScaleVisualizationComponent->SetVisibility(InEnable, true);
+		Client->Invalidate();
+	}
+}
+
+float SPVEditorViewport::GetMannequinOffset() const
+{
+	return GetDefault<UPVEditorSettings>()->MannequinOffset;
+}
+
+void SPVEditorViewport::SetMannequinOffset(const float NewValue, const bool bSaveConfig) const
+{
+	UPVEditorSettings* const EditorSettings = GetMutableDefault<UPVEditorSettings>();
+	EditorSettings->MannequinOffset = NewValue;
+	if (bSaveConfig)
+	{
+		EditorSettings->SaveConfig();
+		return;
+	}
+	
+	if (MannequinComponent)
+	{
+		const FBoxSphereBounds MannequinBounds = MannequinComponent->CalcLocalBounds();
+		MannequinComponent->SetRelativeLocation(
+			FVector((-FocusBounds.BoxExtent.X - MannequinBounds.BoxExtent.X) * 1.5f - EditorSettings->MannequinOffset, 0.0, 0.0)
+		);
+		Client->Invalidate();
+	}
+}
+
+UPVMannequinWidgetContext* SPVEditorViewport::CreateMannequinWidgetContext(TSharedPtr<SWidget> InOffsetWidget)
+{
+	UPVMannequinWidgetContext* Context = NewObject<UPVMannequinWidgetContext>();
+	Context->MannequinOffsetWidget = InOffsetWidget;
+	return Context;
+}
+
 void SPVEditorViewport::OnSetupScene()
 {
 	SPCGEditorViewport::OnSetupScene();
 
-	SetMannequinState(IsShowManniquenChecked());
+	InitVisualizationScene();
+	
+	SetMannequinState(IsShowMannequinChecked());
+	SetScaleVisState(IsShowScaleVisChecked());
 }
 
 void SPVEditorViewport::OnResetScene()
 {
 	SPCGEditorViewport::OnResetScene();
 
-	SetMannequinState(IsShowManniquenChecked());
+	ResetVisualizationScene();
 }
 #undef LOCTEXT_NAMESPACE
diff --git a/Engine/Plugins/Experimental/ProceduralVegetationEditor/Source/ProceduralVegetationEditor/Private/SPVEditorViewport.h b/Engine/Plugins/Experimental/ProceduralVegetationEditor/Source/ProceduralVegetationEditor/Private/SPVEditorViewport.h
--- a/Engine/Plugins/Experimental/ProceduralVegetationEditor/Source/ProceduralVegetationEditor/Private/SPVEditorViewport.h
+++ b/Engine/Plugins/Experimental/ProceduralVegetationEditor/Source/ProceduralVegetationEditor/Private/SPVEditorViewport.h
@@ -2,11 +2,25 @@
 
 #pragma once
 
-#include "PVEditorSettings.h"
+#include "CoreMinimal.h"
+
 #include "Widgets/AssetEditorViewport/SPCGEditorViewport.h"
 
+#include "SPVEditorViewport.generated.h"
+
 struct FToolMenuEntry;
+
 class UPVEditorSettings;
+class UPVScaleVisualizationComponent;
+
+UCLASS(MinimalAPI)
+class UPVMannequinWidgetContext : public UObject
+{
+	GENERATED_BODY()
+
+public:
+	TSharedPtr<SWidget> MannequinOffsetWidget = nullptr;
+};
 
 class SPVEditorViewport : public SPCGEditorViewport
 {
@@ -18,11 +32,29 @@ protected:
 	virtual void BindCommands() override;
 
 private:
-	void ToggleShowManniquen();
-	bool IsShowManniquenChecked() const;
+	void InitVisualizationScene();
+	void ResetVisualizationScene();
+
+	TSharedRef<SWidget> CreateMannequinOffsetWidget() const;
+
+	void ToggleShowMannequin();
+	bool IsShowMannequinChecked() const;
 	void SetMannequinState(bool InEnable);
 
+	void ToggleShowScaleVis();
+	bool IsShowScaleVisChecked() const;
+	void SetScaleVisState(bool InEnable);
+
+	float GetMannequinOffset() const;
+	void SetMannequinOffset(float NewValue, bool bSaveConfig = false) const;
+
+	static UPVMannequinWidgetContext* CreateMannequinWidgetContext(TSharedPtr<SWidget> InOffsetWidget);
+
 protected:
 	virtual void OnSetupScene() override;
 	virtual void OnResetScene() override;
+
+private:
+	TObjectPtr<UStaticMeshComponent> MannequinComponent = nullptr;
+	TObjectPtr<UPVScaleVisualizationComponent> ScaleVisualizationComponent = nullptr;
 };
diff --git a/Engine/Plugins/Experimental/ProceduralVegetationEditor/Source/ProceduralVegetationEditor/Private/Visualizations/PVDataVisualization.cpp b/Engine/Plugins/Experimental/ProceduralVegetationEditor/Source/ProceduralVegetationEditor/Private/Visualizations/PVDataVisualization.cpp
--- a/Engine/Plugins/Experimental/ProceduralVegetationEditor/Source/ProceduralVegetationEditor/Private/Visualizations/PVDataVisualization.cpp
+++ b/Engine/Plugins/Experimental/ProceduralVegetationEditor/Source/ProceduralVegetationEditor/Private/Visualizations/PVDataVisualization.cpp
@@ -68,7 +68,6 @@ FPCGSetupSceneFunc FPVDataVisualization::GetViewportSetupFunc(const UPCGSettings
 		}
 			
 		InOutParams.EditorViewportClient->SetRealtime(true);
-
 		if (!InOutParams.Scene->IsGridEnabled())
 		{
 			InOutParams.Scene->HandleToggleGrid();
@@ -189,6 +188,11 @@ void FPVDataVisualization::FoliageRenderer(const FManagedArrayCollection& InColl
 	PV::Utilities::AddFoliageInstances(FoliageFacade, nullptr,
 	                                           PV::Utilities::FFoliageComponentCreatedCallback::CreateLambda(OnComponentCreated),
 	                                           OutInstancedComponentMap);
+
+	for (const auto& [FoliageName , Component] : OutInstancedComponentMap)
+	{
+		OutBounds = OutBounds + Component->CalcLocalBounds();
+	}
 }
 
 void FPVDataVisualization::BonesRenderer(const FManagedArrayCollection& InCollection, FPCGSceneSetupParams& InOutParams,
diff --git a/Engine/Plugins/Experimental/ProceduralVegetationEditor/Source/ProceduralVegetationEditor/Private/Visualizations/PVLineBatchComponent.cpp b/Engine/Plugins/Experimental/ProceduralVegetationEditor/Source/ProceduralVegetationEditor/Private/Visualizations/PVLineBatchComponent.cpp
--- a/Engine/Plugins/Experimental/ProceduralVegetationEditor/Source/ProceduralVegetationEditor/Private/Visualizations/PVLineBatchComponent.cpp
+++ b/Engine/Plugins/Experimental/ProceduralVegetationEditor/Source/ProceduralVegetationEditor/Private/Visualizations/PVLineBatchComponent.cpp
@@ -97,6 +97,11 @@ void UPVLineBatchComponent::Draw(const FSceneView* View, FPrimitiveDrawInterface
 	}
 }
 
+void UPVLineBatchComponent::Flush()
+{
+	Lines.Empty();
+}
+
 FPrimitiveSceneProxy* UPVLineBatchComponent::CreateSceneProxy()
 {
 	FPVLineSceneProxy* Proxy = new FPVLineSceneProxy(this);
diff --git a/Engine/Plugins/Experimental/ProceduralVegetationEditor/Source/ProceduralVegetationEditor/Private/Visualizations/PVLineBatchComponent.h b/Engine/Plugins/Experimental/ProceduralVegetationEditor/Source/ProceduralVegetationEditor/Private/Visualizations/PVLineBatchComponent.h
--- a/Engine/Plugins/Experimental/ProceduralVegetationEditor/Source/ProceduralVegetationEditor/Private/Visualizations/PVLineBatchComponent.h
+++ b/Engine/Plugins/Experimental/ProceduralVegetationEditor/Source/ProceduralVegetationEditor/Private/Visualizations/PVLineBatchComponent.h
@@ -57,6 +57,8 @@ public:
 
 	void Draw(const FSceneView* View, FPrimitiveDrawInterface* PDI) const;
 
+	void Flush();
+	
 	virtual FPrimitiveSceneProxy* CreateSceneProxy() override;
 	virtual FBoxSphereBounds CalcBounds(const FTransform& LocalToWorld) const override;
 
diff --git a/Engine/Plugins/Experimental/ProceduralVegetationEditor/Source/ProceduralVegetationEditor/Private/Visualizations/PVScaleVisualizationComponent.cpp b/Engine/Plugins/Experimental/ProceduralVegetationEditor/Source/ProceduralVegetationEditor/Private/Visualizations/PVScaleVisualizationComponent.cpp
new file mode 100644
--- /dev/null
+++ b/Engine/Plugins/Experimental/ProceduralVegetationEditor/Source/ProceduralVegetationEditor/Private/Visualizations/PVScaleVisualizationComponent.cpp
@@ -0,0 +1,145 @@
+// Copyright Epic Games, Inc. All Rights Reserved.
+
+#include "PVScaleVisualizationComponent.h"
+
+#include "Components/TextRenderComponent.h"
+
+#include "Engine/World.h"
+
+#include "GameFramework/WorldSettings.h"
+
+UPVScaleVisualizationComponent::UPVScaleVisualizationComponent()
+	: Bounds(ForceInit)
+{
+	bAutoActivate = true;
+	bTickInEditor = true;
+	PrimaryComponentTick.bCanEverTick = true;
+	bUseEditorCompositing = true;
+	SetGenerateOverlapEvents(false);
+	bIgnoreStreamingManagerUpdate = true;
+
+	int32 ComponentIndex = 0;
+	for (TObjectPtr<UTextRenderComponent>& Component : TextRenderComponents)
+	{
+		const FString ComponentName = TEXT("LengthTextRenderComponent") + FString::FromInt(ComponentIndex++);
+		Component = CreateDefaultSubobject<UTextRenderComponent>(*ComponentName);
+		Component->SetGenerateOverlapEvents(false);
+		Component->SetupAttachment(this);
+	}
+}
+
+FBoxSphereBounds& UPVScaleVisualizationComponent::GetScaleBounds()
+{
+	return Bounds;
+}
+
+const FBoxSphereBounds& UPVScaleVisualizationComponent::GetScaleBounds() const
+{
+	return Bounds;
+}
+
+void UPVScaleVisualizationComponent::SetScaleBounds(const FBoxSphereBounds& InBounds)
+{
+	Bounds = InBounds;
+	UpdateScaleVisualizations();
+}
+
+TObjectPtr<UTextRenderComponent> UPVScaleVisualizationComponent::GetTextRenderComponent(int32 Index) const
+{
+	return TextRenderComponents[Index];
+}
+
+const TStaticArray<TObjectPtr<UTextRenderComponent>, 3>& UPVScaleVisualizationComponent::GetTextRenderComponents() const
+{
+	return TextRenderComponents;
+}
+
+void UPVScaleVisualizationComponent::OnRegister()
+{
+	Super::OnRegister();
+	UpdateScaleVisualizations();
+}
+
+void UPVScaleVisualizationComponent::UpdateScaleVisualizations()
+{
+	if (!GetWorld())
+	{
+		return;
+	}
+
+	Flush();
+
+	const auto WorldToMetersScale = GetWorld()->GetWorldSettings()->WorldToMeters;
+	const double TextSize = (Bounds.BoxExtent.GetAbsMax() * 20.0f) / WorldToMetersScale;
+	{
+		FVector StartPosition = Bounds.Origin;
+		StartPosition.X -= Bounds.BoxExtent.X;
+		StartPosition.Y -= Bounds.BoxExtent.Y;
+		StartPosition.Z -= Bounds.BoxExtent.Z;
+
+		FVector EndPosition = Bounds.Origin;
+		EndPosition.X += Bounds.BoxExtent.X;
+		EndPosition.Y -= Bounds.BoxExtent.Y;
+		EndPosition.Z -= Bounds.BoxExtent.Z;
+
+		const double Length = (Bounds.BoxExtent.X * 2.0f) / WorldToMetersScale;
+		const TObjectPtr<UTextRenderComponent> TextRenderComponent = GetTextRenderComponent(0);
+
+		TextRenderComponent->SetWorldLocation(EndPosition);
+		TextRenderComponent->SetWorldRotation(FRotator(0, 90, 0));
+		TextRenderComponent->SetTextRenderColor(FLinearColor::Red.ToFColorSRGB());
+		TextRenderComponent->SetHorizontalAlignment(EHTA_Right);
+		TextRenderComponent->SetVerticalAlignment(EVRTA_TextBottom);
+		TextRenderComponent->SetWorldSize(TextSize);
+		TextRenderComponent->SetText(FText::FromString(FString::Printf(TEXT("%.2fm"), Length)));
+
+		AddLine(StartPosition, EndPosition, FLinearColor::Red);
+	}
+
+	{
+		FVector StartPosition = Bounds.Origin;
+		StartPosition.X -= Bounds.BoxExtent.X;
+		StartPosition.Y -= Bounds.BoxExtent.Y;
+		StartPosition.Z -= Bounds.BoxExtent.Z;
+
+		FVector EndPosition = Bounds.Origin;
+		EndPosition.X -= Bounds.BoxExtent.X;
+		EndPosition.Y += Bounds.BoxExtent.Y;
+		EndPosition.Z -= Bounds.BoxExtent.Z;
+
+		const double Length = (Bounds.BoxExtent.Y * 2.0f) / WorldToMetersScale;
+		const TObjectPtr<UTextRenderComponent> TextRenderComponent = GetTextRenderComponent(1);
+		TextRenderComponent->SetWorldLocation(EndPosition);
+		TextRenderComponent->SetTextRenderColor(FColor::Green);
+		TextRenderComponent->SetHorizontalAlignment(EHTA_Left);
+		TextRenderComponent->SetVerticalAlignment(EVRTA_TextBottom);
+		TextRenderComponent->SetWorldSize(TextSize);
+		TextRenderComponent->SetText(FText::FromString(FString::Printf(TEXT("%.2fm"), Length)));
+
+		AddLine(StartPosition, EndPosition, FLinearColor::Green);
+	}
+
+	{
+		FVector StartPosition = Bounds.Origin;
+		StartPosition.X -= Bounds.BoxExtent.X;
+		StartPosition.Y -= Bounds.BoxExtent.Y;
+		StartPosition.Z -= Bounds.BoxExtent.Z;
+
+		FVector EndPosition = Bounds.Origin;
+		EndPosition.X -= Bounds.BoxExtent.X;
+		EndPosition.Y -= Bounds.BoxExtent.Y;
+		EndPosition.Z += Bounds.BoxExtent.Z;
+
+		const double Length = (Bounds.BoxExtent.Z * 2.0f) / WorldToMetersScale;
+		const TObjectPtr<UTextRenderComponent> TextRenderComponent = GetTextRenderComponent(2);
+		TextRenderComponent->SetWorldLocation(EndPosition);
+		TextRenderComponent->SetWorldRotation(FRotator(0, 45, 0));
+		TextRenderComponent->SetTextRenderColor(FLinearColor::Blue.ToFColorSRGB());
+		TextRenderComponent->SetHorizontalAlignment(EHTA_Center);
+		TextRenderComponent->SetVerticalAlignment(EVRTA_TextBottom);
+		TextRenderComponent->SetWorldSize(TextSize);
+		TextRenderComponent->SetText(FText::FromString(FString::Printf(TEXT("%.2fm"), Length)));
+		
+		AddLine(StartPosition, EndPosition, FLinearColor::Blue);
+	}
+}
diff --git a/Engine/Plugins/Experimental/ProceduralVegetationEditor/Source/ProceduralVegetationEditor/Private/Visualizations/PVScaleVisualizationComponent.h b/Engine/Plugins/Experimental/ProceduralVegetationEditor/Source/ProceduralVegetationEditor/Private/Visualizations/PVScaleVisualizationComponent.h
new file mode 100644
--- /dev/null
+++ b/Engine/Plugins/Experimental/ProceduralVegetationEditor/Source/ProceduralVegetationEditor/Private/Visualizations/PVScaleVisualizationComponent.h
@@ -0,0 +1,37 @@
+// Copyright Epic Games, Inc. All Rights Reserved.
+
+#pragma once
+
+#include "CoreMinimal.h"
+#include "PVLineBatchComponent.h"
+
+#include "Containers/StaticArray.h"
+
+#include "PVScaleVisualizationComponent.generated.h"
+
+class UTextRenderComponent;
+
+UCLASS(Blueprintable, BlueprintType, meta=(BlueprintSpawnableComponent), MinimalAPI)
+class UPVScaleVisualizationComponent : public UPVLineBatchComponent
+{
+	GENERATED_BODY()
+
+public:
+	UPVScaleVisualizationComponent();
+
+	FBoxSphereBounds& GetScaleBounds();
+	const FBoxSphereBounds& GetScaleBounds() const;
+	void SetScaleBounds(const FBoxSphereBounds& InBounds);
+
+	TObjectPtr<UTextRenderComponent> GetTextRenderComponent(int32 Index) const;
+	const TStaticArray<TObjectPtr<UTextRenderComponent>, 3>& GetTextRenderComponents() const;
+
+	virtual void OnRegister() override;
+
+	void UpdateScaleVisualizations();
+
+private:
+	FBoxSphereBounds Bounds;
+
+	TStaticArray<TObjectPtr<UTextRenderComponent>, 3> TextRenderComponents = {nullptr, nullptr, nullptr};
+};
diff --git a/Engine/Plugins/Experimental/ProceduralVegetationEditor/Source/ProceduralVegetationEditor/Public/PVEditorSettings.h b/Engine/Plugins/Experimental/ProceduralVegetationEditor/Source/ProceduralVegetationEditor/Public/PVEditorSettings.h
--- a/Engine/Plugins/Experimental/ProceduralVegetationEditor/Source/ProceduralVegetationEditor/Public/PVEditorSettings.h
+++ b/Engine/Plugins/Experimental/ProceduralVegetationEditor/Source/ProceduralVegetationEditor/Public/PVEditorSettings.h
@@ -25,4 +25,10 @@ public:
 
 	UPROPERTY(EditAnywhere, config, Category = Debug)
 	bool bShowMannequin = false;
+
+	UPROPERTY(EditAnywhere, config, Category = Debug)
+	float MannequinOffset = 0.0f;
+
+	UPROPERTY(EditAnywhere, config, Category = Debug)
+	bool bShowScaleVisualization = false;
 };
