PVE Export optmization

- Parallelize parts of CombineMeshInstancesToDynamicMesh
- Replace FDynamicMeshEditor::AppendMesh with FDynamicMesh3::AppendWithOffsets as it's much faster
- Copy tangents from foliage mesh to avoid recalculating all tangents during mesh build

diff --git a/Engine/Plugins/Experimental/ProceduralVegetationEditor/Source/ProceduralVegetationEditor/Private/Helpers/PVExportHelper.cpp b/Engine/Plugins/Experimental/ProceduralVegetationEditor/Source/ProceduralVegetationEditor/Private/Helpers/PVExportHelper.cpp
--- a/Engine/Plugins/Experimental/ProceduralVegetationEditor/Source/ProceduralVegetationEditor/Private/Helpers/PVExportHelper.cpp
+++ b/Engine/Plugins/Experimental/ProceduralVegetationEditor/Source/ProceduralVegetationEditor/Private/Helpers/PVExportHelper.cpp
@@ -30,6 +30,7 @@
 
 #include "DynamicMesh/DynamicBoneAttribute.h"
 #include "DynamicMesh/DynamicVertexSkinWeightsAttribute.h"
+#include "DynamicMesh/MeshTransforms.h"
 
 #include "Engine/SkeletalMesh.h"
 #include "Engine/StaticMesh.h"
@@ -204,33 +205,40 @@ namespace PV::Export::Internal
 		TRACE_CPUPROFILER_EVENT_SCOPE(PV::Export::CollectionToDynamicMesh);
 		
 		TObjectPtr<UDynamicMesh> NewMesh = NewObject<UDynamicMesh>();
-
-		if (Collection.NumElements(FGeometryCollection::TransformGroup) > 0)
+		if (Collection.NumElements(FGeometryCollection::TransformGroup) == 0)
 		{
-			if (const TUniquePtr<FGeometryCollection> GeomCollection = TUniquePtr<FGeometryCollection>(Collection.NewCopy<FGeometryCollection>()))
-			{
-				const TManagedArray<FTransform3f>& BoneTransforms = Collection.GetAttribute<FTransform3f>(
-					"Transform", FGeometryCollection::TransformGroup);
+			return NewMesh;
+		}
 
-				TArray<int32> TransformIndices;
-				TransformIndices.AddUninitialized(BoneTransforms.Num());
+		const TUniquePtr<FGeometryCollection> GeomCollection = TUniquePtr<FGeometryCollection>(Collection.NewCopy<FGeometryCollection>());
+		if (!GeomCollection.IsValid())
+		{
+			return NewMesh;
+		}
 
-				int32 Idx = 0;
-				for (int32& TransformIdx : TransformIndices)
-				{
-					TransformIdx = Idx++;
-				}
+		const TManagedArray<FTransform3f>& BoneTransforms = Collection.GetAttribute<FTransform3f>(
+			"Transform", FGeometryCollection::TransformGroup);
 
-				constexpr bool bCenterPivot = false;
-				FTransform TransformOut;
-				
-				FDynamicMesh3 CombinedMesh;
-				ConvertGeometryCollectionToDynamicMesh(CombinedMesh, TransformOut, bCenterPivot, *GeomCollection, true, BoneTransforms.GetConstArray(), true, TransformIndices);
+		TArray<int32> TransformIndices;
+		TransformIndices.AddUninitialized(BoneTransforms.Num());
 
-				NewMesh->SetMesh(CombinedMesh);
-			}
+		int32 Idx = 0;
+		for (int32& TransformIdx : TransformIndices)
+		{
+			TransformIdx = Idx++;
 		}
 
+		FMeshDescription MeshDescription;
+		FStaticMeshAttributes Attributes(MeshDescription);
+		Attributes.Register();
+
+		constexpr bool bCenterPivot = false;
+		FTransform TransformOut;
+		FDynamicMesh3 CombinedMesh;
+		ConvertGeometryCollectionToDynamicMesh(CombinedMesh, TransformOut, bCenterPivot, *GeomCollection, true, BoneTransforms.GetConstArray(), true, TransformIndices);
+
+		NewMesh->SetMesh(CombinedMesh);
+
 		return NewMesh;
 	}
 
@@ -239,7 +247,7 @@ namespace PV::Export::Internal
 		FGeometryScriptCopyMeshToAssetOptions Options;
 		Options.bEmitTransaction = false;
 		Options.bDeferMeshPostEditChange = true;
-		Options.bEnableRecomputeTangents = true;
+		Options.bEnableRecomputeTangents = false;
 		Options.bReplaceMaterials = true;
 		Options.NewMaterials = InMaterials;
 
@@ -286,7 +294,7 @@ namespace PV::Export::Internal
 			}
 
 			FMeshDescriptionToDynamicMesh Converter;
-			Converter.Convert(SourceMesh, ToDynamicMesh);
+			Converter.Convert(SourceMesh, ToDynamicMesh, true);
 			return true;
 		}
 
@@ -312,152 +320,213 @@ namespace PV::Export::Internal
 			return true;
 		}
 
-		TMap<FString, FDynamicMesh3> FoliageDynamicMeshes;
-		TMap<FString, TObjectPtr<UStaticMesh>> FoliageStaticMeshes;
-		TMap<FString, TObjectPtr<USkeletalMesh>> FoliageSkeletalMeshes;
-		
-		FoliageDynamicMeshes.Reserve(NumFoliageNames);
-		FoliageStaticMeshes.Reserve(NumFoliageNames);
-		FoliageSkeletalMeshes.Reserve(NumFoliageNames);
+		struct FConvertFoliageToDynamicMeshResult
+		{
+			FString FoliageName;
+			FDynamicMesh3 OutDynamicMesh;
+			TArray<UMaterialInterface*> OutMaterials;
+			FString OutMeshPath;
+			bool bSuccess = true;
+		};
 
-		for (int32 FoliageNameIndex = 0; FoliageNameIndex < NumFoliageNames; ++FoliageNameIndex)
+		TArray<FConvertFoliageToDynamicMeshResult> ConvertFoliageToDynamicMeshResults;
+		TMap<FString, int32> FoliageNameToDynamicMeshResult;
 		{
-			const FString FoliageName = FoliageFacade.GetFoliageName(FoliageNameIndex);
-			if (FoliageDynamicMeshes.Contains(FoliageName))
+			TRACE_CPUPROFILER_EVENT_SCOPE(ConvertFoliageToDynamicMeshes);
+
+			ConvertFoliageToDynamicMeshResults.Reserve(NumFoliageNames);
+			FoliageNameToDynamicMeshResult.Reserve(NumFoliageNames);
+
+			for (int32 FoliageNameIndex = 0; FoliageNameIndex < NumFoliageNames; ++FoliageNameIndex)
 			{
-				UE_LOG(LogProceduralVegetationEditor, Warning, TEXT("Duplicate foliage names in foliage facade"));
-				continue;
+				const FString FoliageName = FoliageFacade.GetFoliageName(FoliageNameIndex);
+
+				if (FoliageNameToDynamicMeshResult.Contains(FoliageName))
+				{
+					UE_LOG(LogProceduralVegetationEditor, Warning, TEXT("Duplicate foliage name %s in foliage facade"), *FoliageName);
+					continue;
+				}
+
+				FConvertFoliageToDynamicMeshResult& Result = ConvertFoliageToDynamicMeshResults.AddDefaulted_GetRef();
+				Result.FoliageName = FoliageName;
+
+				FoliageNameToDynamicMeshResult.Add(FoliageName, ConvertFoliageToDynamicMeshResults.Num() - 1);
 			}
 
-			if (const TObjectPtr<UStaticMesh> StaticMesh = LoadObject<UStaticMesh>(nullptr, *FoliageName))
+			ParallelFor(ConvertFoliageToDynamicMeshResults.Num(), [&ConvertFoliageToDynamicMeshResults](int32 Index)
 			{
-				FoliageStaticMeshes.Add(FoliageName, StaticMesh);
-					
-				for (auto Material : StaticMesh->GetStaticMaterials())
+				FConvertFoliageToDynamicMeshResult& Result = ConvertFoliageToDynamicMeshResults[Index];
+
+				const FString& FoliageName = Result.FoliageName;
+
+				if (const TObjectPtr<UStaticMesh> StaticMesh = LoadObject<UStaticMesh>(nullptr, *FoliageName))
 				{
-					Materials.AddUnique(Material.MaterialInterface);
+					FText OutErrorMessage;
+					UE::Conversion::FStaticMeshConversionOptions Options;
+					if (!StaticMeshToDynamicMesh(StaticMesh, Result.OutDynamicMesh, OutErrorMessage, Options))
+					{
+						Result.bSuccess = false;
+						return;
+					}
+
+					for (auto Material : StaticMesh->GetStaticMaterials())
+					{
+						Result.OutMaterials.Add(Material.MaterialInterface);
+					}
+
+					Result.OutMeshPath = StaticMesh->GetFullName();
 				}
+				else if (const TObjectPtr<USkeletalMesh> SkeletalMesh = LoadObject<USkeletalMesh>(nullptr, *FoliageName))
+				{
+					if (!SkeletalMeshToDynamicMesh(SkeletalMesh, 0, Result.OutDynamicMesh))
+					{
+						Result.bSuccess = false;
+						return;
+					}
 
-				FDynamicMesh3 FoliageDynamicMesh;
-				FText OutErrorMessage;
-				UE::Conversion::FStaticMeshConversionOptions Options;
-				if (StaticMeshToDynamicMesh(StaticMesh, FoliageDynamicMesh, OutErrorMessage, Options))
+					for (auto Material : SkeletalMesh->GetMaterials())
+					{
+						Result.OutMaterials.Add(Material.MaterialInterface);
+					}
+
+					Result.OutDynamicMesh.Attributes()->RemoveAttribute("BoneIndex");
+					Result.OutDynamicMesh.Attributes()->RemoveAttribute("BoneParentIndex");
+					Result.OutMeshPath = SkeletalMesh->GetFullName();
+				}
+				else
 				{
-					FoliageDynamicMeshes.Add(FoliageName, MoveTemp(FoliageDynamicMesh));
+					Result.bSuccess = false;
 				}
-			}
-			else if (const TObjectPtr<USkeletalMesh> SkeletalMesh = LoadObject<USkeletalMesh>(nullptr, *FoliageName))
+			});
+
 			{
-				FoliageSkeletalMeshes.Add(FoliageName, SkeletalMesh);
-					
-				for (auto Material : SkeletalMesh->GetMaterials())
+				TRACE_CPUPROFILER_EVENT_SCOPE(RemapMaterials);
+
+				TSet<UMaterialInterface*> UniqueMaterials;
+				UniqueMaterials.Reserve(ConvertFoliageToDynamicMeshResults.Num() * 2);
+				for (const FConvertFoliageToDynamicMeshResult& Result : ConvertFoliageToDynamicMeshResults)
 				{
-					Materials.AddUnique(Material.MaterialInterface);
+					UniqueMaterials.Append(Result.OutMaterials);
 				}
 
-				FDynamicMesh3 FoliageDynamicMesh;
-				if (SkeletalMeshToDynamicMesh(SkeletalMesh, 0, FoliageDynamicMesh))
+				for (UMaterialInterface* Material : UniqueMaterials)
 				{
-					FoliageDynamicMesh.Attributes()->RemoveAttribute("BoneIndex");
-					FoliageDynamicMesh.Attributes()->RemoveAttribute("BoneParentIndex");
-					FoliageDynamicMeshes.Add(FoliageName, MoveTemp(FoliageDynamicMesh));
+					Materials.AddUnique(Material);
 				}
-			}
 
-			if (!OnProgressUpdated((FoliageNameIndex / static_cast<float>(NumFoliageNames)) * 0.5f))
-			{
-				return false;
-			}
-		}
+				for (FConvertFoliageToDynamicMeshResult& Result : ConvertFoliageToDynamicMeshResults)
+				{
+					UE::Geometry::FDynamicMeshMaterialAttribute* MaterialIDs = Result.OutDynamicMesh.Attributes()->GetMaterialID();
+					if (!MaterialIDs)
+					{
+						continue;
+					}
+					
+					for (int32 Tid = 0; Tid < Result.OutDynamicMesh.GetTrianglesBuffer().Num(); Tid++)
+					{
+						const int32 MaterialID = MaterialIDs->GetValue(Tid);
+						if (MaterialID == INDEX_NONE)
+						{
+							continue;
+						}
 
-		for (int32 Id = 0; Id < NumInstances; Id++)
-		{
-			const FString FoliageName = GetFoliageName(FoliageFacade, Id);
+						check(Result.OutMaterials.IsValidIndex(MaterialID));
 
-			FDynamicMesh3* FoliageDynamicMesh = FoliageDynamicMeshes.Find(FoliageName);
-			if (FoliageDynamicMesh == nullptr)
-			{
-				continue;
+						const int32 NewMaterialID = Materials.Find(Result.OutMaterials[MaterialID]);
+						check(NewMaterialID != INDEX_NONE);
+
+						MaterialIDs->SetValue(Tid, NewMaterialID);
+					}
+				}
 			}
-			
-			const TObjectPtr<UStaticMesh> StaticMesh = FoliageStaticMeshes.FindRef(FoliageName, nullptr);
-			const TObjectPtr<USkeletalMesh> SkeletalMesh = FoliageSkeletalMeshes.FindRef(FoliageName, nullptr);
+		}
 
-			check(StaticMesh || SkeletalMesh);
+		const float ConvertToDynamicMeshProgress = 0.1f;
+		if (!OnProgressUpdated(ConvertToDynamicMeshProgress))
+		{
+			return false;
+		}
 
-			if (SkeletalMesh != nullptr)
-			{
-				const int32 FoliageParentBoneID = FoliageFacade.GetParentBoneID(Id);
+		struct FGenerateDynamicMeshesResult
+		{
+			FDynamicMesh3 DynamicMesh;
+			int32 ConvertToDynamicMeshResultIndex;
+		};
 
-				// Reparent root bone according to the instanced transform
-				UE::Geometry::FDynamicMeshAttributeSet* MeshAttributes = FoliageDynamicMesh->Attributes();
-				UE::Geometry::FDynamicMeshBoneParentIndexAttribute* BoneParentIndices = MeshAttributes->GetBoneParentIndices();
-				BoneParentIndices->SetValue(0, FoliageParentBoneID);
-			}
-				
-			const FTransform Transform = FoliageFacade.GetFoliageTransform(Id);
-			UE::Geometry::FTransformSRT3d GeoTransform(Transform);
+		TArray<FGenerateDynamicMeshesResult> GenerateDynamicMeshesResults;
+		{
+			TRACE_CPUPROFILER_EVENT_SCOPE(GenerateDynamicMeshInstances);
 
-			DynamicMesh->EditMesh([&](FDynamicMesh3& Mesh)
+			GenerateDynamicMeshesResults.SetNum(NumInstances);
+			ParallelFor(NumInstances, [
+				&GenerateDynamicMeshesResults,
+				&FoliageFacade, 
+				&FoliageNameToDynamicMeshResult,
+				&ConvertFoliageToDynamicMeshResults
+			](int32 Index)
 			{
-				const int32 VertexStart = Mesh.VertexCount();
-				UE::Geometry::FDynamicMeshEditor Editor = UE::Geometry::FDynamicMeshEditor(&Mesh);
-
-				UE::Geometry::FMeshIndexMappings IndexMap;
-				Editor.AppendMesh(
-					FoliageDynamicMesh,
-					IndexMap,
-					[&](int32, const FVector3d& Position)
-					{
-						return GeoTransform.TransformPosition(Position);
-					},
-					[&](int32, const FVector3d& Normal)
-					{
-						return GeoTransform.TransformVectorNoScale(Normal);
-					}
-				);
+				const FString FoliageName = GetFoliageName(FoliageFacade, Index);
 
-				// Remap Material IDs
-				if (Mesh.Attributes()->HasMaterialID())
+				const int32* PrevResultIndex = FoliageNameToDynamicMeshResult.Find(FoliageName);
+				if (PrevResultIndex == nullptr)
 				{
-					const UE::Geometry::FIndexMapi& TriangleMap = IndexMap.GetTriangleMap();
-					UE::Geometry::FDynamicMeshMaterialAttribute* MaterialIDs = Mesh.Attributes()->GetMaterialID();
-
-					for (const auto& [OldTid, NewTid] : TriangleMap.GetForwardMap())
-					{
-						const int32 MaterialID = MaterialIDs->GetValue(OldTid);
-						if (StaticMesh)
-						{
-							UMaterialInterface* Material = StaticMesh->GetMaterial(MaterialID);
-							const int32 NewMaterialID = Materials.Find(Material);
-							MaterialIDs->SetValue(NewTid, NewMaterialID);
-						}
-						else if (SkeletalMesh)
-						{
-							UMaterialInterface* Material = SkeletalMesh->GetMaterials()[MaterialID].MaterialInterface;
-							const int32 NewMaterialID = Materials.Find(Material);
-							MaterialIDs->SetValue(NewTid, NewMaterialID);
-						}
-					}
+					return;
 				}
 
-				if (StaticMesh)
+				FConvertFoliageToDynamicMeshResult& ConvertFoliageToDynamicMeshResult = ConvertFoliageToDynamicMeshResults[*PrevResultIndex];
+				if (!ConvertFoliageToDynamicMeshResult.bSuccess)
 				{
-					OnMeshInstanceCombined(StaticMesh->GetFullName(), Id, VertexStart, Mesh.VertexCount());
+					return;
 				}
-				else if (SkeletalMesh != nullptr)
+
+				FGenerateDynamicMeshesResult& Result = GenerateDynamicMeshesResults[Index];
+				Result.DynamicMesh = ConvertFoliageToDynamicMeshResult.OutDynamicMesh;
+				Result.ConvertToDynamicMeshResultIndex = *PrevResultIndex;
+
+				const int32 FoliageParentBoneID = FoliageFacade.GetParentBoneID(Index);
+				UE::Geometry::FDynamicMeshAttributeSet* MeshAttributes = Result.DynamicMesh.Attributes();
+				UE::Geometry::FDynamicMeshBoneParentIndexAttribute* BoneParentIndices = MeshAttributes->GetBoneParentIndices();
+				if (FoliageParentBoneID != INDEX_NONE && BoneParentIndices != nullptr)
 				{
-					OnMeshInstanceCombined(SkeletalMesh->GetFullName(), Id, VertexStart, Mesh.VertexCount());
+					// Reparent root bone according to the instanced transform
+					BoneParentIndices->SetValue(0, FoliageParentBoneID);
 				}
+
+				const FTransform Transform = FoliageFacade.GetFoliageTransform(Index);
+				MeshTransforms::ApplyTransform(Result.DynamicMesh, Transform, true);
 			});
+		}
+
+		bool bWasCanceled = false;
+		DynamicMesh->EditMesh([&](FDynamicMesh3& Mesh)
+		{
+			TRACE_CPUPROFILER_EVENT_SCOPE(AppendDynamicMeshes);
 
-			if (!OnProgressUpdated(0.5f + (Id / static_cast<float>(NumInstances)) * 0.5f))
+			for (int32 i = 0; i < GenerateDynamicMeshesResults.Num(); ++i)
 			{
-				return false;
+				const FGenerateDynamicMeshesResult& GenerateDynamicMeshesResult = GenerateDynamicMeshesResults[i];
+				const FDynamicMesh3& MeshToAppend = GenerateDynamicMeshesResult.DynamicMesh;
+				if (MeshToAppend.VertexCount() == 0)
+				{
+					continue;
+				}
+
+				const int32 VertexStart = Mesh.VertexCount();
+				Mesh.AppendWithOffsets(MeshToAppend);
+				
+				const FConvertFoliageToDynamicMeshResult& ConvertFoliageToDynamicMeshResult = ConvertFoliageToDynamicMeshResults[GenerateDynamicMeshesResult.ConvertToDynamicMeshResultIndex];
+				OnMeshInstanceCombined(ConvertFoliageToDynamicMeshResult.OutMeshPath, i, VertexStart, Mesh.VertexCount());
+
+				const float CombineMeshesProgress = ConvertToDynamicMeshProgress + i / (float)NumInstances * (1.f - ConvertToDynamicMeshProgress);
+				if (!OnProgressUpdated(CombineMeshesProgress))
+				{
+					bWasCanceled = true;
+					return;
+				}
 			}
-		}
+		});
 
-		return true;
+		return !bWasCanceled;
 	}
 
 	bool ContainsFoliageName(const Facades::FFoliageFacade& FoliageFacade, const FString& InFoliageName)
@@ -926,7 +995,7 @@ namespace PV::Export::Internal
 		}
 
 		{
-			TRACE_CPUPROFILER_EVENT_SCOPE(BuildSkinnedMesh);
+			TRACE_CPUPROFILER_EVENT_SCOPE(BuildStaticMesh);
 			if (!UpdateStage(EExportMeshStage::BuildingMesh, 1.f))
 			{
 				return false;
