[PVE] BugHawk Fixes for CL 46203729

diff --git a/Engine/Plugins/Experimental/ProceduralVegetationEditor/Source/ProceduralVegetationEditor/Private/Customizations/PVFloatRampCustomization.cpp b/Engine/Plugins/Experimental/ProceduralVegetationEditor/Source/ProceduralVegetationEditor/Private/Customizations/PVFloatRampCustomization.cpp
--- a/Engine/Plugins/Experimental/ProceduralVegetationEditor/Source/ProceduralVegetationEditor/Private/Customizations/PVFloatRampCustomization.cpp
+++ b/Engine/Plugins/Experimental/ProceduralVegetationEditor/Source/ProceduralVegetationEditor/Private/Customizations/PVFloatRampCustomization.cpp
@@ -250,8 +250,11 @@ void FPVFloatRampCustomization::CustomizeHeader(
 					if (YAxisMin != -FLT_MAX || YAxisMax != FLT_MAX)
 						ClampKeyValue(Curve, *KeyIter, YAxisMin, YAxisMax);
 				}
-
-				FloatRampOwner->GetOwner()->PostEditChange();
+				
+				if (UObject* const Owner = FloatRampOwner->GetOwner()) 
+				{
+					Owner->PostEditChange();
+				}
 			}
 		);
 
diff --git a/Engine/Plugins/Experimental/ProceduralVegetationEditor/Source/ProceduralVegetationEditor/Private/SPVEditorViewport.cpp b/Engine/Plugins/Experimental/ProceduralVegetationEditor/Source/ProceduralVegetationEditor/Private/SPVEditorViewport.cpp
--- a/Engine/Plugins/Experimental/ProceduralVegetationEditor/Source/ProceduralVegetationEditor/Private/SPVEditorViewport.cpp
+++ b/Engine/Plugins/Experimental/ProceduralVegetationEditor/Source/ProceduralVegetationEditor/Private/SPVEditorViewport.cpp
@@ -227,7 +227,7 @@ bool SPVEditorViewport::IsShowScaleVisChecked() const
 
 void SPVEditorViewport::SetScaleVisState(bool InEnable)
 {
-	if (MannequinComponent)
+	if (ScaleVisualizationComponent)
 	{
 		ScaleVisualizationComponent->SetVisibility(InEnable, true);
 		Client->Invalidate();
diff --git a/Engine/Plugins/Experimental/ProceduralVegetationEditor/Source/ProceduralVegetationEditor/Private/Visualizations/PVLineBatchComponent.cpp b/Engine/Plugins/Experimental/ProceduralVegetationEditor/Source/ProceduralVegetationEditor/Private/Visualizations/PVLineBatchComponent.cpp
--- a/Engine/Plugins/Experimental/ProceduralVegetationEditor/Source/ProceduralVegetationEditor/Private/Visualizations/PVLineBatchComponent.cpp
+++ b/Engine/Plugins/Experimental/ProceduralVegetationEditor/Source/ProceduralVegetationEditor/Private/Visualizations/PVLineBatchComponent.cpp
@@ -100,6 +100,7 @@ void UPVLineBatchComponent::Draw(const FSceneView* View, FPrimitiveDrawInterface
 void UPVLineBatchComponent::Flush()
 {
 	Lines.Empty();
+	BBox = FBox(ForceInit);
 }
 
 FPrimitiveSceneProxy* UPVLineBatchComponent::CreateSceneProxy()
