[PVE] Added viewport visualization mode

diff --git a/Engine/Plugins/Experimental/ProceduralVegetationEditor/Source/ProceduralVegetation/Private/DataTypes/PVData.cpp b/Engine/Plugins/Experimental/ProceduralVegetationEditor/Source/ProceduralVegetation/Private/DataTypes/PVData.cpp
--- a/Engine/Plugins/Experimental/ProceduralVegetationEditor/Source/ProceduralVegetation/Private/DataTypes/PVData.cpp
+++ b/Engine/Plugins/Experimental/ProceduralVegetationEditor/Source/ProceduralVegetation/Private/DataTypes/PVData.cpp
@@ -32,7 +32,6 @@ void UPVData::AddToCrc(FArchiveCrc32& Ar, bool bFullDataCrc) const
 		Ar << ClassName;
 		Ar << Collection;
 #if WITH_EDITOR
-		Ar << RenderType;
 		Ar << DebugSettings;
 #endif
 	}
@@ -70,7 +69,6 @@ UPCGSpatialData* UPVData::CopyInternal(FPCGContext* Context) const
 	NewProceduralVegetationData->Initialize(MoveTemp(NewCollection), /*bCanTakeOwnership=*/false);
 
 #if WITH_EDITOR
-	NewProceduralVegetationData->RenderType = RenderType;
 	NewProceduralVegetationData->DebugSettings = DebugSettings;
 #endif
 	
diff --git a/Engine/Plugins/Experimental/ProceduralVegetationEditor/Source/ProceduralVegetation/Private/Nodes/PVBaseSettings.cpp b/Engine/Plugins/Experimental/ProceduralVegetationEditor/Source/ProceduralVegetation/Private/Nodes/PVBaseSettings.cpp
--- a/Engine/Plugins/Experimental/ProceduralVegetationEditor/Source/ProceduralVegetation/Private/Nodes/PVBaseSettings.cpp
+++ b/Engine/Plugins/Experimental/ProceduralVegetationEditor/Source/ProceduralVegetation/Private/Nodes/PVBaseSettings.cpp
@@ -6,6 +6,23 @@
 #include "DataTypes/PVData.h"
 #include "PCGData.h"
 
+UPVBaseSettings::UPVBaseSettings()
+{
+	if (PCGHelpers::IsNewObjectAndNotDefault(this))
+	{
+		
+	}
+}
+
+#if WITH_EDITOR
+void UPVBaseSettings::SetCurrentRenderType(TArray<EPVRenderType> InRenderTypes)
+{
+	CurrentRenderType = InRenderTypes;
+
+	PostEditChange();
+}
+#endif
+
 TArray<FPCGPinProperties> UPVBaseSettings::InputPinProperties() const
 {
 	TArray<FPCGPinProperties> Properties;
@@ -58,7 +75,6 @@ void FPVBaseElement::PostExecuteInternal(FPCGContext* InContext) const
 		{
 			if (const UPVData* ProceduralVegetationData = Cast<UPVData>(OutputData.Data))
 			{
-				ProceduralVegetationData->SetRenderType(Settings->GetRenderType());
 				ProceduralVegetationData->SetDebugSettings(Settings->GetDebugVisualizationSettings());
 			}
 		}
diff --git a/Engine/Plugins/Experimental/ProceduralVegetationEditor/Source/ProceduralVegetation/Private/Nodes/PVBoneReductionSettings.h b/Engine/Plugins/Experimental/ProceduralVegetationEditor/Source/ProceduralVegetation/Private/Nodes/PVBoneReductionSettings.h
--- a/Engine/Plugins/Experimental/ProceduralVegetationEditor/Source/ProceduralVegetation/Private/Nodes/PVBoneReductionSettings.h
+++ b/Engine/Plugins/Experimental/ProceduralVegetationEditor/Source/ProceduralVegetation/Private/Nodes/PVBoneReductionSettings.h
@@ -19,7 +19,7 @@ public:
 	virtual FText GetNodeTooltipText() const override;
 	virtual FLinearColor GetNodeTitleColor() const override { return FLinearColor(1.0f, 0.0f, 0.599f); }
 
-	virtual TArray<EPVRenderType> GetRenderType() const override { return TArray{EPVRenderType::Bones, EPVRenderType::Mesh}; }
+	virtual TArray<EPVRenderType> GetDefaultRenderType() const override { return TArray{EPVRenderType::Bones, EPVRenderType::Mesh}; }
 #endif
 
 protected:
diff --git a/Engine/Plugins/Experimental/ProceduralVegetationEditor/Source/ProceduralVegetation/Private/Nodes/PVCarveSettings.h b/Engine/Plugins/Experimental/ProceduralVegetationEditor/Source/ProceduralVegetation/Private/Nodes/PVCarveSettings.h
--- a/Engine/Plugins/Experimental/ProceduralVegetationEditor/Source/ProceduralVegetation/Private/Nodes/PVCarveSettings.h
+++ b/Engine/Plugins/Experimental/ProceduralVegetationEditor/Source/ProceduralVegetation/Private/Nodes/PVCarveSettings.h
@@ -21,7 +21,7 @@ public:
 	virtual FText GetNodeTooltipText() const override;
 	virtual FLinearColor GetNodeTitleColor() const override { return FLinearColor::Yellow; }
 
-	virtual TArray<EPVRenderType> GetRenderType() const override { return TArray{EPVRenderType::Skeleton}; }
+	virtual TArray<EPVRenderType> GetDefaultRenderType() const override { return TArray{EPVRenderType::Skeleton}; }
 #endif
 
 protected:
diff --git a/Engine/Plugins/Experimental/ProceduralVegetationEditor/Source/ProceduralVegetation/Private/Nodes/PVFoliageDistributorSettings.h b/Engine/Plugins/Experimental/ProceduralVegetationEditor/Source/ProceduralVegetation/Private/Nodes/PVFoliageDistributorSettings.h
--- a/Engine/Plugins/Experimental/ProceduralVegetationEditor/Source/ProceduralVegetation/Private/Nodes/PVFoliageDistributorSettings.h
+++ b/Engine/Plugins/Experimental/ProceduralVegetationEditor/Source/ProceduralVegetation/Private/Nodes/PVFoliageDistributorSettings.h
@@ -28,7 +28,7 @@ public:
 	virtual FText GetNodeTooltipText() const override;
 	virtual FLinearColor GetNodeTitleColor() const override { return FLinearColor::Blue; }
 
-	virtual TArray<EPVRenderType> GetRenderType() const override { return TArray{ EPVRenderType::Mesh, EPVRenderType::Foliage }; }
+	virtual TArray<EPVRenderType> GetDefaultRenderType() const override { return TArray{ EPVRenderType::Mesh, EPVRenderType::Foliage }; }
 	#endif
 
 protected:
diff --git a/Engine/Plugins/Experimental/ProceduralVegetationEditor/Source/ProceduralVegetation/Private/Nodes/PVFoliagePaletteSettings.h b/Engine/Plugins/Experimental/ProceduralVegetationEditor/Source/ProceduralVegetation/Private/Nodes/PVFoliagePaletteSettings.h
--- a/Engine/Plugins/Experimental/ProceduralVegetationEditor/Source/ProceduralVegetation/Private/Nodes/PVFoliagePaletteSettings.h
+++ b/Engine/Plugins/Experimental/ProceduralVegetationEditor/Source/ProceduralVegetation/Private/Nodes/PVFoliagePaletteSettings.h
@@ -18,8 +18,9 @@ public:
 	virtual FText GetNodeTooltipText() const override;
 	virtual FLinearColor GetNodeTitleColor() const override { return FLinearColor::Green; }
 
-	virtual TArray<EPVRenderType> GetRenderType() const override { return TArray{EPVRenderType::FoliageGrid}; }
-
+	virtual TArray<EPVRenderType> GetDefaultRenderType() const override { return TArray{EPVRenderType::FoliageGrid}; }
+	virtual TArray<EPVRenderType> GetSupportedRenderTypes() const override { return TArray{EPVRenderType::FoliageGrid, EPVRenderType::Foliage, EPVRenderType::Mesh}; }
+	
 	virtual void PostEditChangeProperty(FPropertyChangedEvent& PropertyChangedEvent) override;
 #endif
 
diff --git a/Engine/Plugins/Experimental/ProceduralVegetationEditor/Source/ProceduralVegetation/Private/Nodes/PVGravitySettings.h b/Engine/Plugins/Experimental/ProceduralVegetationEditor/Source/ProceduralVegetation/Private/Nodes/PVGravitySettings.h
--- a/Engine/Plugins/Experimental/ProceduralVegetationEditor/Source/ProceduralVegetation/Private/Nodes/PVGravitySettings.h
+++ b/Engine/Plugins/Experimental/ProceduralVegetationEditor/Source/ProceduralVegetation/Private/Nodes/PVGravitySettings.h
@@ -18,7 +18,7 @@ public:
 	virtual FText GetNodeTooltipText() const override;
 	virtual FLinearColor GetNodeTitleColor() const override { return FLinearColor::Yellow; }
 	
-	virtual TArray<EPVRenderType> GetRenderType() const override { return TArray{ EPVRenderType::Skeleton }; }
+	virtual TArray<EPVRenderType> GetDefaultRenderType() const override { return TArray{ EPVRenderType::Skeleton }; }
 #endif
 
 protected:
diff --git a/Engine/Plugins/Experimental/ProceduralVegetationEditor/Source/ProceduralVegetation/Private/Nodes/PVImporterSettings.h b/Engine/Plugins/Experimental/ProceduralVegetationEditor/Source/ProceduralVegetation/Private/Nodes/PVImporterSettings.h
--- a/Engine/Plugins/Experimental/ProceduralVegetationEditor/Source/ProceduralVegetation/Private/Nodes/PVImporterSettings.h
+++ b/Engine/Plugins/Experimental/ProceduralVegetationEditor/Source/ProceduralVegetation/Private/Nodes/PVImporterSettings.h
@@ -21,7 +21,7 @@ public:
 	virtual FText GetNodeTooltipText() const override { return NSLOCTEXT("PVImporterSettings", "NodeTooltip", "Create Procedural Vegetation Skeleton"); }
 	virtual FLinearColor GetNodeTitleColor() const override { return FLinearColor(FColor::Black); }
 
-	virtual TArray<EPVRenderType> GetRenderType() const override { return TArray{ EPVRenderType::Skeleton }; }
+	virtual TArray<EPVRenderType> GetDefaultRenderType() const override { return TArray{ EPVRenderType::Skeleton }; }
 #endif
 
 protected:
diff --git a/Engine/Plugins/Experimental/ProceduralVegetationEditor/Source/ProceduralVegetation/Private/Nodes/PVMeshBuilderSettings.h b/Engine/Plugins/Experimental/ProceduralVegetationEditor/Source/ProceduralVegetation/Private/Nodes/PVMeshBuilderSettings.h
--- a/Engine/Plugins/Experimental/ProceduralVegetationEditor/Source/ProceduralVegetation/Private/Nodes/PVMeshBuilderSettings.h
+++ b/Engine/Plugins/Experimental/ProceduralVegetationEditor/Source/ProceduralVegetation/Private/Nodes/PVMeshBuilderSettings.h
@@ -20,7 +20,8 @@ public:
 	virtual FText GetDefaultNodeTitle() const override;
 	virtual FText GetNodeTooltipText() const override;
 	virtual FLinearColor GetNodeTitleColor() const override { return FLinearColor(FColor::Cyan); }
-	virtual TArray<EPVRenderType> GetRenderType() const override { return TArray{EPVRenderType::Mesh}; }
+	virtual TArray<EPVRenderType> GetDefaultRenderType() const override { return TArray{EPVRenderType::Mesh}; }
+	virtual TArray<EPVRenderType> GetSupportedRenderTypes() const override { return TArray{EPVRenderType::Skeleton, EPVRenderType::Mesh}; }
 
 	virtual void PostEditChangeProperty(FPropertyChangedEvent& PropertyChangedEvent) override;
 
diff --git a/Engine/Plugins/Experimental/ProceduralVegetationEditor/Source/ProceduralVegetation/Private/Nodes/PVPresetLoaderSettings.h b/Engine/Plugins/Experimental/ProceduralVegetationEditor/Source/ProceduralVegetation/Private/Nodes/PVPresetLoaderSettings.h
--- a/Engine/Plugins/Experimental/ProceduralVegetationEditor/Source/ProceduralVegetation/Private/Nodes/PVPresetLoaderSettings.h
+++ b/Engine/Plugins/Experimental/ProceduralVegetationEditor/Source/ProceduralVegetation/Private/Nodes/PVPresetLoaderSettings.h
@@ -17,7 +17,7 @@ public:
 	virtual FText GetDefaultNodeTitle() const override;
 	virtual FText GetNodeTooltipText() const override;
 	virtual EPCGSettingsType GetType() const override { return EPCGSettingsType::InputOutput; }
-	virtual TArray<EPVRenderType> GetRenderType() const override { return TArray{ EPVRenderType::Skeleton }; }
+	virtual TArray<EPVRenderType> GetDefaultRenderType() const override { return TArray{ EPVRenderType::Skeleton }; }
 #endif
 
 protected:
diff --git a/Engine/Plugins/Experimental/ProceduralVegetationEditor/Source/ProceduralVegetation/Private/Nodes/PVRemoveBranchesSettings.h b/Engine/Plugins/Experimental/ProceduralVegetationEditor/Source/ProceduralVegetation/Private/Nodes/PVRemoveBranchesSettings.h
--- a/Engine/Plugins/Experimental/ProceduralVegetationEditor/Source/ProceduralVegetation/Private/Nodes/PVRemoveBranchesSettings.h
+++ b/Engine/Plugins/Experimental/ProceduralVegetationEditor/Source/ProceduralVegetation/Private/Nodes/PVRemoveBranchesSettings.h
@@ -18,7 +18,7 @@ public:
 	virtual FText GetNodeTooltipText() const override;
 	virtual FLinearColor GetNodeTitleColor() const override { return FLinearColor::Yellow; }
 
-	virtual TArray<EPVRenderType> GetRenderType() const override { return TArray{ EPVRenderType::Skeleton }; }
+	virtual TArray<EPVRenderType> GetDefaultRenderType() const override { return TArray{ EPVRenderType::Skeleton }; }
 #endif
 
 protected:
diff --git a/Engine/Plugins/Experimental/ProceduralVegetationEditor/Source/ProceduralVegetation/Private/Nodes/PVScaleSettings.h b/Engine/Plugins/Experimental/ProceduralVegetationEditor/Source/ProceduralVegetation/Private/Nodes/PVScaleSettings.h
--- a/Engine/Plugins/Experimental/ProceduralVegetationEditor/Source/ProceduralVegetation/Private/Nodes/PVScaleSettings.h
+++ b/Engine/Plugins/Experimental/ProceduralVegetationEditor/Source/ProceduralVegetation/Private/Nodes/PVScaleSettings.h
@@ -17,7 +17,7 @@ public:
 	virtual FText GetNodeTooltipText() const override;
 	virtual FLinearColor GetNodeTitleColor() const override { return FLinearColor::Yellow; }
 
-	virtual TArray<EPVRenderType> GetRenderType() const override { return TArray{ EPVRenderType::Skeleton }; }
+	virtual TArray<EPVRenderType> GetDefaultRenderType() const override { return TArray{ EPVRenderType::Skeleton }; }
 #endif
 
 protected:
diff --git a/Engine/Plugins/Experimental/ProceduralVegetationEditor/Source/ProceduralVegetation/Public/DataTypes/PVData.h b/Engine/Plugins/Experimental/ProceduralVegetationEditor/Source/ProceduralVegetation/Public/DataTypes/PVData.h
--- a/Engine/Plugins/Experimental/ProceduralVegetationEditor/Source/ProceduralVegetation/Public/DataTypes/PVData.h
+++ b/Engine/Plugins/Experimental/ProceduralVegetationEditor/Source/ProceduralVegetation/Public/DataTypes/PVData.h
@@ -2,6 +2,7 @@
 
 #pragma once
 
+#include "PCGContext.h"
 #include "Data/PCGSpatialData.h"
 #include "GeometryCollection/ManagedArrayCollection.h"
 #include "Helpers/PVUtilities.h"
@@ -116,8 +117,6 @@ public:
 	FManagedArrayCollection& GetCollection() { return Collection; }
 	
 #if WITH_EDITOR
-	const TArray<EPVRenderType>& GetRenderType() const { return RenderType; }
-	void SetRenderType(TArray<EPVRenderType> InRenderType) const { RenderType = InRenderType; }
 	void SetDebugSettings(const FPVDebugSettings& InSettings) const { DebugSettings = InSettings; }
 	const FPVDebugSettings& GetDebugSettings() const { return DebugSettings; }
 #endif
@@ -127,7 +126,6 @@ protected:
 	FManagedArrayCollection Collection;
 
 #if WITH_EDITOR
-	mutable TArray<EPVRenderType> RenderType;
 	mutable FPVDebugSettings DebugSettings;
 #endif
 	
diff --git a/Engine/Plugins/Experimental/ProceduralVegetationEditor/Source/ProceduralVegetation/Public/Nodes/PVBaseSettings.h b/Engine/Plugins/Experimental/ProceduralVegetationEditor/Source/ProceduralVegetation/Public/Nodes/PVBaseSettings.h
--- a/Engine/Plugins/Experimental/ProceduralVegetationEditor/Source/ProceduralVegetation/Public/Nodes/PVBaseSettings.h
+++ b/Engine/Plugins/Experimental/ProceduralVegetationEditor/Source/ProceduralVegetation/Public/Nodes/PVBaseSettings.h
@@ -14,11 +14,16 @@ class UPVBaseSettings : public UPCGSettings
 
 public:
 	bool bAutoFocusViewport = false;
+
+	UPVBaseSettings();
 	
 #if WITH_EDITOR
 	virtual EPCGSettingsType GetType() const override { return EPCGSettingsType::Spatial; }
-	virtual TArray<EPVRenderType> GetRenderType() const { return TArray<EPVRenderType>(); }
-
+	virtual TArray<EPVRenderType> GetDefaultRenderType() const { return TArray<EPVRenderType>(); }
+	virtual TArray<EPVRenderType> GetSupportedRenderTypes() const { return GetDefaultRenderType(); }
+	virtual TArray<EPVRenderType> GetCurrentRenderTypes() const { return CurrentRenderType; }
+	virtual void SetCurrentRenderType(TArray<EPVRenderType> InRenderTypes);
+	
 	FPVDebugSettings GetDebugVisualizationSettings() const { return DebugVisualizationSettings; }
 #endif
 
@@ -35,6 +40,9 @@ private:
 #if WITH_EDITORONLY_DATA
 	UPROPERTY(EditAnywhere, Category = DebugSettings, meta=(EditCondition = "bDebug", EditConditionHides))
 	FPVDebugSettings DebugVisualizationSettings;
+
+	UPROPERTY()
+	TArray<EPVRenderType> CurrentRenderType = GetDefaultRenderType();
 #endif
 };
 
diff --git a/Engine/Plugins/Experimental/ProceduralVegetationEditor/Source/ProceduralVegetation/Public/Nodes/PVOutputSettings.h b/Engine/Plugins/Experimental/ProceduralVegetationEditor/Source/ProceduralVegetation/Public/Nodes/PVOutputSettings.h
--- a/Engine/Plugins/Experimental/ProceduralVegetationEditor/Source/ProceduralVegetation/Public/Nodes/PVOutputSettings.h
+++ b/Engine/Plugins/Experimental/ProceduralVegetationEditor/Source/ProceduralVegetation/Public/Nodes/PVOutputSettings.h
@@ -19,6 +19,8 @@ public:
 	virtual FText GetDefaultNodeTitle() const override;
 	virtual FText GetNodeTooltipText() const override;
 	virtual EPCGSettingsType GetType() const override { return EPCGSettingsType::InputOutput; }
+	virtual TArray<EPVRenderType> GetSupportedRenderTypes() const override { return TArray{EPVRenderType::Skeleton, EPVRenderType::Foliage, EPVRenderType::FoliageGrid, EPVRenderType::Mesh, EPVRenderType::Bones}; }
+	virtual TArray<EPVRenderType> GetDefaultRenderType() const override { return TArray{EPVRenderType::Foliage, EPVRenderType::Mesh}; }
 #endif
 
 protected:
diff --git a/Engine/Plugins/Experimental/ProceduralVegetationEditor/Source/ProceduralVegetationEditor/Private/PVEditor.cpp b/Engine/Plugins/Experimental/ProceduralVegetationEditor/Source/ProceduralVegetationEditor/Private/PVEditor.cpp
--- a/Engine/Plugins/Experimental/ProceduralVegetationEditor/Source/ProceduralVegetationEditor/Private/PVEditor.cpp
+++ b/Engine/Plugins/Experimental/ProceduralVegetationEditor/Source/ProceduralVegetationEditor/Private/PVEditor.cpp
@@ -305,6 +305,11 @@ void FPVEditor::OnSelectedNodesChanged(const TSet<UObject*>& InNewSelection)
 			{
 				Settings->bAutoFocusViewport = bFocusViewportOnSelection;
 				bFocusViewportOnSelection = false;
+
+				if (ViewportWidgets.Num() > 0)
+				{
+					ViewportWidgets[0]->OnNodeSelectionChanged(Settings);
+				}
 			}
 			
 			SetNodeInspected(Node, true);
@@ -545,7 +550,14 @@ FReply FPVEditor::ExportInternal()
 
 TSharedRef<SPCGEditorViewport> FPVEditor::CreateViewportWidget()
 {
-	return SNew(SPVEditorViewport);
+	TSharedPtr<SPVEditorViewport> ViewportWidget;
+	
+	SAssignNew(ViewportWidget, SPVEditorViewport);
+
+	ViewportWidgets.Add(ViewportWidget.ToSharedRef());
+
+	return ViewportWidget.ToSharedRef();
+	
 }
 
 bool FPVEditor::GetCollectionFromNode(const UPCGNode* InNode, FName& OutPinName, FManagedArrayCollection& OutCollection) const
diff --git a/Engine/Plugins/Experimental/ProceduralVegetationEditor/Source/ProceduralVegetationEditor/Private/PVEditor.h b/Engine/Plugins/Experimental/ProceduralVegetationEditor/Source/ProceduralVegetationEditor/Private/PVEditor.h
--- a/Engine/Plugins/Experimental/ProceduralVegetationEditor/Source/ProceduralVegetationEditor/Private/PVEditor.h
+++ b/Engine/Plugins/Experimental/ProceduralVegetationEditor/Source/ProceduralVegetationEditor/Private/PVEditor.h
@@ -3,6 +3,7 @@
 #pragma once
 
 #include "PCGEditor.h"
+#include "SPVEditorViewport.h"
 
 class SCollectionSpreadSheetWidget;
 struct FManagedArrayCollection;
@@ -49,6 +50,8 @@ protected:
 private:
 	TObjectPtr<UProceduralVegetation> ProceduralVegetationBeingEdited = nullptr;
 
+	TArray<TSharedRef<SPVEditorViewport>> ViewportWidgets;
+	
 	FDateTime SessionStartTime;
 
 	bool bFocusViewportOnSelection = true;
diff --git a/Engine/Plugins/Experimental/ProceduralVegetationEditor/Source/ProceduralVegetationEditor/Private/SPVEditorViewport.cpp b/Engine/Plugins/Experimental/ProceduralVegetationEditor/Source/ProceduralVegetationEditor/Private/SPVEditorViewport.cpp
--- a/Engine/Plugins/Experimental/ProceduralVegetationEditor/Source/ProceduralVegetationEditor/Private/SPVEditorViewport.cpp
+++ b/Engine/Plugins/Experimental/ProceduralVegetationEditor/Source/ProceduralVegetationEditor/Private/SPVEditorViewport.cpp
@@ -14,45 +14,21 @@
 #include "Components/TextRenderComponent.h"
 
 #include "Visualizations/PVScaleVisualizationComponent.h"
+#include "DataTypes/PVData.h"
+#include "Nodes/PVBaseSettings.h"
 
 #include "Widgets/Input/SNumericEntryBox.h"
 
 #define LOCTEXT_NAMESPACE "PVEditorViewportToolbarSections"
 
-FToolMenuEntry SPVEditorViewport::CreateSettingsToolbarMenu()
+SPVEditorViewport::~SPVEditorViewport()
 {
-	return FToolMenuEntry::InitDynamicEntry(TEXT("PVE Settings"), FNewToolMenuSectionDelegate::CreateLambda([](FToolMenuSection& Section)
-		{
-			FToolMenuEntry& Entry = Section.AddSubMenu(
-				/*Name=*/TEXT("PVESettingsSubmenu"),
-				FText::GetEmpty(),
-				FText::GetEmpty(),
-				FNewToolMenuDelegate::CreateLambda([](UToolMenu* Submenu)
-					{
-						const UPVMannequinWidgetContext* Context = Submenu->FindContext<UPVMannequinWidgetContext>();
-						
-						// Add preference toggles to this section.
-						FToolMenuSection& MannequinSection = Submenu->FindOrAddSection(TEXT("PVE_MannequinSettings"),
-							LOCTEXT("PVE_MannequinSettingsSectionLabel", "Mannequin Settings"));
-						MannequinSection.AddMenuEntry(FPVEditorCommands::Get().ShowMannequin);
-						MannequinSection.AddEntry(
-							FToolMenuEntry::InitWidget(
-								"PVE_MannequinOffset",
-								Context->MannequinOffsetWidget.ToSharedRef(),
-								LOCTEXT("PVE_MannequinOffsetLabel", "Mannequin Offset"),
-								false,
-								true,
-								false,
-								LOCTEXT("PVE_MannequinOffsetTooltip", "Manually Offset the Mannequin")
-							)
-						);
-						FToolMenuSection& ScaleVisSection = Submenu->FindOrAddSection(TEXT("PVE_ScaleVisSettings"),
-							LOCTEXT("PVE_ScaleVisSettingsSectionLabel", "Scale Vis. Settings"));
-						ScaleVisSection.AddMenuEntry(FPVEditorCommands::Get().ShowScaleVisualization);
-					}),
-				/*bInOpenSubMenuOnClick=*/false,
-				FSlateIcon(FAppStyle::GetAppStyleSetName(), "LevelEditor.GameSettings"));
-		}));
+	const FName ToolbarName = "PVE.ViewportToolbar";
+
+	if (UToolMenus::Get()->IsMenuRegistered(ToolbarName))
+	{
+		UToolMenus::Get()->RemoveMenu(ToolbarName);
+	}
 }
 
 TSharedPtr<SWidget> SPVEditorViewport::BuildViewportToolbar()
@@ -78,6 +54,7 @@
 			UE::AdvancedPreviewScene::Menus::FSettingsOptions().ShowToggleGrid(false)
 		);
 
+		RightSection.AddEntry(CreateVisualizationModeToolbarMenu());
 		RightSection.AddEntry(CreateSettingsToolbarMenu());
 	}
 
@@ -274,6 +251,13 @@
 	
 	SetMannequinState(IsShowMannequinChecked());
 	SetScaleVisState(IsShowScaleVisChecked());
+
+	if (bFocusOnNextUpdate)
+	{
+		OnFocusViewportToSelection();
+
+		bFocusOnNextUpdate = false;
+	}
 }
 
 void SPVEditorViewport::OnResetScene()
@@ -282,4 +266,218 @@
 
 	ResetVisualizationScene();
 }
+
+FToolMenuEntry SPVEditorViewport::CreateSettingsToolbarMenu()
+{
+	return FToolMenuEntry::InitDynamicEntry(TEXT("PVE Settings"), FNewToolMenuSectionDelegate::CreateLambda([](FToolMenuSection& Section)
+		{
+			FToolMenuEntry& Entry = Section.AddSubMenu(
+				/*Name=*/TEXT("PVESettingsSubmenu"),
+				FText::GetEmpty(),
+				FText::GetEmpty(),
+				FNewToolMenuDelegate::CreateLambda([](UToolMenu* Submenu)
+					{
+						const UPVMannequinWidgetContext* Context = Submenu->FindContext<UPVMannequinWidgetContext>();
+						
+						// Add preference toggles to this section.
+						FToolMenuSection& MannequinSection = Submenu->FindOrAddSection(TEXT("PVE_MannequinSettings"),
+							LOCTEXT("PVE_MannequinSettingsSectionLabel", "Mannequin Settings"));
+						MannequinSection.AddMenuEntry(FPVEditorCommands::Get().ShowMannequin);
+						MannequinSection.AddEntry(
+							FToolMenuEntry::InitWidget(
+								"PVE_MannequinOffset",
+								Context->MannequinOffsetWidget.ToSharedRef(),
+								LOCTEXT("PVE_MannequinOffsetLabel", "Mannequin Offset"),
+								false,
+								true,
+								false,
+								LOCTEXT("PVE_MannequinOffsetTooltip", "Manually Offset the Mannequin")
+							)
+						);
+						FToolMenuSection& ScaleVisSection = Submenu->FindOrAddSection(TEXT("PVE_ScaleVisSettings"),
+							LOCTEXT("PVE_ScaleVisSettingsSectionLabel", "Scale Vis. Settings"));
+						ScaleVisSection.AddMenuEntry(FPVEditorCommands::Get().ShowScaleVisualization);
+					}),
+				/*bInOpenSubMenuOnClick=*/false,
+				FSlateIcon(FAppStyle::GetAppStyleSetName(), "LevelEditor.GameSettings"));
+		}));
+}
+
+EPVVisualizationMode PVRenderModesToVisualizationMode(const TArray<EPVRenderType>& InRenderType)
+{
+	if (InRenderType.Contains(EPVRenderType::Foliage) && InRenderType.Contains(EPVRenderType::Mesh))
+	{
+		return EPVVisualizationMode::FoliageMesh;
+	}
+	else if (InRenderType.Contains(EPVRenderType::FoliageGrid))
+	{
+		return EPVVisualizationMode::FoliageGrid;
+	}
+	else if (InRenderType.Contains(EPVRenderType::Mesh))
+	{
+		return EPVVisualizationMode::Mesh;
+	}
+	else if (InRenderType.Contains(EPVRenderType::Bones))
+	{
+		return EPVVisualizationMode::Bones;
+	}
+	else if (InRenderType.Contains(EPVRenderType::Skeleton))
+	{
+		return EPVVisualizationMode::PointData;
+	}
+
+	return EPVVisualizationMode::PointData;
+}
+
+TArray<EPVRenderType> VisualizationModeToRenderModes(EPVVisualizationMode InMode, TArray<EPVRenderType> DefautlSettings = TArray<EPVRenderType>{})
+{
+	switch (InMode)
+	{
+	case EPVVisualizationMode::Default:
+		return DefautlSettings;
+		
+	case EPVVisualizationMode::FoliageMesh:
+		return {EPVRenderType::Mesh, EPVRenderType::Foliage};			
+
+	case EPVVisualizationMode::FoliageGrid:
+		return {EPVRenderType::FoliageGrid};
+
+	case EPVVisualizationMode::Bones:
+		return {EPVRenderType::Bones};
+	
+	case EPVVisualizationMode::Mesh:
+		return {EPVRenderType::Mesh};
+
+	case EPVVisualizationMode::PointData:
+		return {EPVRenderType::Skeleton};
+
+	case EPVVisualizationMode::BonesMesh:
+		return {EPVRenderType::Bones, EPVRenderType::Mesh};
+
+	case EPVVisualizationMode::PointDataMesh:
+		return {EPVRenderType::Skeleton, EPVRenderType::Mesh};
+
+	default:
+		return TArray<EPVRenderType>{};
+		
+	}
+}
+
+FToolMenuEntry SPVEditorViewport::CreateVisualizationModeToolbarMenu()
+{
+	return FToolMenuEntry::InitDynamicEntry
+	(
+	    "VisualizationMode",
+	    FNewToolMenuSectionDelegate::CreateSPLambda(this,[this](FToolMenuSection& InnerSection) -> void
+	       {
+	       	TAttribute<FText> LabelAttribute = UE::UnrealEd::GetViewModesSubmenuLabel(nullptr);
+
+	       	LabelAttribute = TAttribute<FText>::CreateSPLambda(this,
+					[this]()
+					{
+						UEnum* EnumPtr = StaticEnum<EPVVisualizationMode>();
+							return EnumPtr->GetDisplayNameTextByIndex((int)CurrentVisualizationMode);
+					}
+				);
+	       		
+	          FToolMenuEntry& Entry = InnerSection.AddSubMenu(
+	             "VisualizationModes",
+	             LabelAttribute,
+	             FText::GetEmpty(),
+	             FNewToolMenuDelegate::CreateSPLambda
+	             (this, [this](UToolMenu* Submenu) -> void
+	                   {
+	                      FToolMenuSection& RenderModeSelectionSection = Submenu->FindOrAddSection(
+	                         "VisualizationModeSection",
+	                         LOCTEXT("FoliageVisualizationModeSelectionSectionLabel", "Visualization modes")
+	                      );
+
+	                      UEnum* EnumPtr = StaticEnum<EPVVisualizationMode>();
+	                      
+	                      if (EnumPtr)
+	                      {
+	                         for (int32 i = 0; i < EnumPtr->NumEnums() - 1; i++) // last entry is _MAX or hidden
+	                         {
+	                            FText ModeName = EnumPtr->GetDisplayNameTextByIndex(i);
+	                            EPVVisualizationMode Mode = static_cast<EPVVisualizationMode>(EnumPtr->GetValueByIndex(i));
+	                         	
+	                         	RenderModeSelectionSection.AddMenuEntry
+	                         	(
+								NAME_None,
+								ModeName,
+								FText(),
+								FSlateIcon(),
+								FUIAction(
+								   FExecuteAction::CreateSPLambda(this,
+									  [this, Mode]()
+									  {
+									  		OnVisualizationModeChanged(Mode);
+									  }
+								   ),
+								   FCanExecuteAction::CreateSPLambda(this,[this, Mode]()
+									   {
+								   		if (Mode != EPVVisualizationMode::Default)
+								   		{
+								   			auto RenderTypes = VisualizationModeToRenderModes(Mode);
+
+											   for (int i = 0; i < RenderTypes.Num(); i++)
+											   {
+												   if (!SupportedRenderTypes.Contains(RenderTypes[i]))
+												   {
+													   return false;
+												   }
+											   }
+								   		}
+								   		
+								   		return true;
+								   		
+									   }),
+								   FIsActionChecked::CreateSPLambda(this,
+									  [this, Mode]()
+									  {
+									  		return CurrentVisualizationMode == Mode;
+									  }
+								   )
+								),EUserInterfaceActionType::RadioButton
+								);
+	                         }
+	                      }
+						}
+	             ),
+	             false,
+	             FSlateIcon(FAppStyle::GetAppStyleSetName(), "EditorViewport.Visualizers")
+	          );
+	          Entry.ToolBarData.ResizeParams.ClippingPriority = 1100;
+	       }
+	    )
+	 );
+}
+
+void SPVEditorViewport::OnVisualizationModeChanged(EPVVisualizationMode InMode)
+{
+	CurrentVisualizationMode = InMode;
+
+	if (CurrentNodeSettings)
+	{
+		const TArray<EPVRenderType> RenderTypes = VisualizationModeToRenderModes(CurrentVisualizationMode, CurrentNodeSettings->GetDefaultRenderType());
+												
+		CurrentNodeSettings->SetCurrentRenderType(RenderTypes);
+
+		Invalidate();
+
+		bFocusOnNextUpdate = true;
+	}
+}
+
+void SPVEditorViewport::OnNodeSelectionChanged(UPVBaseSettings* InSettings)
+{
+	SupportedRenderTypes = InSettings->GetSupportedRenderTypes();
+	CurrentVisualizationMode = EPVVisualizationMode::Default;
+	CurrentNodeSettings = InSettings;
+	
+	CurrentNodeSettings->SetCurrentRenderType(InSettings->GetDefaultRenderType());
+
+	bFocusOnNextUpdate = true;
+}
+
 #undef LOCTEXT_NAMESPACE
diff --git a/Engine/Plugins/Experimental/ProceduralVegetationEditor/Source/ProceduralVegetationEditor/Private/SPVEditorViewport.h b/Engine/Plugins/Experimental/ProceduralVegetationEditor/Source/ProceduralVegetationEditor/Private/SPVEditorViewport.h
--- a/Engine/Plugins/Experimental/ProceduralVegetationEditor/Source/ProceduralVegetationEditor/Private/SPVEditorViewport.h
+++ b/Engine/Plugins/Experimental/ProceduralVegetationEditor/Source/ProceduralVegetationEditor/Private/SPVEditorViewport.h
@@ -8,10 +8,24 @@
 
 #include "SPVEditorViewport.generated.h"
 
-struct FToolMenuEntry;
+UENUM()
+enum class EPVVisualizationMode : uint8
+{
+	Default			UMETA(DisplayName = "Default"),
+	PointData		UMETA(DisplayName = "Point data"),
+	Mesh			UMETA(DisplayName = "Mesh"),
+	FoliageGrid		UMETA(DisplayName = "Foliage grid"),
+	Bones			UMETA(DisplayName = "Bones"),
+	PointDataMesh	UMETA(DisplayName = "Point data + Mesh"),
+	FoliageMesh		UMETA(DisplayName = "Foliage + Mesh"),
+	BonesMesh		UMETA(DisplayName = "Bones + Mesh"),
+};
 
+class UPVBaseSettings;
 class UPVEditorSettings;
 class UPVScaleVisualizationComponent;
+struct FToolMenuEntry;
+enum class EPVRenderType : uint8;
 
 UCLASS(MinimalAPI)
 class UPVMannequinWidgetContext : public UObject
@@ -25,11 +39,17 @@ public:
 class SPVEditorViewport : public SPCGEditorViewport
 {
 public:
+	virtual ~SPVEditorViewport() override;
 	FToolMenuEntry CreateSettingsToolbarMenu();
+	FToolMenuEntry CreateVisualizationModeToolbarMenu();
+
+	void OnNodeSelectionChanged(UPVBaseSettings* InSettings);
 
 protected:
+	//~ Begin SEditorViewport Interface
 	virtual TSharedPtr<SWidget> BuildViewportToolbar() override;
 	virtual void BindCommands() override;
+	//~ End SEditorViewport Interface
 
 private:
 	void InitVisualizationScene();
@@ -50,6 +70,14 @@ private:
 
 	static UPVMannequinWidgetContext* CreateMannequinWidgetContext(TSharedPtr<SWidget> InOffsetWidget);
 
+	TArray<EPVRenderType> SupportedRenderTypes;
+	EPVVisualizationMode CurrentVisualizationMode = EPVVisualizationMode::Default;
+	TObjectPtr<UPVBaseSettings> CurrentNodeSettings = nullptr;
+
+	void OnVisualizationModeChanged(EPVVisualizationMode InMode);
+
+	bool bFocusOnNextUpdate = false;
+
 protected:
 	virtual void OnSetupScene() override;
 	virtual void OnResetScene() override;
diff --git a/Engine/Plugins/Experimental/ProceduralVegetationEditor/Source/ProceduralVegetationEditor/Private/Visualizations/PVDataVisualization.cpp b/Engine/Plugins/Experimental/ProceduralVegetationEditor/Source/ProceduralVegetationEditor/Private/Visualizations/PVDataVisualization.cpp
--- a/Engine/Plugins/Experimental/ProceduralVegetationEditor/Source/ProceduralVegetationEditor/Private/Visualizations/PVDataVisualization.cpp
+++ b/Engine/Plugins/Experimental/ProceduralVegetationEditor/Source/ProceduralVegetationEditor/Private/Visualizations/PVDataVisualization.cpp
@@ -95,7 +95,7 @@ FPCGSetupSceneFunc FPVDataVisualization::GetViewportSetupFunc(const UPCGSettings
 
 		FBoxSphereBounds Bounds(ForceInit);
 
-		for (const EPVRenderType& RenderType : WeakData->GetRenderType())
+		for (const EPVRenderType& RenderType : WeakSettings->GetCurrentRenderTypes())
 		{
 			if (RenderMap.Contains(RenderType))
 			{
diff --git a/Engine/Plugins/Experimental/ProceduralVegetationEditor/Source/ProceduralVegetationEditor/Private/Visualizations/PVLineBatchComponent.cpp b/Engine/Plugins/Experimental/ProceduralVegetationEditor/Source/ProceduralVegetationEditor/Private/Visualizations/PVLineBatchComponent.cpp
--- a/Engine/Plugins/Experimental/ProceduralVegetationEditor/Source/ProceduralVegetationEditor/Private/Visualizations/PVLineBatchComponent.cpp
+++ b/Engine/Plugins/Experimental/ProceduralVegetationEditor/Source/ProceduralVegetationEditor/Private/Visualizations/PVLineBatchComponent.cpp
@@ -73,9 +73,9 @@ void UPVLineBatchComponent::InitBounds()
 	BBox = FBox(ForceInit);
 }
 
-void UPVLineBatchComponent::AddLine(const FVector& InStartPos, const FVector& InEndPos, const FLinearColor& InColor)
+void UPVLineBatchComponent::AddLine(const FVector& InStartPos, const FVector& InEndPos, const FLinearColor& InColor,const ESceneDepthPriorityGroup InDepthPriorityGroup, const EPointDrawSettings InPointDrawSettings)
 {
-	Lines.Emplace(InStartPos, InEndPos, InColor);
+	Lines.Emplace(InStartPos, InEndPos, InColor, InDepthPriorityGroup, InPointDrawSettings);
 
 	BBox += InStartPos;
 	BBox += InEndPos;
@@ -89,11 +89,21 @@ void UPVLineBatchComponent::Draw(const FSceneView* View, FPrimitiveDrawInterface
 			Line.StartPos,
 			Line.EndPos,
 			Line.Color,
-			0,
+			Line.DepthPriorityGroup,
 			3.0f,
 			0,
 			true
 		);
+
+		if (Line.PointDrawSettings == EPointDrawSettings::Start || Line.PointDrawSettings == EPointDrawSettings::Both)
+		{
+			PDI->DrawPoint(Line.StartPos, FLinearColor::White, PointSize, SDPG_Foreground);
+		}
+
+		if (Line.PointDrawSettings == EPointDrawSettings::End || Line.PointDrawSettings == EPointDrawSettings::Both)
+		{
+			PDI->DrawPoint(Line.EndPos, FLinearColor::White, PointSize, SDPG_Foreground);
+		}
 	}
 }
 
diff --git a/Engine/Plugins/Experimental/ProceduralVegetationEditor/Source/ProceduralVegetationEditor/Private/Visualizations/PVLineBatchComponent.h b/Engine/Plugins/Experimental/ProceduralVegetationEditor/Source/ProceduralVegetationEditor/Private/Visualizations/PVLineBatchComponent.h
--- a/Engine/Plugins/Experimental/ProceduralVegetationEditor/Source/ProceduralVegetationEditor/Private/Visualizations/PVLineBatchComponent.h
+++ b/Engine/Plugins/Experimental/ProceduralVegetationEditor/Source/ProceduralVegetationEditor/Private/Visualizations/PVLineBatchComponent.h
@@ -34,6 +34,15 @@ private:
 	const TObjectPtr<const UPVLineBatchComponent> LineComponent;
 };
 
+UENUM()
+enum class EPointDrawSettings : uint8
+{
+	None,
+	Start,
+	End,
+	Both
+};
+
 USTRUCT()
 struct FPVLineInfo
 {
@@ -42,6 +51,8 @@ struct FPVLineInfo
 	FVector StartPos;
 	FVector EndPos;
 	FLinearColor Color;
+	ESceneDepthPriorityGroup DepthPriorityGroup = SDPG_World;
+	EPointDrawSettings PointDrawSettings = EPointDrawSettings::None;
 };
 
 UCLASS(Blueprintable, BlueprintType, meta=(BlueprintSpawnableComponent), MinimalAPI)
@@ -53,7 +64,9 @@ public:
 	UPVLineBatchComponent();
 
 	void InitBounds();
-	void AddLine(const FVector& InStartPos, const FVector& InEndPos, const FLinearColor& InColor = FLinearColor::White);
+	void AddLine(const FVector& InStartPos, const FVector& InEndPos, const FLinearColor& InColor = FLinearColor::White,
+				const ESceneDepthPriorityGroup InDepthPriorityGroup = ESceneDepthPriorityGroup::SDPG_World,
+				const EPointDrawSettings InPointDrawSettings = EPointDrawSettings::None);
 
 	void Draw(const FSceneView* View, FPrimitiveDrawInterface* PDI) const;
 
@@ -65,4 +78,7 @@ public:
 protected:
 	FBox BBox;
 	TArray<FPVLineInfo> Lines;
+
+private:
+	const float PointSize = 7.5f;
 };
diff --git a/Engine/Plugins/Experimental/ProceduralVegetationEditor/Source/ProceduralVegetationEditor/Private/Visualizations/PVSkeletonVisualizerComponent.cpp b/Engine/Plugins/Experimental/ProceduralVegetationEditor/Source/ProceduralVegetationEditor/Private/Visualizations/PVSkeletonVisualizerComponent.cpp
--- a/Engine/Plugins/Experimental/ProceduralVegetationEditor/Source/ProceduralVegetationEditor/Private/Visualizations/PVSkeletonVisualizerComponent.cpp
+++ b/Engine/Plugins/Experimental/ProceduralVegetationEditor/Source/ProceduralVegetationEditor/Private/Visualizations/PVSkeletonVisualizerComponent.cpp
@@ -56,7 +56,7 @@ void UPVSkeletonVisualizerComponent::SetCollection(const FManagedArrayCollection
 		{
 			FVector CurrentPosition = FVector(PointsPositions[BranchPoints[BranchPointIndex]]);
 
-			AddLine(PreviousPosition, CurrentPosition);
+			AddLine(PreviousPosition, CurrentPosition, FLinearColor::White, ESceneDepthPriorityGroup::SDPG_MAX, EPointDrawSettings::Start);
 
 			PreviousPosition = CurrentPosition;
 		}
