Fix iteration of unsafe container in PVE Export

diff --git a/Engine/Plugins/Experimental/ProceduralVegetationEditor/Source/ProceduralVegetationEditor/Private/Helpers/PVExportHelper.cpp b/Engine/Plugins/Experimental/ProceduralVegetationEditor/Source/ProceduralVegetationEditor/Private/Helpers/PVExportHelper.cpp
--- a/Engine/Plugins/Experimental/ProceduralVegetationEditor/Source/ProceduralVegetationEditor/Private/Helpers/PVExportHelper.cpp
+++ b/Engine/Plugins/Experimental/ProceduralVegetationEditor/Source/ProceduralVegetationEditor/Private/Helpers/PVExportHelper.cpp
@@ -422,7 +422,7 @@ namespace PV::Export::Internal
 						continue;
 					}
 					
-					for (int32 Tid = 0; Tid < Result.OutDynamicMesh.GetTrianglesBuffer().Num(); Tid++)
+					for (int32 Tid : Result.OutDynamicMesh.TriangleIndicesItr())
 					{
 						const int32 MaterialID = MaterialIDs->GetValue(Tid);
 						if (MaterialID == INDEX_NONE)
