[PVE] Foliage Distributor warning added

diff --git a/Engine/Plugins/Experimental/ProceduralVegetationEditor/Source/ProceduralVegetation/Private/Nodes/PVOutputSettings.cpp b/Engine/Plugins/Experimental/ProceduralVegetationEditor/Source/ProceduralVegetation/Private/Nodes/PVOutputSettings.cpp
--- a/Engine/Plugins/Experimental/ProceduralVegetationEditor/Source/ProceduralVegetation/Private/Nodes/PVOutputSettings.cpp
+++ b/Engine/Plugins/Experimental/ProceduralVegetationEditor/Source/ProceduralVegetation/Private/Nodes/PVOutputSettings.cpp
@@ -83,12 +83,7 @@ bool FPVOutputElement::ExecuteInternal(FPCGContext* InContext) const
 				return true;
 			}
 
-			FManagedArrayCollection OutCollection = InputData->GetCollection();
-
-			UPVData* ProceduralVegetationData = FPCGContext::NewObject_AnyThread<UPVData>(InContext);
-			ProceduralVegetationData->Initialize(MoveTemp(OutCollection));
-
-			InContext->OutputData.TaggedData.Emplace(ProceduralVegetationData);
+			InContext->OutputData = InContext->InputData;
 		}
 		else
 		{
diff --git a/Engine/Plugins/Experimental/ProceduralVegetationEditor/Source/ProceduralVegetation/Public/PVExportParams.h b/Engine/Plugins/Experimental/ProceduralVegetationEditor/Source/ProceduralVegetation/Public/PVExportParams.h
--- a/Engine/Plugins/Experimental/ProceduralVegetationEditor/Source/ProceduralVegetation/Public/PVExportParams.h
+++ b/Engine/Plugins/Experimental/ProceduralVegetationEditor/Source/ProceduralVegetation/Public/PVExportParams.h
@@ -21,6 +21,9 @@ struct PROCEDURALVEGETATION_API FPVExportParams
 	UPROPERTY()
 	bool bShouldExport = true;
 
+	UPROPERTY()
+	bool bShouldExportFoliage = true;
+
 	UPROPERTY(EditAnywhere, Category = "Export", Meta = (ContentDir, DisplayName = "Content Browser Folder", Tooltip="Path to the folder where exported mesh will be saved\n\nPath to the folder where exported mesh will be saved."))
 	FDirectoryPath ContentBrowserFolder;
 
diff --git a/Engine/Plugins/Experimental/ProceduralVegetationEditor/Source/ProceduralVegetationEditor/Private/Helpers/PVExportHelper.cpp b/Engine/Plugins/Experimental/ProceduralVegetationEditor/Source/ProceduralVegetationEditor/Private/Helpers/PVExportHelper.cpp
--- a/Engine/Plugins/Experimental/ProceduralVegetationEditor/Source/ProceduralVegetationEditor/Private/Helpers/PVExportHelper.cpp
+++ b/Engine/Plugins/Experimental/ProceduralVegetationEditor/Source/ProceduralVegetationEditor/Private/Helpers/PVExportHelper.cpp
@@ -905,6 +905,7 @@ namespace PV::Export::Internal
 		const FManagedArrayCollection& Collection,
 		ENaniteShapePreservation InShapePreservation,
 		bool bBuildNaniteAssemblies,
+		bool bShouldExportFoliage,
 		TArray<UPackage*>& OutModifiedPackages,
 		const FStatusReportCallback& StatusReportCallback
 	)
@@ -945,7 +946,7 @@ namespace PV::Export::Internal
 		}
 		
 		// Combine all the foliage instances into one dynamic mesh if not building nanite assemblies
-		if (!bBuildNaniteAssemblies)
+		if (!bBuildNaniteAssemblies && bShouldExportFoliage)
 		{
 			if (!UpdateStage(EExportMeshStage::CombiningMeshes, 0.f))
 			{
@@ -976,7 +977,7 @@ namespace PV::Export::Internal
 			UGeometryScriptLibrary_StaticMeshFunctions::CopyMeshToStaticMesh(DynamicMesh, ExportMesh, Options, {}, Output);
 		}
 
-		if (bBuildNaniteAssemblies)
+		if (bBuildNaniteAssemblies && bShouldExportFoliage)
 		{
 			if (!UpdateStage(EExportMeshStage::BuildingNaniteAssemblies, 1.f))
 			{
@@ -1028,6 +1029,7 @@ namespace PV::Export::Internal
 		ENaniteShapePreservation InShapePreservation,
 		TObjectPtr<const UPVWindSettings> InWindSettings,
 		bool bBuildNaniteAssemblies,
+		bool bShouldExportFoliage,
 		TArray<UPackage*>& OutModifiedPackages,
 		const FStatusReportCallback& StatusReportCallback
 	)
@@ -1078,7 +1080,10 @@ namespace PV::Export::Internal
 				BoneFacade.CreateBoneData(BoneNodes, 0);
 			}
 
-			AssignBoneIDsToFoliage(BoneNodes, ExportedCollection);
+			if (bShouldExportFoliage)
+			{
+				AssignBoneIDsToFoliage(BoneNodes, ExportedCollection);
+			}
 
 			const int32 BoneCount = BoneNodes.Num();
 			MeshAttributes->EnableBones(BoneCount);
@@ -1119,7 +1124,7 @@ namespace PV::Export::Internal
 			}
 		}
 
-		if (!bBuildNaniteAssemblies)
+		if (!bBuildNaniteAssemblies && bShouldExportFoliage)
 		{
 			if (!UpdateStage(EExportMeshStage::CombiningMeshes, 0.f))
 			{
@@ -1240,7 +1245,7 @@ namespace PV::Export::Internal
 			);
 		}
 
-		if (bBuildNaniteAssemblies)
+		if (bBuildNaniteAssemblies && bShouldExportFoliage)
 		{
 			if (!UpdateStage(EExportMeshStage::BuildingNaniteAssemblies, 1.f))
 			{
@@ -1456,6 +1461,7 @@ namespace PV::Export
 		const FString MeshName = ExportParams.MeshName.ToString();
 		const FString MeshPackageName = ExportParams.GetOutputMeshPackagePath();
 		const bool bCreateNaniteFoliage = ExportParams.bCreateNaniteFoliage && NaniteAssembliesSupported();
+		const bool bShouldExportFoliage = ExportParams.bShouldExportFoliage;
 
 		TArray<UPackage*> ModifiedPackages;
 
@@ -1481,6 +1487,7 @@ namespace PV::Export
 				Collection,
 				ExportParams.NaniteShapePreservation,
 				bCreateNaniteFoliage,
+				bShouldExportFoliage,
 				ModifiedPackages,
 				StatusReportCallback
 			))
@@ -1531,6 +1538,7 @@ namespace PV::Export
 				ExportParams.NaniteShapePreservation,
 				ExportParams.WindSettings,
 				bCreateNaniteFoliage,
+				bShouldExportFoliage,
 				ModifiedPackages,
 				StatusReportCallback
 			))
diff --git a/Engine/Plugins/Experimental/ProceduralVegetationEditor/Source/ProceduralVegetationEditor/Private/PVEditor.cpp b/Engine/Plugins/Experimental/ProceduralVegetationEditor/Source/ProceduralVegetationEditor/Private/PVEditor.cpp
--- a/Engine/Plugins/Experimental/ProceduralVegetationEditor/Source/ProceduralVegetationEditor/Private/PVEditor.cpp
+++ b/Engine/Plugins/Experimental/ProceduralVegetationEditor/Source/ProceduralVegetationEditor/Private/PVEditor.cpp
@@ -15,6 +15,7 @@
 #include "Dataflow/DataflowCollectionSpreadSheetWidget.h"
 
 #include "DataTypes/PVData.h"
+#include "DataTypes/PVFoliageMeshData.h"
 
 #include "Helpers/PVExportHelper.h"
 #include "Helpers/PVUtilities.h"
@@ -445,6 +446,58 @@ FReply FPVEditor::ExportInternal()
 		return true;
 	};
 
+	const auto FoliageDistributorValidation = [this](const TArray<UPCGNode*>& NodesToExport)->bool
+	{
+		TArray<const UPCGNode*> NodesWithoutFoliage;
+		for (UPCGNode* Node : NodesToExport)
+		{
+			FName PinName = NAME_None;
+			const UPVData* PVData = GetPVDataFromNode(Node, PinName);
+			if (PVData && !Cast<UPVFoliageMeshData>(PVData))
+			{
+				NodesWithoutFoliage.Add(Node);
+			}
+		}
+
+		if (NodesWithoutFoliage.Num() > 0)
+		{
+			FString NodeNames = "";
+			for (int32 i = 0; i < NodesWithoutFoliage.Num(); i++)
+			{
+				const UPCGNode* Node = NodesWithoutFoliage[i];
+				NodeNames += FString::Format(TEXT("\" {0} \""), {Node->GetAuthoredTitleName().ToString()});
+				
+				if (NodesWithoutFoliage.Num() > 1 && i == NodesWithoutFoliage.Num() - 2)
+				{
+					NodeNames += " and ";
+				}
+
+				if (i != NodesWithoutFoliage.Num() - 2)
+				{
+					NodeNames += ", ";
+				}
+			}
+			
+			FString Message = FString::Format(TEXT("{0} does not use input from foliage distributor node, Do you want to export with default foliage distribution."),{NodeNames});
+			const EAppReturnType::Type ret = FMessageDialog::Open(EAppMsgType::YesNoCancel, FText::Format(LOCTEXT("Export_Foliage_Warning", "{0}"), FText::FromString(Message)));
+			if (ret != EAppReturnType::Cancel)
+			{
+				const bool bShouldExportFoliage = ret == EAppReturnType::Yes;
+				for (const UPCGNode* Node : NodesWithoutFoliage)
+				{
+					UPVOutputSettings* OutputSettings = Cast<UPVOutputSettings>(Node->GetSettings());
+					OutputSettings->ExportSettings.bShouldExportFoliage = bShouldExportFoliage;
+				}
+			}
+			else
+			{
+				return false;
+			}
+		}
+
+		return true;
+	};
+
 	const TArray<UPCGNode*> NodesToExport = ProceduralVegetationGraph->GetNodes().FilterByPredicate([](UPCGNode* InNode)->bool {
 			const UPVOutputSettings* OutputSettings = Cast<const UPVOutputSettings>(InNode->GetSettings());
 			return OutputSettings && OutputSettings->ExportSettings.bShouldExport;
@@ -466,6 +519,11 @@ FReply FPVEditor::ExportInternal()
 		return FReply::Handled();
 	}
 
+	if (!FoliageDistributorValidation(NodesToExport))
+	{
+		return FReply::Handled();
+	}
+
 	FScopedSlowTask ProgressSlowTask(NodesToExport.Num(), LOCTEXT("ExportAsset_SlowTask", "Exporting Asset(s) ..."));
 	ProgressSlowTask.MakeDialog(true);
 
@@ -560,7 +618,7 @@ TSharedRef<SPCGEditorViewport> FPVEditor::CreateViewportWidget()
 	
 }
 
-bool FPVEditor::GetCollectionFromNode(const UPCGNode* InNode, FName& OutPinName, FManagedArrayCollection& OutCollection) const
+const UPVData* FPVEditor::GetPVDataFromNode(const UPCGNode* InNode, FName& OutPinName) const
 {
 	const UPCGPin* Pin = nullptr;
 		
@@ -593,16 +651,26 @@ bool FPVEditor::GetCollectionFromNode(const UPCGNode* InNode, FName& OutPinName,
 		if(TaggedData != nullptr)
 		{
 			const UPVData* ProceduralVegetationData = Cast<UPVData>(TaggedData->Data);
-		
 			if(ProceduralVegetationData != nullptr && Pin)
 			{
 				OutPinName = Pin->Properties.Label;
-				OutCollection = ProceduralVegetationData->GetCollection();
-				return true;
+				return ProceduralVegetationData;
 			}
 		}
 	}
 
+	return nullptr;
+}
+
+bool FPVEditor::GetCollectionFromNode(const UPCGNode* InNode, FName& OutPinName, FManagedArrayCollection& OutCollection) const
+{
+	const UPVData* PVData = GetPVDataFromNode(InNode, OutPinName);
+	if(PVData != nullptr)
+	{
+		OutCollection = PVData->GetCollection();
+		return true;
+	}
+
 	return false;
 }
 
diff --git a/Engine/Plugins/Experimental/ProceduralVegetationEditor/Source/ProceduralVegetationEditor/Private/PVEditor.h b/Engine/Plugins/Experimental/ProceduralVegetationEditor/Source/ProceduralVegetationEditor/Private/PVEditor.h
--- a/Engine/Plugins/Experimental/ProceduralVegetationEditor/Source/ProceduralVegetationEditor/Private/PVEditor.h
+++ b/Engine/Plugins/Experimental/ProceduralVegetationEditor/Source/ProceduralVegetationEditor/Private/PVEditor.h
@@ -3,6 +3,7 @@
 #pragma once
 
 #include "PCGEditor.h"
+#include "DataTypes/PVData.h"
 #include "SPVEditorViewport.h"
 
 class SCollectionSpreadSheetWidget;
@@ -73,4 +74,6 @@ private:
 
 protected:
 	virtual TSharedRef<SPCGEditorViewport> CreateViewportWidget() override;
+	
+	const UPVData* GetPVDataFromNode(const UPCGNode* InNode, FName& OutPinName) const;
 };
