[PVE] Preset Variation info in preset loader details

diff --git a/Engine/Plugins/Experimental/ProceduralVegetationEditor/Source/ProceduralVegetation/Private/Nodes/PVPresetLoaderSettings.cpp b/Engine/Plugins/Experimental/ProceduralVegetationEditor/Source/ProceduralVegetation/Private/Nodes/PVPresetLoaderSettings.cpp
--- a/Engine/Plugins/Experimental/ProceduralVegetationEditor/Source/ProceduralVegetation/Private/Nodes/PVPresetLoaderSettings.cpp
+++ b/Engine/Plugins/Experimental/ProceduralVegetationEditor/Source/ProceduralVegetation/Private/Nodes/PVPresetLoaderSettings.cpp
@@ -2,6 +2,7 @@
 
 #include "PVPresetLoaderSettings.h"
 #include "PCGContext.h"
+#include "ProceduralVegetationModule.h"
 #include "DataTypes/PVGrowthData.h"
 #include "ProceduralVegetationPreset.h"
 
@@ -21,8 +22,28 @@ FText UPVPresetLoaderSettings::GetNodeTooltipText() const
 		"It contains hormonal information, the main structure of the vegetation, data about the branches, where foliage will / can be attached, references to foliage meshes, materials to be applied"
 	);
 }
+
+void UPVPresetLoaderSettings::PostEditChangeProperty(FPropertyChangedEvent& PropertyChangedEvent)
+{
+	Super::PostEditChangeProperty(PropertyChangedEvent);
+	
+	if (PropertyChangedEvent.Property
+		&& PropertyChangedEvent.Property->GetFName() == GET_MEMBER_NAME_CHECKED(UPVPresetLoaderSettings, Preset))
+	{
+		FillPresetVariationsInfo();
+	}
+}
 #endif
 
+void UPVPresetLoaderSettings::PostLoad()
+{
+	Super::PostLoad();
+	if (PresetVariations.IsEmpty())
+	{
+		FillPresetVariationsInfo();
+	}
+}
+
 TArray<FPCGPinProperties> UPVPresetLoaderSettings::InputPinProperties() const
 {
 	return TArray<FPCGPinProperties>();
@@ -53,6 +74,24 @@ FPCGElementPtr UPVPresetLoaderSettings::CreateElement() const
 	return MakeShared<FPVPresetLoaderElement>();
 }
 
+void UPVPresetLoaderSettings::FillPresetVariationsInfo()
+{
+	PresetVariations.Empty();
+
+	if (!Preset)
+	{
+		UE_LOG(LogProceduralVegetation, Log, TEXT("Preset is null, no preset info can be filled."));
+		return;
+	}
+
+	if (Preset->PresetVariations.IsEmpty())
+	{
+		Preset->FillVariationInfo();
+	}
+
+	PresetVariations = Preset->PresetVariations;
+}
+
 bool FPVPresetLoaderElement::ExecuteInternal(FPCGContext* InContext) const
 {
 	TRACE_CPUPROFILER_EVENT_SCOPE(FPVPresetLoaderElement::ExecuteInternal);
diff --git a/Engine/Plugins/Experimental/ProceduralVegetationEditor/Source/ProceduralVegetation/Private/Nodes/PVPresetLoaderSettings.h b/Engine/Plugins/Experimental/ProceduralVegetationEditor/Source/ProceduralVegetation/Private/Nodes/PVPresetLoaderSettings.h
--- a/Engine/Plugins/Experimental/ProceduralVegetationEditor/Source/ProceduralVegetation/Private/Nodes/PVPresetLoaderSettings.h
+++ b/Engine/Plugins/Experimental/ProceduralVegetationEditor/Source/ProceduralVegetation/Private/Nodes/PVPresetLoaderSettings.h
@@ -6,18 +6,24 @@
 
 #include "PVPresetLoaderSettings.generated.h"
 
+struct FPVPresetVariationInfo;
+
 UCLASS(BlueprintType, ClassGroup = (Procedural))
 class UPVPresetLoaderSettings : public UPVBaseSettings
 {
 	GENERATED_BODY()
 
 public:
+	virtual void PostLoad() override;
+	
 #if WITH_EDITOR
 	virtual FName GetDefaultNodeName() const override { return FName(TEXT("ProceduralVegetationPresetLoader")); }
 	virtual FText GetDefaultNodeTitle() const override;
 	virtual FText GetNodeTooltipText() const override;
 	virtual EPCGSettingsType GetType() const override { return EPCGSettingsType::InputOutput; }
 	virtual TArray<EPVRenderType> GetDefaultRenderType() const override { return TArray{ EPVRenderType::Skeleton }; }
+
+	virtual void PostEditChangeProperty(FPropertyChangedEvent& PropertyChangedEvent) override;
 #endif
 
 protected:
@@ -29,6 +35,12 @@ protected:
 public:
 	UPROPERTY(EditAnywhere, Category="Preset", meta=(Tooltip="Loads the selected procedural vegetation preset\n\nThis allows the user to load a procedural vegetation preset. A procedural vegatation preset contains data for a specie of vegetaion. It contains hormonal information, the main structure of the vegetation, data about the branches, where foliage will/can be attached, references to foliage meshes, materials to be applied"))
 	TObjectPtr<class UProceduralVegetationPreset> Preset = nullptr;
+
+	UPROPERTY(VisibleAnywhere, Category="Preset")
+	TArray<FPVPresetVariationInfo> PresetVariations;
+
+private:
+	void FillPresetVariationsInfo();
 };
 
 class FPVPresetLoaderElement : public FPVBaseElement
diff --git a/Engine/Plugins/Experimental/ProceduralVegetationEditor/Source/ProceduralVegetation/Private/ProceduralVegetationPreset.cpp b/Engine/Plugins/Experimental/ProceduralVegetationEditor/Source/ProceduralVegetation/Private/ProceduralVegetationPreset.cpp
--- a/Engine/Plugins/Experimental/ProceduralVegetationEditor/Source/ProceduralVegetation/Private/ProceduralVegetationPreset.cpp
+++ b/Engine/Plugins/Experimental/ProceduralVegetationEditor/Source/ProceduralVegetation/Private/ProceduralVegetationPreset.cpp
@@ -6,8 +6,32 @@
 #include "UObject/Package.h"
 #include "AssetRegistry/IAssetRegistry.h"
 #include "Facades/PVBranchFacade.h"
+#include "Facades/PVProfileFacade.h"
 #include "Helpers/PVJSONHelper.h"
 
+void FPVPresetVariationInfo::Fill(FName InName, FManagedArrayCollection Collection)
+{
+	Name = InName;
+	
+	const PV::Facades::FFoliageFacade Facade(Collection);
+	for (int32 FoliageIndex = 0; FoliageIndex < Facade.NumFoliageNames(); ++FoliageIndex)
+	{
+		FoliageMeshes.Emplace(FSoftObjectPath(Facade.GetFoliageName(FoliageIndex)));
+	}
+
+	const PV::Facades::FBranchFacade BranchFacade(Collection);
+	Materials.Emplace(FSoftObjectPath(BranchFacade.GetTrunkMaterialPath()));
+
+	if (const PV::Facades::FPlantProfileFacade PlantProfileFacade = PV::Facades::FPlantProfileFacade(Collection);
+			PlantProfileFacade.NumProfileEntries() > 0)
+	{
+		for (int32 i = 0; i < PlantProfileFacade.NumProfileEntries(); i++)
+		{
+			PlantProfiles.Add(FString::FromInt(i));
+		}
+	}
+}
+
 UProceduralVegetationPreset::UProceduralVegetationPreset()
 {
 	FString AssetPath = FPackageName::GetLongPackagePath(GetPackage()->GetName());
@@ -15,6 +39,16 @@ UProceduralVegetationPreset::UProceduralVegetationPreset()
 	MaterialsFolder.Path = FPaths::Combine(AssetPath, TEXT("Materials"));
 }
 
+void UProceduralVegetationPreset::PostLoad()
+{
+	Super::PostLoad();
+
+	if (PresetVariations.IsEmpty())
+	{
+		FillVariationInfo();
+	}
+}
+
 void UProceduralVegetationPreset::LoadFromVariantFiles(const TMap<FString, FString>& InVariantFiles)
 {
 	for (const auto& [VariantName, VariantFile] : InVariantFiles)
@@ -83,9 +117,21 @@ void UProceduralVegetationPreset::UpdateDataAsset()
 	}
 	
 	UpdateFoliageAndMaterialPath();
+	FillVariationInfo();
 	GetPackage()->SetDirtyFlag(true);
 }
 
+void UProceduralVegetationPreset::FillVariationInfo()
+{
+	PresetVariations.Empty();
+	for (const auto& [VariantName, VariantData] : Variants)
+	{
+		FPVPresetVariationInfo Info;
+		Info.Fill(FName(VariantName), VariantData);
+		PresetVariations.Emplace(Info);
+	}
+}
+
 void UProceduralVegetationPreset::CreateProfileDataAsset()
 {
 	FString ProfilePackageName(PlantProfileName);
diff --git a/Engine/Plugins/Experimental/ProceduralVegetationEditor/Source/ProceduralVegetation/Public/ProceduralVegetationPreset.h b/Engine/Plugins/Experimental/ProceduralVegetationEditor/Source/ProceduralVegetation/Public/ProceduralVegetationPreset.h
--- a/Engine/Plugins/Experimental/ProceduralVegetationEditor/Source/ProceduralVegetation/Public/ProceduralVegetationPreset.h
+++ b/Engine/Plugins/Experimental/ProceduralVegetationEditor/Source/ProceduralVegetation/Public/ProceduralVegetationPreset.h
@@ -8,6 +8,26 @@
 
 #include "ProceduralVegetationPreset.generated.h"
 
+USTRUCT()
+struct FPVPresetVariationInfo
+{
+	GENERATED_BODY()
+
+	UPROPERTY(VisibleAnywhere, Category = "Variations")
+	FName Name;
+
+	UPROPERTY(VisibleAnywhere, Category = "Variations")
+	TArray<TSoftObjectPtr<UObject>> FoliageMeshes;
+
+	UPROPERTY(VisibleAnywhere, Category = "Variations")
+	TArray<TSoftObjectPtr<UObject>> Materials;
+
+	UPROPERTY(VisibleAnywhere, Category = "Variations")
+	TArray<FString> PlantProfiles;
+
+	void Fill(FName InName,FManagedArrayCollection Collection);
+};
+
 UCLASS(BlueprintType)
 class PROCEDURALVEGETATION_API UProceduralVegetationPreset : public UDataAsset
 {
@@ -15,6 +35,8 @@ class PROCEDURALVEGETATION_API UProceduralVegetationPreset : public UDataAsset
 
 public:
 	UProceduralVegetationPreset();
+
+	virtual void PostLoad() override;
 	
 	UPROPERTY(EditAnywhere, Category="Preset Data", meta = (DevelopmentOnly))
 	FDirectoryPath JsonDirectoryPath;
@@ -37,15 +59,19 @@ public:
 	UPROPERTY(EditAnywhere, Category = "Preset Data", meta = (DevelopmentOnly, EditCondition = "bCreateProfileDataAsset"))
 	FString PlantProfileName;
 	
-	UPROPERTY(EditAnywhere, Category="Preset Data")
+	UPROPERTY()
 	TMap<FString, FManagedArrayCollection> Variants;
 
+	UPROPERTY(VisibleAnywhere, Category = "Preset Data")
+	TArray<FPVPresetVariationInfo> PresetVariations;
+
 public:
 	void LoadFromVariantFiles(const TMap<FString, FString>& InVariantFiles);
 
 	UFUNCTION(CallInEditor, Category="Preset Data", meta = (DevelopmentOnly))
 	void UpdateDataAsset();
-	
+
+	void FillVariationInfo();
 private:
 	void CreateProfileDataAsset();
 	
