[PVE] Fixed Adding foliage meshes to foliage palette do not contribute to foliage palette or distributor

diff --git a/Engine/Plugins/Experimental/ProceduralVegetationEditor/Source/ProceduralVegetation/Private/Nodes/PVFoliagePaletteSettings.cpp b/Engine/Plugins/Experimental/ProceduralVegetationEditor/Source/ProceduralVegetation/Private/Nodes/PVFoliagePaletteSettings.cpp
--- a/Engine/Plugins/Experimental/ProceduralVegetationEditor/Source/ProceduralVegetation/Private/Nodes/PVFoliagePaletteSettings.cpp
+++ b/Engine/Plugins/Experimental/ProceduralVegetationEditor/Source/ProceduralVegetation/Private/Nodes/PVFoliagePaletteSettings.cpp
@@ -92,7 +92,7 @@ bool FPVFoliageElement::ExecuteInternal(FPCGContext* InContext) const
 #endif
 			}
 
-			const int32 IterationLimit = FMath::Min(Settings->FoliageMeshes.Num(), Facade.NumFoliageNames());
+			const int32 IterationLimit = FMath::Max(Settings->FoliageMeshes.Num(), Facade.NumFoliageNames());
 			TArray<FString> FoliageNames;
 			for (int32 i = 0; i < IterationLimit; ++i)
 			{
