Fix PVE Export canceling sometimes leading to broken assets

- Improve warning notifying the user that the export will overwrite unsaved changes
- Fix issue where DoesConflictingPackageExist did not properly search the package for assets
- Fix issue where asset would not be deleted if the export failed if the asset was placed in the level

diff --git a/Engine/Plugins/Experimental/ProceduralVegetationEditor/Source/ProceduralVegetation/Private/Helpers/PVUtilities.cpp b/Engine/Plugins/Experimental/ProceduralVegetationEditor/Source/ProceduralVegetation/Private/Helpers/PVUtilities.cpp
--- a/Engine/Plugins/Experimental/ProceduralVegetationEditor/Source/ProceduralVegetation/Private/Helpers/PVUtilities.cpp
+++ b/Engine/Plugins/Experimental/ProceduralVegetationEditor/Source/ProceduralVegetation/Private/Helpers/PVUtilities.cpp
@@ -194,11 +194,10 @@ namespace PV::Utilities
 	bool DoesConflictingPackageExist(const FString& LongPackageName, UClass* AssetClass)
 	{
 		UPackage* Package = FindPackage(nullptr, *LongPackageName);
-		if (Package != nullptr)
+		if (IsValid(Package))
 		{
-			const FName AssetName = FPackageName::GetShortFName(Package->GetFName());
-			const bool bPackageIsOfSameType = StaticFindObjectFast(AssetClass, Package, AssetName) != nullptr;
-			return !bPackageIsOfSameType;
+			UObject* ExistingAsset = Package->FindAssetInPackage();
+			return !ExistingAsset || ExistingAsset->GetClass() != AssetClass;
 		}
 		else
 		{
diff --git a/Engine/Plugins/Experimental/ProceduralVegetationEditor/Source/ProceduralVegetationEditor/Private/Helpers/PVExportHelper.cpp b/Engine/Plugins/Experimental/ProceduralVegetationEditor/Source/ProceduralVegetationEditor/Private/Helpers/PVExportHelper.cpp
--- a/Engine/Plugins/Experimental/ProceduralVegetationEditor/Source/ProceduralVegetationEditor/Private/Helpers/PVExportHelper.cpp
+++ b/Engine/Plugins/Experimental/ProceduralVegetationEditor/Source/ProceduralVegetationEditor/Private/Helpers/PVExportHelper.cpp
@@ -4,9 +4,7 @@
 #include "Helpers/PVExportHelper.h"
 
 #include "DynamicMeshEditor.h"
-#if WITH_EDITOR
 #include "DynamicWindImportData.h"
-#endif
 #include "DynamicWindSkeletalData.h"
 #include "MeshDescriptionToDynamicMesh.h"
 #include "PackageTools.h"
@@ -20,6 +18,7 @@
 #include "SkinnedAssetCompiler.h"
 #include "StaticMeshCompiler.h"
 #include "UDynamicMesh.h"
+#include "ObjectTools.h"
 
 #include "Animation/Skeleton.h"
 
@@ -170,8 +169,34 @@ namespace PV::Export::Internal
 
 	void CleanupPackages(const TArray<UPackage*>& Packages)
 	{
-		FText UnloadError;
-		UPackageTools::UnloadPackages(Packages, UnloadError, true);
+		TArray<UPackage*> PackagesToReload;
+		PackagesToReload.Reserve(Packages.Num());
+
+		TArray<UObject*> ObjectsToDelete;
+		ObjectsToDelete.Reserve(Packages.Num());
+
+		for (UPackage* Package : Packages)
+		{
+			if (Package->HasAnyPackageFlags(PKG_InMemoryOnly))
+			{
+				ObjectsToDelete.Add(Package->FindAssetInPackage());
+			}
+			else
+			{
+				PackagesToReload.Add(Package);
+			}
+		}
+
+		if (PackagesToReload.Num() > 0)
+		{
+			FText UnloadError;
+			UPackageTools::ReloadPackages(PackagesToReload, UnloadError, EReloadPackagesInteractionMode::AssumePositive);
+		}
+		
+		if (ObjectsToDelete.Num() > 0)
+		{
+			ObjectTools::ForceDeleteObjects(ObjectsToDelete, false);
+		}
 	}
 
 	template<typename T>
@@ -1557,8 +1582,7 @@ namespace PV::Export
 
 			if (Package->HasAnyPackageFlags(PKG_NewlyCreated))
 			{
-				const FName AssetName = FPackageName::GetShortFName(Package->GetFName());
-				UObject* Asset = FindObjectFast<UObject>(Package, AssetName);
+				UObject* Asset = Package->FindAssetInPackage();
 				if (Asset)
 				{
 					IAssetRegistry::Get()->AssetCreated(Asset);
diff --git a/Engine/Plugins/Experimental/ProceduralVegetationEditor/Source/ProceduralVegetationEditor/Private/PVEditor.cpp b/Engine/Plugins/Experimental/ProceduralVegetationEditor/Source/ProceduralVegetationEditor/Private/PVEditor.cpp
--- a/Engine/Plugins/Experimental/ProceduralVegetationEditor/Source/ProceduralVegetationEditor/Private/PVEditor.cpp
+++ b/Engine/Plugins/Experimental/ProceduralVegetationEditor/Source/ProceduralVegetationEditor/Private/PVEditor.cpp
@@ -11,6 +11,7 @@
 #include "SPVEditorViewport.h"
 #include "ToolMenus.h"
 #include "ToolMenuSection.h"
+#include "PackagesDialog.h"
 
 #include "Dataflow/DataflowCollectionSpreadSheetWidget.h"
 
@@ -387,7 +388,7 @@ FReply FPVEditor::ExportInternal()
 		return !bHasInvalidPath;
 	};
 
-	const static auto ShouldOverwriteAssets = [](const TArray<UPCGNode*>& NodesToExport)->bool 
+	const static auto ShouldOverwriteUnsavedChanges = [](const TArray<UPCGNode*>& NodesToExport)->bool 
 	{
 		TArray<FString> ExportPaths;
 		ExportPaths.Reserve(NodesToExport.Num() * 2);
@@ -403,24 +404,47 @@ FReply FPVEditor::ExportInternal()
 			}
 		}
 
-		int32 NumExistingPackages = 0;
+		TArray<UPackage*> ExistingPackages;
+		ExistingPackages.Reserve(ExportPaths.Num());
 		for (const FString& AssetPath : ExportPaths)
 		{
-			NumExistingPackages += FPackageName::DoesPackageExist(AssetPath);
+			UPackage* ExistingPackage = LoadPackage(nullptr, *AssetPath, LOAD_None);
+			if (ExistingPackage)
+			{
+				ExistingPackages.Add(ExistingPackage);
+			}
 		}
 
-		if (NumExistingPackages == 0)
+		if (ExistingPackages.Num() == 0)
 		{
 			return true;
 		}
 
-		if (FMessageDialog::Open(EAppMsgType::YesNo,
-			FText::Format(LOCTEXT("PVE_OverwritePackagesWarning", "WARNING: {0} assets will be overwritten on export.\nDo you want to continue?"),
-				NumExistingPackages)) == EAppReturnType::No)
+		FPackagesDialogModule& PackagesDialogModule = FModuleManager::LoadModuleChecked<FPackagesDialogModule>(TEXT("PackagesDialog"));
+		PackagesDialogModule.CreatePackagesDialog(
+			LOCTEXT("PackagesDialog_OverwriteFiles_Title", "WARNING: Assets will be overwritten"), 
+			LOCTEXT("PackagesDialog_OverwriteFiles_Message", "The following assets will be overwritten. Would you like to continue?."),
+			/*InReadOnly=*/true
+		);
+		PackagesDialogModule.AddButton(
+			DRT_Save, 
+			LOCTEXT("PackagesDialog_OverwriteFiles_ContinueButton", "Continue"), 
+			LOCTEXT("PackagesDialog_OverwriteFiles_ButtonTip", "Continues the export, overwrites any existing assets")
+		);
+		PackagesDialogModule.AddButton(
+			DRT_Cancel, 
+			LOCTEXT("PackagesDialog_OverwriteFiles_CancelButton", "Cancel"), 
+			LOCTEXT("CancelDeleteButtonTip", "Cancel export")
+		);
+
+		for (UPackage* Package : ExistingPackages)
 		{
-			return false;
+			PackagesDialogModule.AddPackageItem(Package, ECheckBoxState::Checked);
 		}
-		return true;
+
+		const EDialogReturnType UserResponse = PackagesDialogModule.ShowPackagesDialog();
+
+		return UserResponse == DRT_Save;
 	};
 
 	const auto GetOutputNodeCollections = [this](TArray<FManagedArrayCollection>& OutCollections, const TArray<UPCGNode*>& NodesToExport, FScopedErrorMessages& ErrorMessages)->bool
@@ -508,7 +532,7 @@ FReply FPVEditor::ExportInternal()
 		return FReply::Handled();
 	}
 
-	if (!ShouldOverwriteAssets(NodesToExport))
+	if (!ShouldOverwriteUnsavedChanges(NodesToExport))
 	{
 		return FReply::Handled();
 	}
@@ -545,18 +569,21 @@ FReply FPVEditor::ExportInternal()
 		ProgressSlowTask.EnterProgressFrame(1.f, FText::Format(LOCTEXT("ExportingAsset", "Exporting \"{0}\"{1}"), FText::FromName(OutputSettings->ExportSettings.MeshName), AssetCountText));
 
 		FScopedSlowTask ExportCollectionSlowTask(1.f, LOCTEXT("GeneratingMesh", "Generating Mesh ..."));
-		float PrevFrameWork = 0.f;
 		float PrevProgress = 0.f;
 		const auto OnStatusUpdated = [&](const FText& Stage, float Progress)->bool
 		{
+			// TODO: Make each stage have its own loading bar with the top-level loading bar being for the whole export
+
 			const float WorkThisFrame = Progress - PrevProgress;
 			PrevProgress = Progress;
 
-			ExportCollectionSlowTask.EnterProgressFrame(PrevFrameWork, Stage);
+			ExportCollectionSlowTask.FrameMessage = Stage;
+
+			ExportCollectionSlowTask.CompletedWork += ExportCollectionSlowTask.CurrentFrameScope;
 
-			// Since our frames are so slow it does not make sense for the progress bar to progress as we start the frame, but rather as the frame finishes,
-			// hence why we repport the previours frame's work.
-			PrevFrameWork = WorkThisFrame;
+			const float WorkRemaining = ExportCollectionSlowTask.TotalAmountOfWork - ExportCollectionSlowTask.CompletedWork;
+			ExportCollectionSlowTask.CurrentFrameScope = FMath::Min(WorkRemaining, WorkThisFrame);
+			ExportCollectionSlowTask.ForceRefresh();
 
 			return !ProgressSlowTask.ShouldCancel();
 		};
@@ -576,7 +603,7 @@ FReply FPVEditor::ExportInternal()
 				for (int32 j = i; j < NodesToExport.Num(); j++)
 				{
 					MeshNames += CastChecked<const UPVOutputSettings>(NodesToExport[j]->GetSettings())->ExportSettings.MeshName.ToString();
-					if (j != NodesToExport.Num())
+					if (j != (NodesToExport.Num() - 1))
 					{
 						MeshNames += ", ";
 					}
@@ -599,8 +626,6 @@ FReply FPVEditor::ExportInternal()
 			
 			break;
 		}
-
-		ExportCollectionSlowTask.EnterProgressFrame(PrevFrameWork); // Let the slow task know all work has been completed.
 	}
 	
 	return FReply::Handled();
diff --git a/Engine/Plugins/Experimental/ProceduralVegetationEditor/Source/ProceduralVegetationEditor/ProceduralVegetationEditor.Build.cs b/Engine/Plugins/Experimental/ProceduralVegetationEditor/Source/ProceduralVegetationEditor/ProceduralVegetationEditor.Build.cs
--- a/Engine/Plugins/Experimental/ProceduralVegetationEditor/Source/ProceduralVegetationEditor/ProceduralVegetationEditor.Build.cs
+++ b/Engine/Plugins/Experimental/ProceduralVegetationEditor/Source/ProceduralVegetationEditor/ProceduralVegetationEditor.Build.cs
@@ -34,6 +34,7 @@ public class ProceduralVegetationEditor : ModuleRules
 				"MeshDescription",
 				"MeshConversion",
 				"ModelingComponents",
+				"PackagesDialog",
 				"PCG",
 				"PCGEditor",
 				"ProceduralVegetation",
