Disable cancel button for PVE export

diff --git a/Engine/Plugins/Experimental/ProceduralVegetationEditor/Source/ProceduralVegetationEditor/Private/PVEditor.cpp b/Engine/Plugins/Experimental/ProceduralVegetationEditor/Source/ProceduralVegetationEditor/Private/PVEditor.cpp
--- a/Engine/Plugins/Experimental/ProceduralVegetationEditor/Source/ProceduralVegetationEditor/Private/PVEditor.cpp
+++ b/Engine/Plugins/Experimental/ProceduralVegetationEditor/Source/ProceduralVegetationEditor/Private/PVEditor.cpp
@@ -549,7 +549,7 @@ FReply FPVEditor::ExportInternal()
 	}
 
 	FScopedSlowTask ProgressSlowTask(NodesToExport.Num(), LOCTEXT("ExportAsset_SlowTask", "Exporting Asset(s) ..."));
-	ProgressSlowTask.MakeDialog(true);
+	ProgressSlowTask.MakeDialog(false);
 
 	for (int32 i = 0; i < NodesToExport.Num(); ++i)
 	{
@@ -572,8 +572,6 @@ FReply FPVEditor::ExportInternal()
 		float PrevProgress = 0.f;
 		const auto OnStatusUpdated = [&](const FText& Stage, float Progress)->bool
 		{
-			// TODO: Make each stage have its own loading bar with the top-level loading bar being for the whole export
-
 			const float WorkThisFrame = Progress - PrevProgress;
 			PrevProgress = Progress;
 
