[PVE] Auto focus viewport settings

diff --git a/Engine/Plugins/Experimental/ProceduralVegetationEditor/Source/ProceduralVegetationEditor/Private/PVEditorCommands.cpp b/Engine/Plugins/Experimental/ProceduralVegetationEditor/Source/ProceduralVegetationEditor/Private/PVEditorCommands.cpp
--- a/Engine/Plugins/Experimental/ProceduralVegetationEditor/Source/ProceduralVegetationEditor/Private/PVEditorCommands.cpp
+++ b/Engine/Plugins/Experimental/ProceduralVegetationEditor/Source/ProceduralVegetationEditor/Private/PVEditorCommands.cpp
@@ -20,6 +20,7 @@ void FPVEditorCommands::RegisterCommands()
 	UI_COMMAND(Export, "Export", "Export outputs.", EUserInterfaceActionType::Button, FInputChord(EModifierKey::Control, EKeys::E));
 	UI_COMMAND(ShowMannequin, "Mannequin", "Show/Hide Mannequin", EUserInterfaceActionType::ToggleButton, FInputChord(EModifierKey::Control, EKeys::M));
 	UI_COMMAND(ShowScaleVisualization, "Scale Visualization", "Show/Hide Scale Visualization", EUserInterfaceActionType::ToggleButton, FInputChord(EModifierKey::Control, EKeys::V));
+	UI_COMMAND(AutoFocusViewport, "Auto focus viewport", "Auto focus viewport on node selection and visualization mode switch", EUserInterfaceActionType::ToggleButton, FInputChord());
 }
 
 #undef LOCTEXT_NAMESPACE
diff --git a/Engine/Plugins/Experimental/ProceduralVegetationEditor/Source/ProceduralVegetationEditor/Private/PVEditorCommands.h b/Engine/Plugins/Experimental/ProceduralVegetationEditor/Source/ProceduralVegetationEditor/Private/PVEditorCommands.h
--- a/Engine/Plugins/Experimental/ProceduralVegetationEditor/Source/ProceduralVegetationEditor/Private/PVEditorCommands.h
+++ b/Engine/Plugins/Experimental/ProceduralVegetationEditor/Source/ProceduralVegetationEditor/Private/PVEditorCommands.h
@@ -19,4 +19,5 @@ public:
 	//Viewport Commands
 	TSharedPtr<FUICommandInfo> ShowMannequin;
 	TSharedPtr<FUICommandInfo> ShowScaleVisualization;
+	TSharedPtr<FUICommandInfo> AutoFocusViewport;
 };
\ No newline at end of file
diff --git a/Engine/Plugins/Experimental/ProceduralVegetationEditor/Source/ProceduralVegetationEditor/Private/SPVEditorViewport.cpp b/Engine/Plugins/Experimental/ProceduralVegetationEditor/Source/ProceduralVegetationEditor/Private/SPVEditorViewport.cpp
--- a/Engine/Plugins/Experimental/ProceduralVegetationEditor/Source/ProceduralVegetationEditor/Private/SPVEditorViewport.cpp
+++ b/Engine/Plugins/Experimental/ProceduralVegetationEditor/Source/ProceduralVegetationEditor/Private/SPVEditorViewport.cpp
@@ -87,6 +87,12 @@ void SPVEditorViewport::BindCommands()
 		FExecuteAction::CreateSP(this, &SPVEditorViewport::ToggleShowScaleVis),
 		FCanExecuteAction(),
 		FIsActionChecked::CreateSP(this, &SPVEditorViewport::IsShowScaleVisChecked));
+
+	CommandList->MapAction(
+		Commands.AutoFocusViewport,
+		FExecuteAction::CreateSP(this, &SPVEditorViewport::ToggleAutofocusViewport),
+		FCanExecuteAction(),
+		FIsActionChecked::CreateSP(this, &SPVEditorViewport::IsAutoFocusViewportChecked));
 }
 
 void SPVEditorViewport::InitVisualizationScene()
@@ -211,6 +217,18 @@ void SPVEditorViewport::SetScaleVisState(bool InEnable)
 	}
 }
 
+void SPVEditorViewport::ToggleAutofocusViewport()
+{
+	UPVEditorSettings* EditorSettings = GetMutableDefault<UPVEditorSettings>();
+	EditorSettings->bAutoFocusViewport = !EditorSettings->bAutoFocusViewport;
+	EditorSettings->SaveConfig();
+}
+
+bool SPVEditorViewport::IsAutoFocusViewportChecked() const
+{
+	return GetDefault<UPVEditorSettings>()->bAutoFocusViewport;
+}
+
 float SPVEditorViewport::GetMannequinOffset() const
 {
 	return GetDefault<UPVEditorSettings>()->MannequinOffset;
@@ -297,6 +315,10 @@ FToolMenuEntry SPVEditorViewport::CreateSettingsToolbarMenu()
 						FToolMenuSection& ScaleVisSection = Submenu->FindOrAddSection(TEXT("PVE_ScaleVisSettings"),
 							LOCTEXT("PVE_ScaleVisSettingsSectionLabel", "Scale Vis. Settings"));
 						ScaleVisSection.AddMenuEntry(FPVEditorCommands::Get().ShowScaleVisualization);
+
+						FToolMenuSection& AutoFocusViewportSection = Submenu->FindOrAddSection(TEXT("PVE_InteractionSettings"),
+							LOCTEXT("PVE_InteractionSettingsSectionLabel", "Interaction Settings"));
+						AutoFocusViewportSection.AddMenuEntry(FPVEditorCommands::Get().AutoFocusViewport);
 					}),
 				/*bInOpenSubMenuOnClick=*/false,
 				FSlateIcon(FAppStyle::GetAppStyleSetName(), "LevelEditor.GameSettings"));
@@ -465,7 +487,10 @@ void SPVEditorViewport::OnVisualizationModeChanged(EPVVisualizationMode InMode)
 
 		Invalidate();
 
-		bFocusOnNextUpdate = true;
+		if (IsAutoFocusViewportChecked())
+		{
+			bFocusOnNextUpdate = true;
+		}
 	}
 }
 
@@ -477,7 +502,10 @@ void SPVEditorViewport::OnNodeSelectionChanged(UPVBaseSettings* InSettings)
 	
 	CurrentNodeSettings->SetCurrentRenderType(InSettings->GetDefaultRenderType());
 
-	bFocusOnNextUpdate = true;
+	if (IsAutoFocusViewportChecked())
+	{
+		bFocusOnNextUpdate = true;
+	}
 }
 
 #undef LOCTEXT_NAMESPACE
diff --git a/Engine/Plugins/Experimental/ProceduralVegetationEditor/Source/ProceduralVegetationEditor/Private/SPVEditorViewport.h b/Engine/Plugins/Experimental/ProceduralVegetationEditor/Source/ProceduralVegetationEditor/Private/SPVEditorViewport.h
--- a/Engine/Plugins/Experimental/ProceduralVegetationEditor/Source/ProceduralVegetationEditor/Private/SPVEditorViewport.h
+++ b/Engine/Plugins/Experimental/ProceduralVegetationEditor/Source/ProceduralVegetationEditor/Private/SPVEditorViewport.h
@@ -65,6 +65,9 @@ private:
 	bool IsShowScaleVisChecked() const;
 	void SetScaleVisState(bool InEnable);
 
+	void ToggleAutofocusViewport();
+	bool IsAutoFocusViewportChecked() const;
+
 	float GetMannequinOffset() const;
 	void SetMannequinOffset(float NewValue, bool bSaveConfig = false) const;
 
diff --git a/Engine/Plugins/Experimental/ProceduralVegetationEditor/Source/ProceduralVegetationEditor/Public/PVEditorSettings.h b/Engine/Plugins/Experimental/ProceduralVegetationEditor/Source/ProceduralVegetationEditor/Public/PVEditorSettings.h
--- a/Engine/Plugins/Experimental/ProceduralVegetationEditor/Source/ProceduralVegetationEditor/Public/PVEditorSettings.h
+++ b/Engine/Plugins/Experimental/ProceduralVegetationEditor/Source/ProceduralVegetationEditor/Public/PVEditorSettings.h
@@ -23,12 +23,15 @@ public:
 	UPROPERTY(EditAnywhere, config, Category = Node, meta = (HideAlphaChannel))
 	FLinearColor FoliageMeshDataPinColor = FColor(0, 255, 0);
 
-	UPROPERTY(EditAnywhere, config, Category = Debug)
+	UPROPERTY(EditAnywhere, config, Category = "Viewport")
 	bool bShowMannequin = false;
 
-	UPROPERTY(EditAnywhere, config, Category = Debug)
+	UPROPERTY(EditAnywhere, config, Category = "Viewport")
 	float MannequinOffset = 0.0f;
 
-	UPROPERTY(EditAnywhere, config, Category = Debug)
+	UPROPERTY(EditAnywhere, config, Category = "Viewport")
 	bool bShowScaleVisualization = false;
+
+	UPROPERTY(EditAnywhere, config, Category = "Viewport")
+	bool bAutoFocusViewport = true;
 };
