[PVE] Added ability to export selected node

diff --git a/Engine/Plugins/Experimental/ProceduralVegetationEditor/Source/ProceduralVegetationEditor/Private/PVEditor.cpp b/Engine/Plugins/Experimental/ProceduralVegetationEditor/Source/ProceduralVegetationEditor/Private/PVEditor.cpp
--- a/Engine/Plugins/Experimental/ProceduralVegetationEditor/Source/ProceduralVegetationEditor/Private/PVEditor.cpp
+++ b/Engine/Plugins/Experimental/ProceduralVegetationEditor/Source/ProceduralVegetationEditor/Private/PVEditor.cpp
@@ -8,6 +8,7 @@
 #include "ProceduralVegetation.h"
 #include "PVEditorCommands.h"
 #include "PVEditorSchema.h"
+#include "PVEditorSettings.h"
 #include "SPVEditorViewport.h"
 #include "ToolMenus.h"
 #include "ToolMenuSection.h"
@@ -291,7 +292,7 @@ void FPVEditor::OnSelectedNodesChanged(const TSet<UObject*>& InNewSelection)
 {
 	FPCGEditor::OnSelectedNodesChanged(InNewSelection);
 
-	const UPCGNode* PCGNode = nullptr;
+	UPCGNode* SelectedNode = nullptr;
 	
 	if (!InNewSelection.IsEmpty())
 	{
@@ -299,11 +300,11 @@ void FPVEditor::OnSelectedNodesChanged(const TSet<UObject*>& InNewSelection)
 
 		if (Node && Node->GetPCGNode())
 		{
-			PCGNode = Node->GetPCGNode();
+			SelectedNode = Node->GetPCGNode();
 
-			UpdateCollectionSpreadSheet(PCGNode);
+			UpdateCollectionSpreadSheet(SelectedNode);
 
-			if (UPVBaseSettings* Settings = Cast<UPVBaseSettings>(PCGNode->GetSettings()))
+			if (UPVBaseSettings* Settings = Cast<UPVBaseSettings>(SelectedNode->GetSettings()))
 			{
 				Settings->bAutoFocusViewport = bFocusViewportOnSelection;
 				bFocusViewportOnSelection = false;
@@ -317,6 +318,8 @@ void FPVEditor::OnSelectedNodesChanged(const TSet<UObject*>& InNewSelection)
 			SetNodeInspected(Node, true);
 		}
 	}
+
+	CurrentSelectedNode = SelectedNode;
 }
 
 void FPVEditor::OnExport_Clicked()
@@ -331,7 +334,8 @@ FReply FPVEditor::ExportInternal()
 	TSharedPtr<SPVExportSelectionDialog> ExportSelectionWidget =
 		SNew(SPVExportSelectionDialog)
 		.Title(LOCTEXT("Export Selection", "Choose outputs to export"))
-		.Graph(ProceduralVegetationGraph);
+		.Graph(ProceduralVegetationGraph)
+		.SelectedNode(CurrentSelectedNode);
 
 	if (ExportSelectionWidget->ShowModal() != EAppReturnType::Ok)
 	{
@@ -522,10 +526,27 @@ FReply FPVEditor::ExportInternal()
 		return true;
 	};
 
-	const TArray<UPCGNode*> NodesToExport = ProceduralVegetationGraph->GetNodes().FilterByPredicate([](UPCGNode* InNode)->bool {
+	TArray<UPCGNode*> NodesToExport;
+
+	const UPVEditorSettings* EditorSettings = GetDefault<UPVEditorSettings>();
+
+	if (EditorSettings->bExportType == 0)
+	{
+		if (CurrentSelectedNode != nullptr)
+		{
+			if (const UPVOutputSettings* OutputSettings = Cast<UPVOutputSettings>(CurrentSelectedNode->GetSettings()))
+			{
+				NodesToExport.Add(CurrentSelectedNode);
+			}
+		}
+	}
+	else
+	{
+		NodesToExport = ProceduralVegetationGraph->GetNodes().FilterByPredicate([](UPCGNode* InNode)->bool {
 			const UPVOutputSettings* OutputSettings = Cast<const UPVOutputSettings>(InNode->GetSettings());
 			return OutputSettings && OutputSettings->ExportSettings.bShouldExport;
 		});
+	}
 
 	if (!ValidateExportSettings(NodesToExport, ErrorMessages))
 	{
diff --git a/Engine/Plugins/Experimental/ProceduralVegetationEditor/Source/ProceduralVegetationEditor/Private/PVEditor.h b/Engine/Plugins/Experimental/ProceduralVegetationEditor/Source/ProceduralVegetationEditor/Private/PVEditor.h
--- a/Engine/Plugins/Experimental/ProceduralVegetationEditor/Source/ProceduralVegetationEditor/Private/PVEditor.h
+++ b/Engine/Plugins/Experimental/ProceduralVegetationEditor/Source/ProceduralVegetationEditor/Private/PVEditor.h
@@ -50,7 +50,8 @@ protected:
 
 private:
 	TObjectPtr<UProceduralVegetation> ProceduralVegetationBeingEdited = nullptr;
-
+	UPCGNode* CurrentSelectedNode = nullptr;
+	
 	TArray<TSharedRef<SPVEditorViewport>> ViewportWidgets;
 	
 	FDateTime SessionStartTime;
diff --git a/Engine/Plugins/Experimental/ProceduralVegetationEditor/Source/ProceduralVegetationEditor/Private/SPVEditorViewport.h b/Engine/Plugins/Experimental/ProceduralVegetationEditor/Source/ProceduralVegetationEditor/Private/SPVEditorViewport.h
--- a/Engine/Plugins/Experimental/ProceduralVegetationEditor/Source/ProceduralVegetationEditor/Private/SPVEditorViewport.h
+++ b/Engine/Plugins/Experimental/ProceduralVegetationEditor/Source/ProceduralVegetationEditor/Private/SPVEditorViewport.h
@@ -19,6 +19,7 @@ enum class EPVVisualizationMode : uint8
 	PointDataMesh	UMETA(DisplayName = "Point data + Mesh"),
 	FoliageMesh		UMETA(DisplayName = "Foliage + Mesh"),
 	BonesMesh		UMETA(DisplayName = "Bones + Mesh"),
+	Count			UMETA(Hidden)
 };
 
 class UPVBaseSettings;
diff --git a/Engine/Plugins/Experimental/ProceduralVegetationEditor/Source/ProceduralVegetationEditor/Private/Widgets/SPVExportSelectionDialog.cpp b/Engine/Plugins/Experimental/ProceduralVegetationEditor/Source/ProceduralVegetationEditor/Private/Widgets/SPVExportSelectionDialog.cpp
--- a/Engine/Plugins/Experimental/ProceduralVegetationEditor/Source/ProceduralVegetationEditor/Private/Widgets/SPVExportSelectionDialog.cpp
+++ b/Engine/Plugins/Experimental/ProceduralVegetationEditor/Source/ProceduralVegetationEditor/Private/Widgets/SPVExportSelectionDialog.cpp
@@ -6,7 +6,9 @@
 #include "Editor.h"
 #include "SPVExportItem.h"
 #include "ProceduralVegetation.h"
+#include "PVEditorSettings.h"
 #include "Nodes/PVOutputSettings.h"
+#include "Widgets/Input/STextComboBox.h"
 #include "Widgets/Layout/SScrollBox.h"
 
 #define LOCTEXT_NAMESPACE "SPVExportSelectionDialog"
@@ -14,6 +16,9 @@
 void SPVExportSelectionDialog::Construct(const FArguments& InArgs)
 {
 	Graph = InArgs._Graph;
+	SelectedNode = InArgs._SelectedNode;
+
+	BuildExportTypeOptions();
 	
 	SWindow::Construct(SWindow::FArguments()
 		.Title(InArgs._Title)
@@ -22,6 +27,35 @@ void SPVExportSelectionDialog::Construct(const FArguments& InArgs)
 		.ClientSize(FVector2D(350, 450))
 		[
 			SNew(SVerticalBox)
+			+ SVerticalBox::Slot()
+			.AutoHeight()
+			.Padding(2, 2, 2, 0)
+			[
+				SNew(SBorder)
+				.BorderImage(FAppStyle::GetBrush("ToolPanel.GroupBorder"))
+				[
+					SNew(SHorizontalBox)
+					+ SHorizontalBox::Slot()
+					.HAlign(HAlign_Left)
+					.VAlign(VAlign_Center)
+					.AutoWidth()
+					.Padding(2, 10, 20, 10)
+					[
+						SNew(STextBlock)
+						.Text(LOCTEXT("PVExportSelectionDialog_ExportTypeLabel", "Export type"))
+					]
+					
+					+ SHorizontalBox::Slot()
+					.FillWidth(1.0f)
+					.HAlign(HAlign_Fill)
+					[
+						SNew(STextComboBox)
+						.OptionsSource(&ExportTypes)
+						.InitiallySelectedItem(CurrentExportType)
+						.OnSelectionChanged(this, &SPVExportSelectionDialog::OnExportTypeChanged)
+					]
+				]
+			]
 			+ SVerticalBox::Slot() // Add user input block
 			.Padding(2, 2, 2, 4)
 			[
@@ -48,7 +82,7 @@ void SPVExportSelectionDialog::Construct(const FArguments& InArgs)
 				.SlotPadding(FAppStyle::GetMargin("StandardDialog.SlotPadding"))
 				+ SUniformGridPanel::Slot(0, 0)
 				[
-					SNew(SPrimaryButton)
+					SAssignNew(ExportButton, SPrimaryButton)
 					.Text(LOCTEXT("Export", "Export"))
 					.OnClicked(this, &SPVExportSelectionDialog::OnButtonClick, EAppReturnType::Ok)
 				]
@@ -61,37 +95,75 @@ void SPVExportSelectionDialog::Construct(const FArguments& InArgs)
 				]
 			]
 		]);
+
 	AddExportItems();
 }
+
+EPVExportType SPVExportSelectionDialog::GetExportType()
+{
+	int SelectionIndex = ExportTypes.IndexOfByKey(CurrentExportType);
+	auto ExportType = (EPVExportType)SelectionIndex;
+
+	return ExportType;
+}
+
 void SPVExportSelectionDialog::AddExportItems()
 {
 	ScrollBox->ClearChildren();
 	
-	Graph->ForEachNode([this](UPCGNode* Node)
+	if (GetExportType() == EPVExportType::BatchExport)
 	{
-		const UPVOutputSettings* OutputSettings = Cast<UPVOutputSettings>(Node->GetSettings());
+		Graph->ForEachNode([this](UPCGNode* Node)
+		{
+			const UPVOutputSettings* OutputSettings = Cast<UPVOutputSettings>(Node->GetSettings());
 
-		if (OutputSettings)
+			if (OutputSettings)
+			{
+				FText OutputName = Node->GetNodeTitle(EPCGNodeTitleType::FullTitle);
+				
+				ScrollBox->AddSlot()
+				.Padding(5)
+				[
+					SNew(SPVExportItem)
+					.Name(OutputName)
+					.OnOutputSelectionChanged(this, &SPVExportSelectionDialog::OnOutputSelectionChanged)
+					.bIsSelected(OutputSettings->ExportSettings.bShouldExport)
+				];
+				ScrollBox->AddSlot()
+				[
+					SNew(SSeparator)
+					.Thickness(1)
+				];
+			}
+
+			return true;
+		});
+	}
+	else
+	{
+		FString OutputName = "None";
+		
+		if (SelectedNode)
 		{
-			FText OutputName = Node->GetNodeTitle(EPCGNodeTitleType::FullTitle);
-	
-			ScrollBox->AddSlot()
-			.Padding(5)
-			[
-				SNew(SPVExportItem)
-				.Name(OutputName)
-				.OnOutputSelectionChanged(this, &SPVExportSelectionDialog::OnOutputSelectionChanged)
-				.bIsSelected(OutputSettings->ExportSettings.bShouldExport)
-			];
-			ScrollBox->AddSlot()
-			[
-				SNew(SSeparator)
-				.Thickness(1)
-			];
+			const UPVOutputSettings* OutputSettings = Cast<UPVOutputSettings>(SelectedNode->GetSettings());
+
+			if (OutputSettings)
+			{
+				OutputName = SelectedNode->GetNodeTitle(EPCGNodeTitleType::FullTitle).ToString();
+			}	
 		}
+		
+		FText Label = FText::FromString(FString::Printf(TEXT("Current output: %s"), *OutputName));
+			
+		ScrollBox->AddSlot()
+		.Padding(5)
+		[
+			SNew(STextBlock)
+			.Text(Label)
+		];
+	}
 
-		return true;
-	});
+	UpdateExportButtonState();
 }
 FReply SPVExportSelectionDialog::OnButtonClick(EAppReturnType::Type ButtonID)
 {
@@ -108,6 +180,58 @@ FReply SPVExportSelectionDialog::OnButtonClick(EAppReturnType::Type ButtonID)
 	}
 	return FReply::Handled();
 }
+
+void SPVExportSelectionDialog::OnExportTypeChanged(TSharedPtr<FString> Selection, ESelectInfo::Type SelectInfo)
+{
+	if (!Selection->IsEmpty())
+	{
+		int SelectionIndex = ExportTypes.IndexOfByKey(Selection);
+
+		CurrentExportType = ExportTypes[SelectionIndex];
+		
+		UPVEditorSettings* EditorSettings = GetMutableDefault<UPVEditorSettings>();
+		EditorSettings->bExportType = SelectionIndex;
+
+		AddExportItems();
+	}
+}
+
+void SPVExportSelectionDialog::UpdateExportButtonState()
+{
+	bool HasValidExport = false;
+
+	if (GetExportType() == EPVExportType::BatchExport)
+	{
+		if (Graph)
+		{
+			Graph->ForEachNode([this, &HasValidExport](UPCGNode* Node)
+			{
+				if (UPVOutputSettings* OutputSettings = Cast<UPVOutputSettings>(Node->GetSettings()))
+				{
+					HasValidExport |= OutputSettings->ExportSettings.bShouldExport;
+				}
+		
+				return true;
+			});	
+		}	
+	}
+	else
+	{
+		if (SelectedNode)
+		{
+			const UPVOutputSettings* OutputSettings = Cast<UPVOutputSettings>(SelectedNode->GetSettings());
+
+			if (OutputSettings)
+			{
+				HasValidExport = true;
+			}	
+		}
+	}
+
+	ExportButton->SetEnabled(HasValidExport);
+	
+}
+
 EAppReturnType::Type SPVExportSelectionDialog::ShowModal()
 {
 	GEditor->EditorAddModalWindow(SharedThis(this));
@@ -127,5 +251,31 @@ void SPVExportSelectionDialog::OnOutputSelectionChanged(const FString ItemName,
 		
 		return true;
 	});
+
+	UpdateExportButtonState();
+}
+
+void SPVExportSelectionDialog::BuildExportTypeOptions()
+{
+	UEnum* EnumPtr = StaticEnum<EPVExportType>();
+	                      
+	if (EnumPtr)
+	{
+		for (int32 i = 0; i < EnumPtr->NumEnums() - 1; i++) // last entry is _MAX or hidden
+		{
+			FText ModeName = EnumPtr->GetDisplayNameTextByIndex(i);
+			ExportTypes.Add(MakeShared<FString>(ModeName.ToString()));
+		}
+	}
+
+	if (ExportTypes.Num() > 0)
+	{
+		const UPVEditorSettings* EditorSettings = GetDefault<UPVEditorSettings>();
+
+		if (EditorSettings->bExportType >= 0 && EditorSettings->bExportType < ExportTypes.Num())
+		{
+			CurrentExportType = ExportTypes[EditorSettings->bExportType];
+		}
+	}
 }
-#undef LOCTEXT_NAMESPACE
\ No newline at end of file
+#undef LOCTEXT_NAMESPACE
diff --git a/Engine/Plugins/Experimental/ProceduralVegetationEditor/Source/ProceduralVegetationEditor/Private/Widgets/SPVExportSelectionDialog.h b/Engine/Plugins/Experimental/ProceduralVegetationEditor/Source/ProceduralVegetationEditor/Private/Widgets/SPVExportSelectionDialog.h
--- a/Engine/Plugins/Experimental/ProceduralVegetationEditor/Source/ProceduralVegetationEditor/Private/Widgets/SPVExportSelectionDialog.h
+++ b/Engine/Plugins/Experimental/ProceduralVegetationEditor/Source/ProceduralVegetationEditor/Private/Widgets/SPVExportSelectionDialog.h
@@ -3,9 +3,19 @@
 #include "CoreMinimal.h"
 #include "Input/Reply.h"
 #include "Widgets/SWindow.h"
+#include "SPVExportSelectionDialog.generated.h"
 
+class UPVOutputSettings;
 class SScrollBox;
 class UProceduralVegetationGraph;
+class UPCGNode;
+class SPrimaryButton;
+UENUM()
+enum class EPVExportType
+{
+	Selection	UMETA(DisplayName = "Selected output node"),
+	BatchExport UMETA(DisplayName = "Batch export")
+};
 
 class SPVExportSelectionDialog : public SWindow
 {
@@ -14,6 +24,7 @@ class SPVExportSelectionDialog : public SWindow
 	SLATE_ARGUMENT(FText, Title)
 	SLATE_ARGUMENT(const UProceduralVegetationGraph*, Graph)
 	SLATE_ARGUMENT(TArray<FText>, ExportItems)
+	SLATE_ARGUMENT(UPCGNode*, SelectedNode)
 	SLATE_END_ARGS()
 
 public:
@@ -25,11 +36,23 @@ public:
 protected:
 	FReply OnButtonClick(EAppReturnType::Type ButtonID);
 	EAppReturnType::Type UserResponse = EAppReturnType::Type::Cancel;
-
+	
 private:
 	TSharedPtr<SScrollBox> ScrollBox;
+	TSharedPtr<SPrimaryButton> ExportButton;
 	const UProceduralVegetationGraph* Graph = nullptr;
+	const UPCGNode* SelectedNode = nullptr;
 
+	TArray<TSharedPtr<FString>> ExportTypes;
+	TSharedPtr<FString> CurrentExportType;
+
+	EPVExportType GetExportType();
+	
 	void AddExportItems();
 	void OnOutputSelectionChanged(const FString ItemName, ECheckBoxState NewState);
+
+	void BuildExportTypeOptions();
+	void OnExportTypeChanged(TSharedPtr<FString> Selection, ESelectInfo::Type SelectInfo);
+
+	void UpdateExportButtonState();
 };
diff --git a/Engine/Plugins/Experimental/ProceduralVegetationEditor/Source/ProceduralVegetationEditor/Public/PVEditorSettings.h b/Engine/Plugins/Experimental/ProceduralVegetationEditor/Source/ProceduralVegetationEditor/Public/PVEditorSettings.h
--- a/Engine/Plugins/Experimental/ProceduralVegetationEditor/Source/ProceduralVegetationEditor/Public/PVEditorSettings.h
+++ b/Engine/Plugins/Experimental/ProceduralVegetationEditor/Source/ProceduralVegetationEditor/Public/PVEditorSettings.h
@@ -23,15 +23,18 @@ public:
 	UPROPERTY(EditAnywhere, config, Category = Node, meta = (HideAlphaChannel))
 	FLinearColor FoliageMeshDataPinColor = FColor(0, 255, 0);
 
-	UPROPERTY(EditAnywhere, config, Category = "Viewport")
+	UPROPERTY(EditAnywhere, config, Category = Viewport)
 	bool bShowMannequin = false;
 
-	UPROPERTY(EditAnywhere, config, Category = "Viewport")
+	UPROPERTY(EditAnywhere, config, Category = Viewport)
 	float MannequinOffset = 0.0f;
 
-	UPROPERTY(EditAnywhere, config, Category = "Viewport")
+	UPROPERTY(EditAnywhere, config, Category = Viewport)
 	bool bShowScaleVisualization = false;
 
-	UPROPERTY(EditAnywhere, config, Category = "Viewport")
+	UPROPERTY(EditAnywhere, config, Category = Export)
 	bool bAutoFocusViewport = true;
+
+	UPROPERTY(EditAnywhere, config, Category = Settings)
+	int bExportType = 0;
 };
