Rename RenderType::Skeleton to PointData

diff --git a/Engine/Plugins/Experimental/ProceduralVegetationEditor/Source/ProceduralVegetation/Private/Nodes/PVCarveSettings.h b/Engine/Plugins/Experimental/ProceduralVegetationEditor/Source/ProceduralVegetation/Private/Nodes/PVCarveSettings.h
--- a/Engine/Plugins/Experimental/ProceduralVegetationEditor/Source/ProceduralVegetation/Private/Nodes/PVCarveSettings.h
+++ b/Engine/Plugins/Experimental/ProceduralVegetationEditor/Source/ProceduralVegetation/Private/Nodes/PVCarveSettings.h
@@ -21,7 +21,7 @@ public:
 	virtual FText GetNodeTooltipText() const override;
 	virtual FLinearColor GetNodeTitleColor() const override { return FLinearColor::Yellow; }
 
-	virtual TArray<EPVRenderType> GetDefaultRenderType() const override { return TArray{EPVRenderType::Skeleton}; }
+	virtual TArray<EPVRenderType> GetDefaultRenderType() const override { return TArray{EPVRenderType::PointData}; }
 #endif
 
 protected:
diff --git a/Engine/Plugins/Experimental/ProceduralVegetationEditor/Source/ProceduralVegetation/Private/Nodes/PVGravitySettings.h b/Engine/Plugins/Experimental/ProceduralVegetationEditor/Source/ProceduralVegetation/Private/Nodes/PVGravitySettings.h
--- a/Engine/Plugins/Experimental/ProceduralVegetationEditor/Source/ProceduralVegetation/Private/Nodes/PVGravitySettings.h
+++ b/Engine/Plugins/Experimental/ProceduralVegetationEditor/Source/ProceduralVegetation/Private/Nodes/PVGravitySettings.h
@@ -18,7 +18,7 @@ public:
 	virtual FText GetNodeTooltipText() const override;
 	virtual FLinearColor GetNodeTitleColor() const override { return FLinearColor::Yellow; }
 	
-	virtual TArray<EPVRenderType> GetDefaultRenderType() const override { return TArray{ EPVRenderType::Skeleton }; }
+	virtual TArray<EPVRenderType> GetDefaultRenderType() const override { return TArray{ EPVRenderType::PointData }; }
 #endif
 
 protected:
diff --git a/Engine/Plugins/Experimental/ProceduralVegetationEditor/Source/ProceduralVegetation/Private/Nodes/PVImporterSettings.h b/Engine/Plugins/Experimental/ProceduralVegetationEditor/Source/ProceduralVegetation/Private/Nodes/PVImporterSettings.h
--- a/Engine/Plugins/Experimental/ProceduralVegetationEditor/Source/ProceduralVegetation/Private/Nodes/PVImporterSettings.h
+++ b/Engine/Plugins/Experimental/ProceduralVegetationEditor/Source/ProceduralVegetation/Private/Nodes/PVImporterSettings.h
@@ -21,7 +21,7 @@ public:
 	virtual FText GetNodeTooltipText() const override { return NSLOCTEXT("PVImporterSettings", "NodeTooltip", "Create Procedural Vegetation Skeleton"); }
 	virtual FLinearColor GetNodeTitleColor() const override { return FLinearColor(FColor::Black); }
 
-	virtual TArray<EPVRenderType> GetDefaultRenderType() const override { return TArray{ EPVRenderType::Skeleton }; }
+	virtual TArray<EPVRenderType> GetDefaultRenderType() const override { return TArray{ EPVRenderType::PointData }; }
 #endif
 
 protected:
diff --git a/Engine/Plugins/Experimental/ProceduralVegetationEditor/Source/ProceduralVegetation/Private/Nodes/PVMeshBuilderSettings.h b/Engine/Plugins/Experimental/ProceduralVegetationEditor/Source/ProceduralVegetation/Private/Nodes/PVMeshBuilderSettings.h
--- a/Engine/Plugins/Experimental/ProceduralVegetationEditor/Source/ProceduralVegetation/Private/Nodes/PVMeshBuilderSettings.h
+++ b/Engine/Plugins/Experimental/ProceduralVegetationEditor/Source/ProceduralVegetation/Private/Nodes/PVMeshBuilderSettings.h
@@ -21,7 +21,7 @@ public:
 	virtual FText GetNodeTooltipText() const override;
 	virtual FLinearColor GetNodeTitleColor() const override { return FLinearColor(FColor::Cyan); }
 	virtual TArray<EPVRenderType> GetDefaultRenderType() const override { return TArray{EPVRenderType::Mesh}; }
-	virtual TArray<EPVRenderType> GetSupportedRenderTypes() const override { return TArray{EPVRenderType::Skeleton, EPVRenderType::Mesh}; }
+	virtual TArray<EPVRenderType> GetSupportedRenderTypes() const override { return TArray{EPVRenderType::PointData, EPVRenderType::Mesh}; }
 
 	virtual void PostEditChangeProperty(FPropertyChangedEvent& PropertyChangedEvent) override;
 
diff --git a/Engine/Plugins/Experimental/ProceduralVegetationEditor/Source/ProceduralVegetation/Private/Nodes/PVPresetLoaderSettings.h b/Engine/Plugins/Experimental/ProceduralVegetationEditor/Source/ProceduralVegetation/Private/Nodes/PVPresetLoaderSettings.h
--- a/Engine/Plugins/Experimental/ProceduralVegetationEditor/Source/ProceduralVegetation/Private/Nodes/PVPresetLoaderSettings.h
+++ b/Engine/Plugins/Experimental/ProceduralVegetationEditor/Source/ProceduralVegetation/Private/Nodes/PVPresetLoaderSettings.h
@@ -21,7 +21,7 @@ public:
 	virtual FText GetDefaultNodeTitle() const override;
 	virtual FText GetNodeTooltipText() const override;
 	virtual EPCGSettingsType GetType() const override { return EPCGSettingsType::InputOutput; }
-	virtual TArray<EPVRenderType> GetDefaultRenderType() const override { return TArray{ EPVRenderType::Skeleton }; }
+	virtual TArray<EPVRenderType> GetDefaultRenderType() const override { return TArray{ EPVRenderType::PointData }; }
 
 	virtual void PostEditChangeProperty(FPropertyChangedEvent& PropertyChangedEvent) override;
 #endif
diff --git a/Engine/Plugins/Experimental/ProceduralVegetationEditor/Source/ProceduralVegetation/Private/Nodes/PVRemoveBranchesSettings.h b/Engine/Plugins/Experimental/ProceduralVegetationEditor/Source/ProceduralVegetation/Private/Nodes/PVRemoveBranchesSettings.h
--- a/Engine/Plugins/Experimental/ProceduralVegetationEditor/Source/ProceduralVegetation/Private/Nodes/PVRemoveBranchesSettings.h
+++ b/Engine/Plugins/Experimental/ProceduralVegetationEditor/Source/ProceduralVegetation/Private/Nodes/PVRemoveBranchesSettings.h
@@ -18,7 +18,7 @@ public:
 	virtual FText GetNodeTooltipText() const override;
 	virtual FLinearColor GetNodeTitleColor() const override { return FLinearColor::Yellow; }
 
-	virtual TArray<EPVRenderType> GetDefaultRenderType() const override { return TArray{ EPVRenderType::Skeleton }; }
+	virtual TArray<EPVRenderType> GetDefaultRenderType() const override { return TArray{ EPVRenderType::PointData }; }
 #endif
 
 protected:
diff --git a/Engine/Plugins/Experimental/ProceduralVegetationEditor/Source/ProceduralVegetation/Private/Nodes/PVScaleSettings.h b/Engine/Plugins/Experimental/ProceduralVegetationEditor/Source/ProceduralVegetation/Private/Nodes/PVScaleSettings.h
--- a/Engine/Plugins/Experimental/ProceduralVegetationEditor/Source/ProceduralVegetation/Private/Nodes/PVScaleSettings.h
+++ b/Engine/Plugins/Experimental/ProceduralVegetationEditor/Source/ProceduralVegetation/Private/Nodes/PVScaleSettings.h
@@ -17,7 +17,7 @@ public:
 	virtual FText GetNodeTooltipText() const override;
 	virtual FLinearColor GetNodeTitleColor() const override { return FLinearColor::Yellow; }
 
-	virtual TArray<EPVRenderType> GetDefaultRenderType() const override { return TArray{ EPVRenderType::Skeleton }; }
+	virtual TArray<EPVRenderType> GetDefaultRenderType() const override { return TArray{ EPVRenderType::PointData }; }
 #endif
 
 protected:
diff --git a/Engine/Plugins/Experimental/ProceduralVegetationEditor/Source/ProceduralVegetation/Public/DataTypes/PVData.h b/Engine/Plugins/Experimental/ProceduralVegetationEditor/Source/ProceduralVegetation/Public/DataTypes/PVData.h
--- a/Engine/Plugins/Experimental/ProceduralVegetationEditor/Source/ProceduralVegetation/Public/DataTypes/PVData.h
+++ b/Engine/Plugins/Experimental/ProceduralVegetationEditor/Source/ProceduralVegetation/Public/DataTypes/PVData.h
@@ -14,7 +14,7 @@ UENUM()
 enum class EPVRenderType : uint8
 {
 	None 		= 0,
-	Skeleton 	= 1 << 0,
+	PointData 	= 1 << 0,
 	Mesh 		= 1 << 1,
 	Foliage 	= 1 << 2,
 	Bones 		= 1 << 3,
diff --git a/Engine/Plugins/Experimental/ProceduralVegetationEditor/Source/ProceduralVegetation/Public/Nodes/PVOutputSettings.h b/Engine/Plugins/Experimental/ProceduralVegetationEditor/Source/ProceduralVegetation/Public/Nodes/PVOutputSettings.h
--- a/Engine/Plugins/Experimental/ProceduralVegetationEditor/Source/ProceduralVegetation/Public/Nodes/PVOutputSettings.h
+++ b/Engine/Plugins/Experimental/ProceduralVegetationEditor/Source/ProceduralVegetation/Public/Nodes/PVOutputSettings.h
@@ -19,7 +19,7 @@ public:
 	virtual FText GetDefaultNodeTitle() const override;
 	virtual FText GetNodeTooltipText() const override;
 	virtual EPCGSettingsType GetType() const override { return EPCGSettingsType::InputOutput; }
-	virtual TArray<EPVRenderType> GetSupportedRenderTypes() const override { return TArray{EPVRenderType::Skeleton, EPVRenderType::Foliage, EPVRenderType::FoliageGrid, EPVRenderType::Mesh, EPVRenderType::Bones}; }
+	virtual TArray<EPVRenderType> GetSupportedRenderTypes() const override { return TArray{EPVRenderType::PointData, EPVRenderType::Foliage, EPVRenderType::FoliageGrid, EPVRenderType::Mesh, EPVRenderType::Bones}; }
 	virtual TArray<EPVRenderType> GetDefaultRenderType() const override { return TArray{EPVRenderType::Foliage, EPVRenderType::Mesh}; }
 #endif
 
diff --git a/Engine/Plugins/Experimental/ProceduralVegetationEditor/Source/ProceduralVegetationEditor/Private/SPVEditorViewport.cpp b/Engine/Plugins/Experimental/ProceduralVegetationEditor/Source/ProceduralVegetationEditor/Private/SPVEditorViewport.cpp
--- a/Engine/Plugins/Experimental/ProceduralVegetationEditor/Source/ProceduralVegetationEditor/Private/SPVEditorViewport.cpp
+++ b/Engine/Plugins/Experimental/ProceduralVegetationEditor/Source/ProceduralVegetationEditor/Private/SPVEditorViewport.cpp
@@ -343,7 +343,7 @@ EPVVisualizationMode PVRenderModesToVisualizationMode(const TArray<EPVRenderType
 	{
 		return EPVVisualizationMode::Bones;
 	}
-	else if (InRenderType.Contains(EPVRenderType::Skeleton))
+	else if (InRenderType.Contains(EPVRenderType::PointData))
 	{
 		return EPVVisualizationMode::PointData;
 	}
@@ -371,13 +371,13 @@ TArray<EPVRenderType> VisualizationModeToRenderModes(EPVVisualizationMode InMode
 		return {EPVRenderType::Mesh};
 
 	case EPVVisualizationMode::PointData:
-		return {EPVRenderType::Skeleton};
+		return {EPVRenderType::PointData};
 
 	case EPVVisualizationMode::BonesMesh:
 		return {EPVRenderType::Bones, EPVRenderType::Mesh};
 
 	case EPVVisualizationMode::PointDataMesh:
-		return {EPVRenderType::Skeleton, EPVRenderType::Mesh};
+		return {EPVRenderType::PointData, EPVRenderType::Mesh};
 
 	default:
 		return TArray<EPVRenderType>{};
diff --git a/Engine/Plugins/Experimental/ProceduralVegetationEditor/Source/ProceduralVegetationEditor/Private/SPVEditorViewport.h b/Engine/Plugins/Experimental/ProceduralVegetationEditor/Source/ProceduralVegetationEditor/Private/SPVEditorViewport.h
--- a/Engine/Plugins/Experimental/ProceduralVegetationEditor/Source/ProceduralVegetationEditor/Private/SPVEditorViewport.h
+++ b/Engine/Plugins/Experimental/ProceduralVegetationEditor/Source/ProceduralVegetationEditor/Private/SPVEditorViewport.h
@@ -18,8 +18,7 @@ enum class EPVVisualizationMode : uint8
 	Bones			UMETA(DisplayName = "Bones"),
 	PointDataMesh	UMETA(DisplayName = "Point data + Mesh"),
 	FoliageMesh		UMETA(DisplayName = "Foliage + Mesh"),
-	BonesMesh		UMETA(DisplayName = "Bones + Mesh"),
-	Count			UMETA(Hidden)
+	BonesMesh		UMETA(DisplayName = "Bones + Mesh")
 };
 
 class UPVBaseSettings;
diff --git a/Engine/Plugins/Experimental/ProceduralVegetationEditor/Source/ProceduralVegetationEditor/Private/Visualizations/PVDataVisualization.cpp b/Engine/Plugins/Experimental/ProceduralVegetationEditor/Source/ProceduralVegetationEditor/Private/Visualizations/PVDataVisualization.cpp
--- a/Engine/Plugins/Experimental/ProceduralVegetationEditor/Source/ProceduralVegetationEditor/Private/Visualizations/PVDataVisualization.cpp
+++ b/Engine/Plugins/Experimental/ProceduralVegetationEditor/Source/ProceduralVegetationEditor/Private/Visualizations/PVDataVisualization.cpp
@@ -41,7 +41,7 @@ FPVDataVisualization::FPVDataVisualization()
 
 void FPVDataVisualization::RegisterVisualizations()
 {
-	RenderMap.Add(EPVRenderType::Skeleton, FPVRenderCallback::CreateRaw(this, &FPVDataVisualization::SkeletonRenderer));
+	RenderMap.Add(EPVRenderType::PointData, FPVRenderCallback::CreateRaw(this, &FPVDataVisualization::SkeletonRenderer));
 	RenderMap.Add(EPVRenderType::Mesh, FPVRenderCallback::CreateRaw(this, &FPVDataVisualization::MeshRenderer));
 	RenderMap.Add(EPVRenderType::Foliage, FPVRenderCallback::CreateRaw(this, &FPVDataVisualization::FoliageRenderer));
 	RenderMap.Add(EPVRenderType::Bones, FPVRenderCallback::CreateRaw(this, &FPVDataVisualization::BonesRenderer));
