[PVE] Analytics added
- NodeAdded
- NodeTweaked
- MeshMaterialChanged
- FoliageMeshChanged
- WindPresetChanged
- ExportStarted
- ExportFinished

diff --git a/Engine/Plugins/Experimental/ProceduralVegetationEditor/Source/ProceduralVegetation/Private/Helpers/PVAnalyticsHelper.cpp b/Engine/Plugins/Experimental/ProceduralVegetationEditor/Source/ProceduralVegetation/Private/Helpers/PVAnalyticsHelper.cpp
new file mode 100644
--- /dev/null
+++ b/Engine/Plugins/Experimental/ProceduralVegetationEditor/Source/ProceduralVegetation/Private/Helpers/PVAnalyticsHelper.cpp
@@ -0,0 +1,178 @@
+﻿// Copyright Epic Games, Inc. All Rights Reserved.
+
+#include "Helpers/PVAnalyticsHelper.h"
+
+#include "EngineAnalytics.h"
+#include "PVExportParams.h"
+#include "Engine/RendererSettings.h"
+#include "Facades/PVFoliageFacade.h"
+#include "GeometryCollection/GeometryCollection.h"
+#include "GeometryCollection/ManagedArrayCollection.h"
+#include "Helpers/PVUtilities.h"
+#include "UObject/UnrealTypePrivate.h"
+
+namespace PV::Analytics
+{
+	void SendEvent(const FString& InEventName, const TArray<FAnalyticsEventAttribute>& InAttributes = TArray<FAnalyticsEventAttribute>())
+	{
+		if (FEngineAnalytics::IsAvailable())
+		{
+			FString EventName = FString::Printf(TEXT("Editor.Usage.PVE.%s"), *InEventName);
+			FEngineAnalytics::GetProvider().RecordEvent(EventName, InAttributes);
+
+			// FString AttributesData;
+			// for (int32 i = 0; i < InAttributes.Num(); i++)
+			// {
+			// 	AttributesData += FString::Printf(TEXT("(%s,%s),"), *InAttributes[i].GetName(), *InAttributes[i].GetValue());
+			// }
+			//
+			// UE_LOG(LogTemp, Display, TEXT("PV::SendEvent: %s : Attributes: {%s}"), *EventName, *AttributesData);
+		}
+	}
+
+	void SendSessionStartedEvent()
+	{
+		if (FEngineAnalytics::IsAvailable())
+		{
+			SendEvent(TEXT("SessionStarted"));
+		}
+	}
+
+	void SendSessionEndedEvent(const double InTimeInSeconds)
+	{
+		if (FEngineAnalytics::IsAvailable())
+		{
+			TArray<FAnalyticsEventAttribute> Attributes;
+			Attributes.Add(FAnalyticsEventAttribute(TEXT("TimeActiveSeconds"),  InTimeInSeconds));
+			
+			SendEvent(TEXT("SessionEnded"), Attributes);
+		}
+	}
+
+	void SendNodeAddedEvent(const FString& InNodeType)
+	{
+		if (FEngineAnalytics::IsAvailable())
+		{
+			TArray<FAnalyticsEventAttribute> Attributes;
+			Attributes.Add(FAnalyticsEventAttribute(TEXT("Name"),  InNodeType));
+		
+			SendEvent(TEXT("NodeAdded"), Attributes);	
+		}
+	}
+
+	void SendNodeTweakedEvent(const FString& InNodeType, const FProperty* InProperty)
+	{
+		if (FEngineAnalytics::IsAvailable())
+		{
+			FString PropertyName = InProperty->GetName();
+			FString PropertyCategory = InProperty->GetMetaData(TEXT("Category"));
+			
+			TArray<FAnalyticsEventAttribute> Attributes;
+			Attributes.Add(FAnalyticsEventAttribute(TEXT("Name"),  InNodeType));
+			Attributes.Add(FAnalyticsEventAttribute(TEXT("PropertyCategory"),  PropertyCategory));
+			Attributes.Add(FAnalyticsEventAttribute(TEXT("PropertyName"),  PropertyName));
+		
+			SendEvent(TEXT("NodeTweaked"), Attributes);
+		}
+	}
+
+	void SendMaterialChangeEvent(FString InMatPath)
+	{
+		if (FEngineAnalytics::IsAvailable())
+		{
+			TArray<FAnalyticsEventAttribute> Attributes;
+			Attributes.Add(FAnalyticsEventAttribute(TEXT("Path"),  InMatPath));
+			
+			SendEvent(TEXT("MeshMaterialChanged"), Attributes);
+		}
+	}
+
+	void SendFoliageMeshChangeEvent(FString InFoliageMeshPath)
+	{
+		if (FEngineAnalytics::IsAvailable())
+		{
+			TArray<FAnalyticsEventAttribute> Attributes;
+			Attributes.Add(FAnalyticsEventAttribute(TEXT("Path"),  InFoliageMeshPath));
+			
+			SendEvent(TEXT("FoliageMeshChanged"), Attributes);
+		}
+	}
+
+	void SendWindSettingsChangeEvent(FString InPath)
+	{
+		if (FEngineAnalytics::IsAvailable())
+		{
+			TArray<FAnalyticsEventAttribute> Attributes;
+			Attributes.Add(FAnalyticsEventAttribute(TEXT("Path"),  InPath));
+			
+			SendEvent(TEXT("WindPresetChanged"), Attributes);
+		}
+	}
+
+	void GatherFoliageMeshesTriangles(const Facades::FFoliageFacade& InFoliageFacade, uint32& OutUniqueFoliageTriangles, uint32& OutTotalFoliageTriangles)
+	{
+		OutUniqueFoliageTriangles = 0;
+		OutTotalFoliageTriangles = 0;
+		
+		for (int32 FoliageNameIndex = 0; FoliageNameIndex < InFoliageFacade.NumFoliageNames(); ++FoliageNameIndex)
+		{
+			FString FoliagePath = InFoliageFacade.GetFoliageName(FoliageNameIndex);
+			
+			OutUniqueFoliageTriangles += PV::Utilities::GetMeshTriangles(FoliagePath);
+		}
+
+		int32 NumInstances = InFoliageFacade.NumFoliageEntries();
+		
+		for (int32 Id = 0; Id < NumInstances; Id++)
+		{
+			Facades::FFoliageEntryData Data = InFoliageFacade.GetFoliageEntry(Id);
+			FString FoliagePath = InFoliageFacade.GetFoliageName(Data.NameId);
+
+			OutTotalFoliageTriangles += PV::Utilities::GetMeshTriangles(FoliagePath);
+		}
+	}
+
+	TArray<FAnalyticsEventAttribute> GatherExportCommonAttributes(const FManagedArrayCollection& InCollection, const FPVExportParams& InExportParam)
+	{
+		TArray<FAnalyticsEventAttribute> ExportAttributes;
+		ExportAttributes.Add(FAnalyticsEventAttribute(TEXT("ExportSettings.MeshType"),  StaticEnum<EPVExportMeshType>()->GetNameStringByValue((uint8)InExportParam.ExportMeshType)));
+		ExportAttributes.Add(FAnalyticsEventAttribute(TEXT("ExportSettings.CreateNaniteFoliage"),  InExportParam.bCreateNaniteFoliage ? TEXT("True") : TEXT("False")));
+		ExportAttributes.Add(FAnalyticsEventAttribute(TEXT("ExportSettings.NaniteShapePreservation"),  StaticEnum<ENaniteShapePreservation>()->GetNameStringByValue((uint8)InExportParam.NaniteShapePreservation)));
+		ExportAttributes.Add(FAnalyticsEventAttribute(TEXT("NaniteSettings.NaniteFoliageEnabled"),  GetDefault<URendererSettings>()->bEnableNaniteFoliage));
+		ExportAttributes.Add(FAnalyticsEventAttribute(TEXT("NaniteSettings.Enabled"),  GetDefault<URendererSettings>()->bNanite));
+
+		Facades::FFoliageFacade FoliageFacade(InCollection);
+		
+		const int FoliageInstances = FoliageFacade.GetElementCount();
+		const int TrunkFaces = InCollection.NumElements(FGeometryCollection::FacesGroup);
+
+		uint32 UniqueFoliageTriangles = 0;
+		uint32 TotalFoliageTriangles = 0;
+		GatherFoliageMeshesTriangles(FoliageFacade, UniqueFoliageTriangles, TotalFoliageTriangles);
+		
+		ExportAttributes.Add(FAnalyticsEventAttribute(TEXT("FoliageInstances"),  FoliageInstances));
+		ExportAttributes.Add(FAnalyticsEventAttribute(TEXT("TrunkTriangles"),  TrunkFaces));
+		ExportAttributes.Add(FAnalyticsEventAttribute(TEXT("FoliageTrianglesUnique"),  UniqueFoliageTriangles));
+		ExportAttributes.Add(FAnalyticsEventAttribute(TEXT("FoliageTrianglesTotal"),  TotalFoliageTriangles));
+
+		return ExportAttributes;
+	}
+
+	void SendExportStartEvent(const TArray<FAnalyticsEventAttribute>& InCommonAttributes)
+	{
+		if (FEngineAnalytics::IsAvailable())
+		{
+			SendEvent(TEXT("ExportStarted"), InCommonAttributes);
+		}
+	}
+
+	void SendExportFinishedEvent(TArray<FAnalyticsEventAttribute>& InCommonAttributes, float InTotalTimeInSeconds)
+	{
+		if (FEngineAnalytics::IsAvailable())
+		{
+			InCommonAttributes.Add(FAnalyticsEventAttribute(TEXT("TotalTimeSeconds"),  InTotalTimeInSeconds));
+			
+			SendEvent(TEXT("ExportFinished"), InCommonAttributes);
+		}
+	}
+}
diff --git a/Engine/Plugins/Experimental/ProceduralVegetationEditor/Source/ProceduralVegetation/Private/Helpers/PVUtilities.cpp b/Engine/Plugins/Experimental/ProceduralVegetationEditor/Source/ProceduralVegetation/Private/Helpers/PVUtilities.cpp
--- a/Engine/Plugins/Experimental/ProceduralVegetationEditor/Source/ProceduralVegetation/Private/Helpers/PVUtilities.cpp
+++ b/Engine/Plugins/Experimental/ProceduralVegetationEditor/Source/ProceduralVegetation/Private/Helpers/PVUtilities.cpp
@@ -13,6 +13,7 @@
 #include "Misc/PackageName.h"
 #include "AssetRegistry/AssetData.h"
 #include "Misc/FileHelper.h"
+#include "Rendering/SkeletalMeshRenderData.h"
 
 namespace PV::Utilities
 {
@@ -219,4 +220,30 @@ namespace PV::Utilities
 
 		return false;
 	}
+
+	int32 GetMeshTriangles(const FString InMeshPath)
+	{
+		UStaticMesh* StaticFoliageMesh = LoadObject<UStaticMesh>(nullptr, *InMeshPath);
+		USkeletalMesh* SkeletalFoliageMesh = LoadObject<USkeletalMesh>(nullptr, *InMeshPath);
+
+		int32 NumTriangles = 0;
+		
+		if (StaticFoliageMesh)
+		{
+			NumTriangles = StaticFoliageMesh->GetNumTriangles(0);
+		}
+		else if (SkeletalFoliageMesh)
+		{
+			FSkeletalMeshRenderData* SkelMeshRenderData = SkeletalFoliageMesh->GetResourceForRendering();
+			if (SkelMeshRenderData && SkelMeshRenderData->LODRenderData.Num() > 0)
+			{
+				const FSkeletalMeshLODRenderData& LODData = SkelMeshRenderData->LODRenderData[0];
+				NumTriangles = LODData.GetTotalFaces();
+			}
+
+			return NumTriangles;
+		}
+
+		return NumTriangles;
+	}
 }
diff --git a/Engine/Plugins/Experimental/ProceduralVegetationEditor/Source/ProceduralVegetation/Private/Nodes/PVBaseSettings.cpp b/Engine/Plugins/Experimental/ProceduralVegetationEditor/Source/ProceduralVegetation/Private/Nodes/PVBaseSettings.cpp
--- a/Engine/Plugins/Experimental/ProceduralVegetationEditor/Source/ProceduralVegetation/Private/Nodes/PVBaseSettings.cpp
+++ b/Engine/Plugins/Experimental/ProceduralVegetationEditor/Source/ProceduralVegetation/Private/Nodes/PVBaseSettings.cpp
@@ -5,12 +5,16 @@
 #include "Helpers/PCGHelpers.h"
 #include "DataTypes/PVData.h"
 #include "PCGData.h"
+#include "Helpers/PVAnalyticsHelper.h"
 
 UPVBaseSettings::UPVBaseSettings()
 {
 	if (PCGHelpers::IsNewObjectAndNotDefault(this))
 	{
-		
+		if (const UClass* Class = GetClass())
+		{
+			PV::Analytics::SendNodeAddedEvent(Class->GetName());
+		}
 	}
 }
 
@@ -56,6 +60,22 @@ FPCGDataTypeIdentifier UPVBaseSettings::GetOutputPinTypeIdentifier() const
 	return FPCGDataTypeIdentifier{ EPCGDataType::Any };
 }
 
+#if WITH_EDITOR
+void UPVBaseSettings::PostEditChangeProperty(FPropertyChangedEvent& PropertyChangedEvent)
+{
+	Super::PostEditChangeProperty(PropertyChangedEvent);
+
+	if (PropertyChangedEvent.Property && PropertyChangedEvent.ChangeType == EPropertyChangeType::ValueSet)
+	{
+		if (const UClass* Class = GetClass())
+		{
+			PV::Analytics::SendNodeTweakedEvent(Class->GetName(), PropertyChangedEvent.Property);
+		}
+	}
+}
+
+#endif
+
 FPCGElementPtr UPVBaseSettings::CreateElement() const
 {
 	return MakeShared<FPVBaseElement>();
diff --git a/Engine/Plugins/Experimental/ProceduralVegetationEditor/Source/ProceduralVegetation/Private/Nodes/PVFoliagePaletteSettings.cpp b/Engine/Plugins/Experimental/ProceduralVegetationEditor/Source/ProceduralVegetation/Private/Nodes/PVFoliagePaletteSettings.cpp
--- a/Engine/Plugins/Experimental/ProceduralVegetationEditor/Source/ProceduralVegetation/Private/Nodes/PVFoliagePaletteSettings.cpp
+++ b/Engine/Plugins/Experimental/ProceduralVegetationEditor/Source/ProceduralVegetation/Private/Nodes/PVFoliagePaletteSettings.cpp
@@ -11,6 +11,7 @@
 #include "GeometryCollection/ManagedArrayCollection.h"
 #include "Helpers/PCGHelpers.h"
 #include "Helpers/PCGSettingsHelpers.h"
+#include "Helpers/PVAnalyticsHelper.h"
 
 #define LOCTEXT_NAMESPACE "PVFoliagePaletteSettings"
 
@@ -39,6 +40,40 @@ void UPVFoliagePaletteSettings::PostEditChangeProperty(FPropertyChangedEvent& Pr
 		DirtyCache();
 	}
 }
+
+void UPVFoliagePaletteSettings::PostEditChangeChainProperty(FPropertyChangedChainEvent& PropertyChangedEvent)
+{
+	Super::PostEditChangeChainProperty(PropertyChangedEvent);
+
+	FProperty* Property = PropertyChangedEvent.Property;
+	FProperty* MemberProperty = nullptr;
+
+	if (PropertyChangedEvent.PropertyChain.GetActiveMemberNode())
+	{
+		MemberProperty = PropertyChangedEvent.PropertyChain.GetActiveMemberNode()->GetValue();
+	}
+
+	if (MemberProperty && Property)
+	{
+		if (MemberProperty->GetFName() == GET_MEMBER_NAME_CHECKED(UPVFoliagePaletteSettings, FoliageMeshes))
+		{
+			FName Name = Property->GetFName();
+			
+			int Index = PropertyChangedEvent.GetArrayIndex(Name.ToString());
+
+			if (Index >= 0 && FoliageMeshes.Num() > Index)
+			{
+				if (!FoliageMeshes[Index].IsNull())
+				{
+					FString MeshPath = FoliageMeshes[Index].ToString();
+
+					PV::Analytics::SendFoliageMeshChangeEvent(MeshPath);	
+				}
+			}
+		}
+	}
+}
+
 #endif
 
 FPCGDataTypeIdentifier UPVFoliagePaletteSettings::GetInputPinTypeIdentifier() const
diff --git a/Engine/Plugins/Experimental/ProceduralVegetationEditor/Source/ProceduralVegetation/Private/Nodes/PVFoliagePaletteSettings.h b/Engine/Plugins/Experimental/ProceduralVegetationEditor/Source/ProceduralVegetation/Private/Nodes/PVFoliagePaletteSettings.h
--- a/Engine/Plugins/Experimental/ProceduralVegetationEditor/Source/ProceduralVegetation/Private/Nodes/PVFoliagePaletteSettings.h
+++ b/Engine/Plugins/Experimental/ProceduralVegetationEditor/Source/ProceduralVegetation/Private/Nodes/PVFoliagePaletteSettings.h
@@ -22,6 +22,7 @@ public:
 	virtual TArray<EPVRenderType> GetSupportedRenderTypes() const override { return TArray{EPVRenderType::FoliageGrid, EPVRenderType::Foliage, EPVRenderType::Mesh}; }
 	
 	virtual void PostEditChangeProperty(FPropertyChangedEvent& PropertyChangedEvent) override;
+	virtual void PostEditChangeChainProperty(FPropertyChangedChainEvent& PropertyChangedEvent) override;
 #endif
 
 protected:
diff --git a/Engine/Plugins/Experimental/ProceduralVegetationEditor/Source/ProceduralVegetation/Private/Nodes/PVMeshBuilderSettings.cpp b/Engine/Plugins/Experimental/ProceduralVegetationEditor/Source/ProceduralVegetation/Private/Nodes/PVMeshBuilderSettings.cpp
--- a/Engine/Plugins/Experimental/ProceduralVegetationEditor/Source/ProceduralVegetation/Private/Nodes/PVMeshBuilderSettings.cpp
+++ b/Engine/Plugins/Experimental/ProceduralVegetationEditor/Source/ProceduralVegetation/Private/Nodes/PVMeshBuilderSettings.cpp
@@ -9,6 +9,7 @@
 #include "Facades/PVProfileFacade.h"
 #include "GeometryCollection/ManagedArrayCollection.h"
 #include "Helpers/PCGSettingsHelpers.h"
+#include "Helpers/PVAnalyticsHelper.h"
 #include "Implementations/PVMeshBuilder.h"
 #include "Materials/MaterialInterface.h"
 
@@ -63,6 +64,39 @@ void UPVMeshBuilderSettings::PostEditChangeProperty(FPropertyChangedEvent& Prope
 	}
 }
 
+void UPVMeshBuilderSettings::PostEditChangeChainProperty(FPropertyChangedChainEvent& PropertyChangedEvent)
+{
+	Super::PostEditChangeChainProperty(PropertyChangedEvent);
+
+	FProperty* Property = PropertyChangedEvent.Property;
+	FProperty* MemberProperty = nullptr;
+
+	if (PropertyChangedEvent.PropertyChain.GetActiveMemberNode())
+	{
+		MemberProperty = PropertyChangedEvent.PropertyChain.GetActiveMemberNode()->GetValue();
+	}
+
+	if (MemberProperty && Property)
+	{
+		if (MemberProperty->GetFName() == GET_MEMBER_NAME_CHECKED(UPVMeshBuilderSettings, MesherSettings))
+		{
+			FName Name = Property->GetFName();
+			
+			if(Name == "Material")
+			{
+				int Index = PropertyChangedEvent.GetArrayIndex("MaterialSetups");
+
+				if (Index >= 0 && MesherSettings.MaterialSettings.MaterialSetups.Num() > Index)
+				{
+					FString MatPath = MesherSettings.MaterialSettings.MaterialSetups[Index].Material->GetPathName();
+
+					PV::Analytics::SendMaterialChangeEvent(MatPath);
+				}
+			}
+		}
+	}
+}
+
 void UPVMeshBuilderSettings::PostLoad()
 {
 	Super::PostLoad();
diff --git a/Engine/Plugins/Experimental/ProceduralVegetationEditor/Source/ProceduralVegetation/Private/Nodes/PVMeshBuilderSettings.h b/Engine/Plugins/Experimental/ProceduralVegetationEditor/Source/ProceduralVegetation/Private/Nodes/PVMeshBuilderSettings.h
--- a/Engine/Plugins/Experimental/ProceduralVegetationEditor/Source/ProceduralVegetation/Private/Nodes/PVMeshBuilderSettings.h
+++ b/Engine/Plugins/Experimental/ProceduralVegetationEditor/Source/ProceduralVegetation/Private/Nodes/PVMeshBuilderSettings.h
@@ -24,6 +24,7 @@ public:
 	virtual TArray<EPVRenderType> GetSupportedRenderTypes() const override { return TArray{EPVRenderType::PointData, EPVRenderType::Mesh}; }
 
 	virtual void PostEditChangeProperty(FPropertyChangedEvent& PropertyChangedEvent) override;
+	virtual void PostEditChangeChainProperty(FPropertyChangedChainEvent& PropertyChangedEvent) override;
 
 	virtual void PostLoad() override;
 #endif
diff --git a/Engine/Plugins/Experimental/ProceduralVegetationEditor/Source/ProceduralVegetation/Private/Nodes/PVOutputSettings.cpp b/Engine/Plugins/Experimental/ProceduralVegetationEditor/Source/ProceduralVegetation/Private/Nodes/PVOutputSettings.cpp
--- a/Engine/Plugins/Experimental/ProceduralVegetationEditor/Source/ProceduralVegetation/Private/Nodes/PVOutputSettings.cpp
+++ b/Engine/Plugins/Experimental/ProceduralVegetationEditor/Source/ProceduralVegetation/Private/Nodes/PVOutputSettings.cpp
@@ -8,6 +8,7 @@
 #include "Misc/PackageName.h"
 #include "Engine/StaticMesh.h"
 #include "Engine/SkeletalMesh.h"
+#include "Helpers/PVAnalyticsHelper.h"
 
 
 #define LOCTEXT_NAMESPACE "PVOutputSettings"
@@ -67,6 +68,22 @@ FPCGElementPtr UPVOutputSettings::CreateElement() const
 	return MakeShared<FPVOutputElement>();
 }
 
+#if	WITH_EDITOR
+void UPVOutputSettings::PostEditChangeProperty(FPropertyChangedEvent& PropertyChangedEvent)
+{
+	Super::PostEditChangeProperty(PropertyChangedEvent);
+
+	if (PropertyChangedEvent.Property
+		&& PropertyChangedEvent.Property->GetFName() == GET_MEMBER_NAME_CHECKED(FPVExportParams, WindSettings))
+	{
+		if (ExportSettings.WindSettings)
+		{
+			PV::Analytics::SendWindSettingsChangeEvent(ExportSettings.WindSettings->GetPathName());
+		}
+	}
+}
+#endif
+
 bool FPVOutputElement::ExecuteInternal(FPCGContext* InContext) const
 {
 	const UPVOutputSettings* Settings = InContext->GetInputSettings<UPVOutputSettings>();
diff --git a/Engine/Plugins/Experimental/ProceduralVegetationEditor/Source/ProceduralVegetation/Public/Helpers/PVAnalyticsHelper.h b/Engine/Plugins/Experimental/ProceduralVegetationEditor/Source/ProceduralVegetation/Public/Helpers/PVAnalyticsHelper.h
new file mode 100644
--- /dev/null
+++ b/Engine/Plugins/Experimental/ProceduralVegetationEditor/Source/ProceduralVegetation/Public/Helpers/PVAnalyticsHelper.h
@@ -0,0 +1,23 @@
+﻿// Copyright Epic Games, Inc. All Rights Reserved.
+#pragma once
+
+#include "CoreMinimal.h"
+
+struct FManagedArrayCollection;
+struct FPVExportParams;
+class FProperty;
+struct FAnalyticsEventAttribute;
+
+namespace PV::Analytics
+{
+	void PROCEDURALVEGETATION_API SendSessionStartedEvent();
+	void PROCEDURALVEGETATION_API SendSessionEndedEvent(const double InTimeInSeconds);
+	void SendNodeAddedEvent(const FString& InNodeType);
+	void SendNodeTweakedEvent(const FString& InNodeType, const FProperty* InProperty);
+	void SendMaterialChangeEvent(FString InMatPath);
+	void SendFoliageMeshChangeEvent(FString InFoliageMeshPath);
+	void SendWindSettingsChangeEvent(FString InPath);
+	TArray<FAnalyticsEventAttribute> PROCEDURALVEGETATION_API GatherExportCommonAttributes(const FManagedArrayCollection& InCollection, const FPVExportParams& InExportParam);
+	void PROCEDURALVEGETATION_API SendExportStartEvent(const TArray<FAnalyticsEventAttribute>& InCommonAttributes);
+	void PROCEDURALVEGETATION_API SendExportFinishedEvent(TArray<FAnalyticsEventAttribute>& InCommonAttributes, float InTotalTimeInSeconds);
+}
diff --git a/Engine/Plugins/Experimental/ProceduralVegetationEditor/Source/ProceduralVegetation/Public/Helpers/PVUtilities.h b/Engine/Plugins/Experimental/ProceduralVegetationEditor/Source/ProceduralVegetation/Public/Helpers/PVUtilities.h
--- a/Engine/Plugins/Experimental/ProceduralVegetationEditor/Source/ProceduralVegetation/Public/Helpers/PVUtilities.h
+++ b/Engine/Plugins/Experimental/ProceduralVegetationEditor/Source/ProceduralVegetation/Public/Helpers/PVUtilities.h
@@ -23,4 +23,5 @@ namespace PV::Utilities
 
 	bool IsFileNameValid(FName FileName, FText& Reason);
 	bool PROCEDURALVEGETATION_API DoesConflictingPackageExist(const FString& PackageName, UClass* InClass);
+	int32 PROCEDURALVEGETATION_API GetMeshTriangles(const FString InMeshPath);
 }
diff --git a/Engine/Plugins/Experimental/ProceduralVegetationEditor/Source/ProceduralVegetation/Public/Nodes/PVBaseSettings.h b/Engine/Plugins/Experimental/ProceduralVegetationEditor/Source/ProceduralVegetation/Public/Nodes/PVBaseSettings.h
--- a/Engine/Plugins/Experimental/ProceduralVegetationEditor/Source/ProceduralVegetation/Public/Nodes/PVBaseSettings.h
+++ b/Engine/Plugins/Experimental/ProceduralVegetationEditor/Source/ProceduralVegetation/Public/Nodes/PVBaseSettings.h
@@ -35,6 +35,11 @@ protected:
 
 	virtual FPCGDataTypeIdentifier GetInputPinTypeIdentifier() const;
 	virtual FPCGDataTypeIdentifier GetOutputPinTypeIdentifier() const;
+
+#if WITH_EDITOR
+	virtual void PostEditChangeProperty(FPropertyChangedEvent& PropertyChangedEvent) override;
+#endif
+	
 private:
 	
 #if WITH_EDITORONLY_DATA
diff --git a/Engine/Plugins/Experimental/ProceduralVegetationEditor/Source/ProceduralVegetation/Public/Nodes/PVOutputSettings.h b/Engine/Plugins/Experimental/ProceduralVegetationEditor/Source/ProceduralVegetation/Public/Nodes/PVOutputSettings.h
--- a/Engine/Plugins/Experimental/ProceduralVegetationEditor/Source/ProceduralVegetation/Public/Nodes/PVOutputSettings.h
+++ b/Engine/Plugins/Experimental/ProceduralVegetationEditor/Source/ProceduralVegetation/Public/Nodes/PVOutputSettings.h
@@ -28,6 +28,10 @@ protected:
 	virtual TArray<FPCGPinProperties> OutputPinProperties() const override;
 
 	virtual FPCGElementPtr CreateElement() const override;
+
+#if WITH_EDITOR
+	virtual void PostEditChangeProperty(FPropertyChangedEvent& PropertyChangedEvent) override;
+#endif
 	
 public:
 	UPROPERTY(EditAnywhere, Category = ExportSettings, meta = (ShowOnlyInnerProperties))
diff --git a/Engine/Plugins/Experimental/ProceduralVegetationEditor/Source/ProceduralVegetationEditor/Private/Helpers/PVExportHelper.cpp b/Engine/Plugins/Experimental/ProceduralVegetationEditor/Source/ProceduralVegetationEditor/Private/Helpers/PVExportHelper.cpp
--- a/Engine/Plugins/Experimental/ProceduralVegetationEditor/Source/ProceduralVegetationEditor/Private/Helpers/PVExportHelper.cpp
+++ b/Engine/Plugins/Experimental/ProceduralVegetationEditor/Source/ProceduralVegetationEditor/Private/Helpers/PVExportHelper.cpp
@@ -6,6 +6,7 @@
 #include "DynamicMeshEditor.h"
 #include "DynamicWindImportData.h"
 #include "DynamicWindSkeletalData.h"
+#include "EngineAnalytics.h"
 #include "MeshDescriptionToDynamicMesh.h"
 #include "PackageTools.h"
 #include "PlanarCut.h"
@@ -43,6 +44,7 @@
 #include "GeometryScript/GeometryScriptTypes.h"
 #include "GeometryScript/MeshAssetFunctions.h"
 #include "GeometryScript/MeshBoneWeightFunctions.h"
+#include "Helpers/PVAnalyticsHelper.h"
 
 #include "Helpers/PVUtilities.h"
 
@@ -225,6 +227,8 @@ namespace PV::Export::Internal
 		return Package;
 	}
 
+	
+
 	TObjectPtr<UDynamicMesh> CollectionToDynamicMesh(const FManagedArrayCollection& Collection)
 	{
 		TRACE_CPUPROFILER_EVENT_SCOPE(PV::Export::CollectionToDynamicMesh);
@@ -1490,6 +1494,9 @@ namespace PV::Export
 
 		TArray<UPackage*> ModifiedPackages;
 
+		FDateTime ExportStartTime = FDateTime::Now();
+		TArray<FAnalyticsEventAttribute> ExportAnalyticsAttributes = Analytics::GatherExportCommonAttributes(Collection, ExportParams);
+
 		if (ExportParams.ExportMeshType == EPVExportMeshType::StaticMesh)
 		{
 			if (Utilities::DoesConflictingPackageExist(MeshPackageName, UStaticMesh::StaticClass()))
@@ -1497,6 +1504,8 @@ namespace PV::Export
 				return false;
 			}
 
+			Analytics::SendExportStartEvent(ExportAnalyticsAttributes);
+
 			if (TObjectPtr<UStaticMesh> ExistingMeshAsset = Internal::GetAssetOfType<UStaticMesh>(MeshPackageName))
 			{
 				ExistingMeshAsset->SetExternalPackage(nullptr);
@@ -1537,6 +1546,8 @@ namespace PV::Export
 				return false;
 			}
 
+			Analytics::SendExportStartEvent(ExportAnalyticsAttributes);
+
 			if (TObjectPtr<USkeletalMesh> ExistingMeshAsset = Internal::GetAssetOfType<USkeletalMesh>(MeshPackageName))
 			{
 				ExistingMeshAsset->SetExternalPackage(nullptr);
@@ -1589,6 +1600,9 @@ namespace PV::Export
 				}
 			}
 		}
+		double ExportTime = (FDateTime::Now() - ExportStartTime).GetTotalSeconds();
+
+		Analytics::SendExportFinishedEvent(ExportAnalyticsAttributes, ExportTime);
 
 		return true;
 	}
diff --git a/Engine/Plugins/Experimental/ProceduralVegetationEditor/Source/ProceduralVegetationEditor/Private/PVEditor.cpp b/Engine/Plugins/Experimental/ProceduralVegetationEditor/Source/ProceduralVegetationEditor/Private/PVEditor.cpp
--- a/Engine/Plugins/Experimental/ProceduralVegetationEditor/Source/ProceduralVegetationEditor/Private/PVEditor.cpp
+++ b/Engine/Plugins/Experimental/ProceduralVegetationEditor/Source/ProceduralVegetationEditor/Private/PVEditor.cpp
@@ -18,6 +18,7 @@
 
 #include "DataTypes/PVData.h"
 #include "DataTypes/PVFoliageMeshData.h"
+#include "Helpers/PVAnalyticsHelper.h"
 
 #include "Helpers/PVExportHelper.h"
 #include "Helpers/PVUtilities.h"
@@ -108,7 +109,7 @@ void FPVEditor::Initialize(const EToolkitMode::Type InMode, const TSharedPtr<cla
 
 	if (FEngineAnalytics::IsAvailable())
 	{
-		FEngineAnalytics::GetProvider().RecordEvent(TEXT("Editor.Usage.PVE.SessionStarted"));
+		PV::Analytics::SendSessionStartedEvent();
 	}
 }
 
@@ -156,13 +157,8 @@ FString FPVEditor::GetWorldCentricTabPrefix() const
 
 void FPVEditor::OnClose()
 {
-	if (FEngineAnalytics::IsAvailable())
-	{
-		TArray<FAnalyticsEventAttribute> Attributes;
-		Attributes.Add(FAnalyticsEventAttribute(TEXT("TimeActive.Seconds"),  (FDateTime::Now() - SessionStartTime).GetTotalSeconds()));
-				
-		FEngineAnalytics::GetProvider().RecordEvent(TEXT("Editor.Usage.PVE.SessionEnded"), Attributes);
-	}
+	double TotalSeconds = (FDateTime::Now() - SessionStartTime).GetTotalSeconds();
+	PV::Analytics::SendSessionEndedEvent(TotalSeconds);
 	
 	FPCGEditor::OnClose();
 
diff --git a/Engine/Plugins/Experimental/ProceduralVegetationEditor/Source/ProceduralVegetationEditor/Private/Visualizations/PVDataVisualization.cpp b/Engine/Plugins/Experimental/ProceduralVegetationEditor/Source/ProceduralVegetationEditor/Private/Visualizations/PVDataVisualization.cpp
--- a/Engine/Plugins/Experimental/ProceduralVegetationEditor/Source/ProceduralVegetationEditor/Private/Visualizations/PVDataVisualization.cpp
+++ b/Engine/Plugins/Experimental/ProceduralVegetationEditor/Source/ProceduralVegetationEditor/Private/Visualizations/PVDataVisualization.cpp
@@ -228,30 +228,6 @@ void FPVDataVisualization::BonesRenderer(const FManagedArrayCollection& InCollec
 	OutBounds = OutBounds + BoneComponent->CalcLocalBounds();
 }
 
-int32 GetMeshTriangles(const FString FoliagePath)
-{
-	UStaticMesh* StaticFoliageMesh = LoadObject<UStaticMesh>(nullptr, *FoliagePath);
-	USkeletalMesh* SkeletalFoliageMesh = LoadObject<USkeletalMesh>(nullptr, *FoliagePath);
-
-	if (StaticFoliageMesh)
-	{
-		return StaticFoliageMesh->GetNumTriangles(0);
-	}
-	else
-	{
-		int32 NumTriangles = 0;
-		FSkeletalMeshRenderData* SkelMeshRenderData = SkeletalFoliageMesh->GetResourceForRendering();
-		if (SkelMeshRenderData && SkelMeshRenderData->LODRenderData.Num() > 0)
-		{
-			const FSkeletalMeshLODRenderData& LODData = SkelMeshRenderData->LODRenderData[0];
-			NumTriangles = LODData.GetTotalFaces();
-		}
-
-		return NumTriangles;
-	}
-	
-}
-
 void FPVDataVisualization::FoliageGridRenderer(const FManagedArrayCollection& InCollection, FPCGSceneSetupParams& InOutParams,
 	FBoxSphereBounds& OutBounds)
 {
@@ -333,7 +309,7 @@ void FPVDataVisualization::FoliageGridRenderer(const FManagedArrayCollection& In
 			FoliagePos += FoliageHorizontalSize;
 			FoliageName = StaticFoliageMesh != nullptr ? StaticFoliageMesh->GetFName() : SkeletalFoliageMesh->GetFName();
 			
-			FText TriCount = FText::AsNumber(GetMeshTriangles(FoliagePath));
+			FText TriCount = FText::AsNumber(PV::Utilities::GetMeshTriangles(FoliagePath));
 			FString InTextToDraw = FString::Format(TEXT("Name: {0}\nTriangles:{1}"),{FoliageName.ToString(), TriCount.ToString()});
 
 			TObjectPtr<UTextRenderComponent> TextRenderComponent = NewObject<UTextRenderComponent>(GetTransientPackage(), NAME_None, RF_Transient);
