[PVE] Fixed horde issue

diff --git a/Engine/Plugins/Experimental/ProceduralVegetationEditor/Source/ProceduralVegetation/Private/Helpers/PVAnalyticsHelper.cpp b/Engine/Plugins/Experimental/ProceduralVegetationEditor/Source/ProceduralVegetation/Private/Helpers/PVAnalyticsHelper.cpp
--- a/Engine/Plugins/Experimental/ProceduralVegetationEditor/Source/ProceduralVegetation/Private/Helpers/PVAnalyticsHelper.cpp
+++ b/Engine/Plugins/Experimental/ProceduralVegetationEditor/Source/ProceduralVegetation/Private/Helpers/PVAnalyticsHelper.cpp
@@ -20,13 +20,13 @@ namespace PV::Analytics
 			FString EventName = FString::Printf(TEXT("Editor.Usage.PVE.%s"), *InEventName);
 			FEngineAnalytics::GetProvider().RecordEvent(EventName, InAttributes);
 
-			// FString AttributesData;
-			// for (int32 i = 0; i < InAttributes.Num(); i++)
-			// {
-			// 	AttributesData += FString::Printf(TEXT("(%s,%s),"), *InAttributes[i].GetName(), *InAttributes[i].GetValue());
-			// }
-			//
-			// UE_LOG(LogTemp, Display, TEXT("PV::SendEvent: %s : Attributes: {%s}"), *EventName, *AttributesData);
+			FString AttributesData;
+			for (int32 i = 0; i < InAttributes.Num(); i++)
+			{
+				AttributesData += FString::Printf(TEXT("(%s,%s),"), *InAttributes[i].GetName(), *InAttributes[i].GetValue());
+			}
+			
+			UE_LOG(LogTemp, Display, TEXT("PV::SendEvent: %s : Attributes: {%s}"), *EventName, *AttributesData);
 		}
 	}
 
@@ -65,7 +65,11 @@ namespace PV::Analytics
 		if (FEngineAnalytics::IsAvailable())
 		{
 			FString PropertyName = InProperty->GetName();
-			FString PropertyCategory = InProperty->GetMetaData(TEXT("Category"));
+
+			FString PropertyCategory = "N/A";
+#if WITH_METADATA
+			PropertyCategory = InProperty->GetMetaData(TEXT("Category"));
+#endif
 			
 			TArray<FAnalyticsEventAttribute> Attributes;
 			Attributes.Add(FAnalyticsEventAttribute(TEXT("Name"),  InNodeType));
