[PVE] Asset Replace Policy in export

diff --git a/Engine/Plugins/Experimental/ProceduralVegetationEditor/Source/ProceduralVegetation/Public/PVExportParams.h b/Engine/Plugins/Experimental/ProceduralVegetationEditor/Source/ProceduralVegetation/Public/PVExportParams.h
--- a/Engine/Plugins/Experimental/ProceduralVegetationEditor/Source/ProceduralVegetation/Public/PVExportParams.h
+++ b/Engine/Plugins/Experimental/ProceduralVegetationEditor/Source/ProceduralVegetation/Public/PVExportParams.h
@@ -3,6 +3,7 @@
 
 #include "Engine/EngineTypes.h"
 #include "PVWindSettings.h"
+
 #include "Misc/Paths.h"
 #include "PVExportParams.generated.h"
 
@@ -13,6 +14,19 @@ enum class EPVExportMeshType: uint8
 	SkeletalMesh,
 };
 
+UENUM()
+enum class EPVAssetReplacePolicy : uint8
+{
+	/** Create new assets with numbered suffixes */
+	Append,
+
+	/** Replaces existing asset with new asset */
+	Replace,
+
+	/** Ignores the new asset and keeps the existing asset */
+	Ignore,
+};
+
 USTRUCT()
 struct PROCEDURALVEGETATION_API FPVExportParams
 {
@@ -24,28 +38,32 @@ struct PROCEDURALVEGETATION_API FPVExportParams
 	UPROPERTY()
 	bool bShouldExportFoliage = true;
 
-	UPROPERTY(EditAnywhere, Category = "Export", Meta = (ContentDir, DisplayName = "Content Browser Folder", Tooltip="Path to the folder where exported mesh will be saved\n\nPath to the folder where exported mesh will be saved."))
+	UPROPERTY(EditAnywhere, Category = "Asset", Meta = (ContentDir, DisplayName = "Content Browser Folder", Tooltip="Path to the folder where exported mesh will be saved\n\nPath to the folder where exported mesh will be saved."))
 	FDirectoryPath ContentBrowserFolder;
 
-	UPROPERTY(EditAnywhere, Category = "Export", Meta = (DisplayName = "Mesh Name", Tooltip="Name of the exported mesh\n\nName of the exported mesh."))
+	UPROPERTY(EditAnywhere, Category = "Asset", Meta = (DisplayName = "Mesh Name", Tooltip="Name of the exported mesh\n\nName of the exported mesh."))
 	FName MeshName;
 
-	UPROPERTY(EditAnywhere, Category = "Export", Meta = (DisplayName = "Export Mesh Type", Tooltip="Static Mesh or Skeletal Mesh."))
+	UPROPERTY(EditAnywhere, Category = "Asset", Meta = (DisplayName = "Asset Replace Policy", Tooltip="Export behaviour if asset already exists."))
+	EPVAssetReplacePolicy ReplacePolicy = EPVAssetReplacePolicy::Replace;
+
+	UPROPERTY(EditAnywhere, Category = "Mesh", Meta = (DisplayName = "Export Mesh Type", Tooltip="Static Mesh or Skeletal Mesh."))
 	EPVExportMeshType ExportMeshType = EPVExportMeshType::SkeletalMesh;
 
-	UPROPERTY(EditAnywhere, Category = "Export", Meta = (DisplayName = "Create Nanite Foliage", Tooltip="Add Nanite foliage support for the exported mesh."))
+	UPROPERTY(EditAnywhere, Category = "Mesh | Nanite", Meta = (DisplayName = "Create Nanite Foliage", Tooltip="Add Nanite foliage support for the exported mesh."))
 	bool bCreateNaniteFoliage = true;
 	
-	UPROPERTY(EditAnywhere, Category = "Export", Meta = (DisplayName = "Nanite Shape Preservation Method", Tooltip="Shape preservation setting to be used in nanite settings."))
+	UPROPERTY(EditAnywhere, Category = "Mesh | Nanite", Meta = (DisplayName = "Nanite Shape Preservation Method", Tooltip="Shape preservation setting to be used in nanite settings."))
 	ENaniteShapePreservation NaniteShapePreservation = ENaniteShapePreservation::Voxelize;
 
-	UPROPERTY(EditAnywhere, Category = "Export", Meta = (EditCondition = "ExportMeshType == EPVExportMeshType::SkeletalMesh"))
+	UPROPERTY(EditAnywhere, Category = "Mesh | Dynamic Wind", Meta = (EditCondition = "ExportMeshType == EPVExportMeshType::SkeletalMesh"))
 	TObjectPtr<UPVWindSettings> WindSettings;
 
 	void Initialize(const FString& InAssetPath, const FString& InName, const FString& InDefaultWindSettingsPath)
 	{
 		ContentBrowserFolder.Path = FPaths::GetPath(InAssetPath);
 		MeshName = FName(InName);
+		ReplacePolicy = EPVAssetReplacePolicy::Replace;
 
 		WindSettings = LoadObject<UPVWindSettings>(nullptr, InDefaultWindSettingsPath);
 	}
diff --git a/Engine/Plugins/Experimental/ProceduralVegetationEditor/Source/ProceduralVegetationEditor/Private/Helpers/PVExportHelper.cpp b/Engine/Plugins/Experimental/ProceduralVegetationEditor/Source/ProceduralVegetationEditor/Private/Helpers/PVExportHelper.cpp
--- a/Engine/Plugins/Experimental/ProceduralVegetationEditor/Source/ProceduralVegetationEditor/Private/Helpers/PVExportHelper.cpp
+++ b/Engine/Plugins/Experimental/ProceduralVegetationEditor/Source/ProceduralVegetationEditor/Private/Helpers/PVExportHelper.cpp
@@ -1,13 +1,13 @@
 ﻿// Copyright Epic Games, Inc. All Rights Reserved.
 
-
 #include "Helpers/PVExportHelper.h"
 
-#include "DynamicMeshEditor.h"
 #include "DynamicWindImportData.h"
 #include "DynamicWindSkeletalData.h"
 #include "EngineAnalytics.h"
+#include "IAssetTools.h"
 #include "MeshDescriptionToDynamicMesh.h"
+#include "ObjectTools.h"
 #include "PackageTools.h"
 #include "PlanarCut.h"
 #include "ProceduralVegetation.h"
@@ -19,7 +19,6 @@
 #include "SkinnedAssetCompiler.h"
 #include "StaticMeshCompiler.h"
 #include "UDynamicMesh.h"
-#include "ObjectTools.h"
 
 #include "Animation/Skeleton.h"
 
@@ -201,19 +200,6 @@ namespace PV::Export::Internal
 		}
 	}
 
-	template<typename T>
-	TObjectPtr<T> GetAssetOfType(const FString& LongPackageName)
-	{
-		UPackage* Package = LoadPackage(nullptr, *LongPackageName, LOAD_None);
-		if (Package == nullptr)
-		{
-			return nullptr;
-		}
-
-		const FName AssetName = FPackageName::GetShortFName(Package->GetFName());
-		return FindObjectFast<T>(Package, AssetName);
-	}
-
 	UPackage* CreateAssetPackage(const FString& LongPackageName)
 	{
 		UPackage* Package = LoadPackage(nullptr, *LongPackageName, LOAD_None);
@@ -257,10 +243,6 @@ namespace PV::Export::Internal
 			TransformIdx = Idx++;
 		}
 
-		FMeshDescription MeshDescription;
-		FStaticMeshAttributes Attributes(MeshDescription);
-		Attributes.Register();
-
 		constexpr bool bCenterPivot = false;
 		FTransform TransformOut;
 		FDynamicMesh3 CombinedMesh;
@@ -1487,8 +1469,20 @@ namespace PV::Export
 	{
 		TRACE_CPUPROFILER_EVENT_SCOPE(PV::Export::ExportCollectionAsMesh);
 
-		const FString MeshName = ExportParams.MeshName.ToString();
-		const FString MeshPackageName = ExportParams.GetOutputMeshPackagePath();
+		FString MeshName = ExportParams.MeshName.ToString();
+		FString MeshPackageName = ExportParams.GetOutputMeshPackagePath();
+		if (FPackageName::DoesPackageExist(MeshPackageName))
+		{
+			if (ExportParams.ReplacePolicy == EPVAssetReplacePolicy::Append)
+			{
+				IAssetTools::Get().CreateUniqueAssetName(MeshPackageName, TEXT(""), MeshPackageName, MeshName);
+			}
+			else if (ExportParams.ReplacePolicy == EPVAssetReplacePolicy::Ignore)
+			{
+				return false;
+			}
+		}
+
 		const bool bCreateNaniteFoliage = ExportParams.bCreateNaniteFoliage && NaniteAssembliesSupported();
 		const bool bShouldExportFoliage = ExportParams.bShouldExportFoliage;
 
@@ -1506,13 +1500,14 @@ namespace PV::Export
 
 			Analytics::SendExportStartEvent(ExportAnalyticsAttributes);
 
-			if (TObjectPtr<UStaticMesh> ExistingMeshAsset = Internal::GetAssetOfType<UStaticMesh>(MeshPackageName))
+			TObjectPtr<UPackage> MeshPackage = Internal::CreateAssetPackage(MeshPackageName);
+			if (TObjectPtr<UStaticMesh> ExistingMeshAsset = Cast<UStaticMesh>(MeshPackage->FindAssetInPackage()))
 			{
 				ExistingMeshAsset->SetExternalPackage(nullptr);
 				ExistingMeshAsset->MarkAsGarbage();
 			}
 
-			const TObjectPtr<UStaticMesh> ExportMesh = NewObject<UStaticMesh>(Internal::CreateAssetPackage(MeshPackageName), *MeshName, RF_Standalone | RF_Public);
+			const TObjectPtr<UStaticMesh> ExportMesh = NewObject<UStaticMesh>(MeshPackage, *MeshName, RF_Standalone | RF_Public);
 
 			ModifiedPackages.Add(ExportMesh->GetOutermost());
 
@@ -1539,8 +1534,21 @@ namespace PV::Export
 				return false;
 			}
 
-			const FString SkeletonName = ExportParams.GetOutputSkeletonName();
-			const FString SkeletonPackageName = ExportParams.GetOutputSkeletonPackagePath();
+			FString SkeletonName = ExportParams.GetOutputSkeletonName();
+			FString SkeletonPackageName = ExportParams.GetOutputSkeletonPackagePath();
+
+			if (FPackageName::DoesPackageExist(SkeletonPackageName))
+			{
+				if (ExportParams.ReplacePolicy == EPVAssetReplacePolicy::Append)
+				{
+					IAssetTools::Get().CreateUniqueAssetName(SkeletonPackageName, TEXT(""), SkeletonPackageName, SkeletonName);
+				}
+				else if (ExportParams.ReplacePolicy == EPVAssetReplacePolicy::Ignore)
+				{
+					return false;
+				}
+			}
+			
 			if (Utilities::DoesConflictingPackageExist(SkeletonPackageName, USkeleton::StaticClass()))
 			{
 				return false;
@@ -1548,20 +1556,23 @@ namespace PV::Export
 
 			Analytics::SendExportStartEvent(ExportAnalyticsAttributes);
 
-			if (TObjectPtr<USkeletalMesh> ExistingMeshAsset = Internal::GetAssetOfType<USkeletalMesh>(MeshPackageName))
+			TObjectPtr<UPackage> MeshPackage = Internal::CreateAssetPackage(MeshPackageName);
+			TObjectPtr<UPackage> SkeletonPackage = Internal::CreateAssetPackage(SkeletonPackageName);
+
+			if (TObjectPtr<USkeletalMesh> ExistingMeshAsset = Cast<USkeletalMesh>(MeshPackage->FindAssetInPackage()))
 			{
 				ExistingMeshAsset->SetExternalPackage(nullptr);
 				ExistingMeshAsset->MarkAsGarbage();
 			}
 			
-			if (TObjectPtr<USkeleton> ExistingSkeletonAsset = Internal::GetAssetOfType<USkeleton>(SkeletonPackageName))
+			if (TObjectPtr<USkeleton> ExistingSkeletonAsset = Cast<USkeleton>(SkeletonPackage->FindAssetInPackage()))
 			{
 				ExistingSkeletonAsset->SetExternalPackage(nullptr);
 				ExistingSkeletonAsset->MarkAsGarbage();
 			}
 
-			const TObjectPtr<USkeletalMesh> ExportMesh = NewObject<USkeletalMesh>(Internal::CreateAssetPackage(MeshPackageName), *MeshName, RF_Standalone | RF_Public);
-			const TObjectPtr<USkeleton> ExportSkeleton = NewObject<USkeleton>(Internal::CreateAssetPackage(SkeletonPackageName), *SkeletonName, RF_Standalone | RF_Public);
+			const TObjectPtr<USkeletalMesh> ExportMesh = NewObject<USkeletalMesh>(MeshPackage, *MeshName, RF_Standalone | RF_Public);
+			const TObjectPtr<USkeleton> ExportSkeleton = NewObject<USkeleton>(SkeletonPackage, *SkeletonName, RF_Standalone | RF_Public);
 			
 			ExportMesh->SetSkeleton(ExportSkeleton);
 
diff --git a/Engine/Plugins/Experimental/ProceduralVegetationEditor/Source/ProceduralVegetationEditor/Private/PVEditor.cpp b/Engine/Plugins/Experimental/ProceduralVegetationEditor/Source/ProceduralVegetationEditor/Private/PVEditor.cpp
--- a/Engine/Plugins/Experimental/ProceduralVegetationEditor/Source/ProceduralVegetationEditor/Private/PVEditor.cpp
+++ b/Engine/Plugins/Experimental/ProceduralVegetationEditor/Source/ProceduralVegetationEditor/Private/PVEditor.cpp
@@ -388,7 +388,7 @@ FReply FPVEditor::ExportInternal()
 		return !bHasInvalidPath;
 	};
 
-	const static auto ShouldOverwriteUnsavedChanges = [](const TArray<UPCGNode*>& NodesToExport)->bool 
+	const static auto ShouldOverwriteAssets = [](const TArray<UPCGNode*>& NodesToExport)->bool 
 	{
 		TArray<FString> ExportPaths;
 		ExportPaths.Reserve(NodesToExport.Num() * 2);
@@ -397,6 +397,10 @@ FReply FPVEditor::ExportInternal()
 		{
 			const UPVOutputSettings* OutputSettings = CastChecked<const UPVOutputSettings>(Node->GetSettings());
 			const FPVExportParams& ExportParams = OutputSettings->ExportSettings;
+			if (ExportParams.ReplacePolicy != EPVAssetReplacePolicy::Replace)
+			{
+				continue;
+			}
 			ExportPaths.Add(ExportParams.GetOutputMeshPackagePath());
 			if (ExportParams.ExportMeshType == EPVExportMeshType::SkeletalMesh)
 			{
@@ -549,7 +553,7 @@ FReply FPVEditor::ExportInternal()
 		return FReply::Handled();
 	}
 
-	if (!ShouldOverwriteUnsavedChanges(NodesToExport))
+	if (!ShouldOverwriteAssets(NodesToExport))
 	{
 		return FReply::Handled();
 	}
@@ -634,7 +638,7 @@ FReply FPVEditor::ExportInternal()
 			else
 			{
 				ErrorMessages.AddError(FText::Format(
-					LOCTEXT("ExportFailed_UnknownError", "Failed to export mesh \"{0}\" due to unkown error, export aborted."),
+					LOCTEXT("ExportFailed_UnknownError", "Failed to export mesh \"{0}\" due to unknown error, export aborted."),
 					FText::FromName(OutputSettings->ExportSettings.MeshName)
 				));
 			}
