[PVE] Warning for nanite support

diff --git a/Engine/Plugins/Experimental/ProceduralVegetationEditor/Source/ProceduralVegetationEditor/Private/Customizations/PVOutputSettingsCustomizations.cpp b/Engine/Plugins/Experimental/ProceduralVegetationEditor/Source/ProceduralVegetationEditor/Private/Customizations/PVOutputSettingsCustomizations.cpp
--- a/Engine/Plugins/Experimental/ProceduralVegetationEditor/Source/ProceduralVegetationEditor/Private/Customizations/PVOutputSettingsCustomizations.cpp
+++ b/Engine/Plugins/Experimental/ProceduralVegetationEditor/Source/ProceduralVegetationEditor/Private/Customizations/PVOutputSettingsCustomizations.cpp
@@ -64,15 +64,33 @@ void FPVOutputSettingsCustomizations::CustomizeChildren(
 	}
 }
 
+bool IsNaniteSupported()
+{
+	return DoesPlatformSupportNanite(GMaxRHIShaderPlatform);
+}
+
+FText FPVOutputSettingsCustomizations::GetWarningTextForAssemblies() const
+{
+	FText WarningText;
+
+	if (!NaniteAssembliesSupported())
+	{
+		WarningText = FText::FromString("WARNING!! Enable Nanite Foliage in Project Settings to enable support for Nanite Assemblies.");
+	}
+
+	if (!IsNaniteSupported())
+	{
+		WarningText = FText::FromString(FString::Format(TEXT("{0}{1}WARNING!! Nanite is not supported. Check default RHI in Project Settings for Nanite Assemblies to work."),{WarningText.ToString(), WarningText.IsEmpty()? "" : "\n"}));
+	}
+
+	return WarningText;
+}
+
 void FPVOutputSettingsCustomizations::CustomizeNaniteAssembliesWidget(
 	const TSharedPtr<IPropertyHandle>& ChildHandle,
 	IDetailPropertyRow& RowBuilder
 )
 {
-	const static FText WarningText = FText::FromString(
-		"WARNING!! Enable Nanite Foliage in Project Settings to enable support for Nanite Assemblies."
-	);
-
 	const TSharedRef<SWidget> NameWidget = ChildHandle->CreatePropertyNameWidget();
 	const TSharedRef<SWidget> ValueWidget =
 		SNew(SHorizontalBox)
@@ -89,7 +107,7 @@ void FPVOutputSettingsCustomizations::CustomizeNaniteAssembliesWidget(
 				{
 					if (bool bCreateNaniteAssemblies; ChildHandle->GetValue(bCreateNaniteAssemblies) == FPropertyAccess::Success)
 					{
-						return (bCreateNaniteAssemblies && !NaniteAssembliesSupported())
+						return (bCreateNaniteAssemblies && (!NaniteAssembliesSupported() || !IsNaniteSupported()))
 							? EVisibility::Visible
 							: EVisibility::Collapsed;
 					}
@@ -107,8 +125,8 @@ void FPVOutputSettingsCustomizations::CustomizeNaniteAssembliesWidget(
 					.ColorAndOpacity(FLinearColor::Yellow)
 					.ShadowColorAndOpacity(FLinearColor::Black)
 					.ShadowOffset(FVector2D(3.0f, 3.0f))
-					.Text(WarningText)
-					.ToolTipText(WarningText)
+					.Text(this, &FPVOutputSettingsCustomizations::GetWarningTextForAssemblies)
+					.ToolTipText(this, &FPVOutputSettingsCustomizations::GetWarningTextForAssemblies)
 				]
 				+ SVerticalBox::Slot()
 				.VAlign(VAlign_Bottom)
diff --git a/Engine/Plugins/Experimental/ProceduralVegetationEditor/Source/ProceduralVegetationEditor/Private/Customizations/PVOutputSettingsCustomizations.h b/Engine/Plugins/Experimental/ProceduralVegetationEditor/Source/ProceduralVegetationEditor/Private/Customizations/PVOutputSettingsCustomizations.h
--- a/Engine/Plugins/Experimental/ProceduralVegetationEditor/Source/ProceduralVegetationEditor/Private/Customizations/PVOutputSettingsCustomizations.h
+++ b/Engine/Plugins/Experimental/ProceduralVegetationEditor/Source/ProceduralVegetationEditor/Private/Customizations/PVOutputSettingsCustomizations.h
@@ -13,6 +13,7 @@ public:
 
 	virtual void CustomizeHeader(TSharedRef<IPropertyHandle> PropertyHandle, FDetailWidgetRow& HeaderRow, IPropertyTypeCustomizationUtils& CustomizationUtils) override;
 	virtual void CustomizeChildren(TSharedRef<IPropertyHandle> PropertyHandle, IDetailChildrenBuilder& ChildBuilder, IPropertyTypeCustomizationUtils& CustomizationUtils) override;
+	FText GetWarningTextForAssemblies() const;
 
 private:
 	void CustomizeNaniteAssembliesWidget(const TSharedPtr<IPropertyHandle>& ChildHandle, IDetailPropertyRow& RowBuilder);
diff --git a/Engine/Plugins/Experimental/ProceduralVegetationEditor/Source/ProceduralVegetationEditor/Private/Helpers/PVExportHelper.cpp b/Engine/Plugins/Experimental/ProceduralVegetationEditor/Source/ProceduralVegetationEditor/Private/Helpers/PVExportHelper.cpp
--- a/Engine/Plugins/Experimental/ProceduralVegetationEditor/Source/ProceduralVegetationEditor/Private/Helpers/PVExportHelper.cpp
+++ b/Engine/Plugins/Experimental/ProceduralVegetationEditor/Source/ProceduralVegetationEditor/Private/Helpers/PVExportHelper.cpp
@@ -253,6 +253,11 @@ namespace PV::Export::Internal
 		return NewMesh;
 	}
 
+	bool IsNaniteSupported()
+	{
+		return DoesPlatformSupportNanite(GMaxRHIShaderPlatform);
+	}
+
 	FGeometryScriptCopyMeshToAssetOptions GetCopyMeshToAssetOptions(const TArray<TObjectPtr<UMaterialInterface>>& InMaterials, const FMeshNaniteSettings& InNaniteSettings, const ENaniteShapePreservation InShapePreservation)
 	{
 		FGeometryScriptCopyMeshToAssetOptions Options;
@@ -263,7 +268,7 @@ namespace PV::Export::Internal
 		Options.NewMaterials = InMaterials;
 
 		Options.NewNaniteSettings = InNaniteSettings;
-		Options.NewNaniteSettings.bEnabled = true;
+		Options.NewNaniteSettings.bEnabled = IsNaniteSupported();
 		Options.NewNaniteSettings.ShapePreservation = InShapePreservation;
 		Options.bApplyNaniteSettings = true;
 
diff --git a/Engine/Plugins/Experimental/ProceduralVegetationEditor/Source/ProceduralVegetationEditor/ProceduralVegetationEditor.Build.cs b/Engine/Plugins/Experimental/ProceduralVegetationEditor/Source/ProceduralVegetationEditor/ProceduralVegetationEditor.Build.cs
--- a/Engine/Plugins/Experimental/ProceduralVegetationEditor/Source/ProceduralVegetationEditor/ProceduralVegetationEditor.Build.cs
+++ b/Engine/Plugins/Experimental/ProceduralVegetationEditor/Source/ProceduralVegetationEditor/ProceduralVegetationEditor.Build.cs
@@ -50,7 +50,8 @@ public class ProceduralVegetationEditor : ModuleRules
 				"ToolWidgets",
 				"UnrealEd",
 				"DynamicWind",
-				"DynamicWindEditor"
+				"DynamicWindEditor",
+				"RHI"
 			}
 		);
 	}
