Fix PVE asset replace policy not working for in-memory only assets

diff --git a/Engine/Plugins/Experimental/ProceduralVegetationEditor/Source/ProceduralVegetation/Private/PVExportParams.cpp b/Engine/Plugins/Experimental/ProceduralVegetationEditor/Source/ProceduralVegetation/Private/PVExportParams.cpp
--- a/Engine/Plugins/Experimental/ProceduralVegetationEditor/Source/ProceduralVegetation/Private/PVExportParams.cpp
+++ b/Engine/Plugins/Experimental/ProceduralVegetationEditor/Source/ProceduralVegetation/Private/PVExportParams.cpp
@@ -7,12 +7,16 @@
 #include "Engine/SkeletalMesh.h"
 #include "Animation/Skeleton.h"
 
+UClass* FPVExportParams::GetMeshClass() const
+{
+	return ExportMeshType == EPVExportMeshType::StaticMesh ? UStaticMesh::StaticClass() : USkeletalMesh::StaticClass();
+}
+
 bool FPVExportParams::Validate(FString& OutError) const
 {
 	const FString& OutputPath = ContentBrowserFolder.Path;
 
-	UClass* MeshClass = ExportMeshType == EPVExportMeshType::StaticMesh ? UStaticMesh::StaticClass() : USkeletalMesh::StaticClass();
-	if (!PV::Utilities::ValidateAssetPathAndName(MeshName.ToString(), OutputPath, MeshClass, OutError))
+	if (!PV::Utilities::ValidateAssetPathAndName(MeshName.ToString(), OutputPath, GetMeshClass(), OutError))
 	{
 		return false;
 	}
diff --git a/Engine/Plugins/Experimental/ProceduralVegetationEditor/Source/ProceduralVegetation/Public/PVExportParams.h b/Engine/Plugins/Experimental/ProceduralVegetationEditor/Source/ProceduralVegetation/Public/PVExportParams.h
--- a/Engine/Plugins/Experimental/ProceduralVegetationEditor/Source/ProceduralVegetation/Public/PVExportParams.h
+++ b/Engine/Plugins/Experimental/ProceduralVegetationEditor/Source/ProceduralVegetation/Public/PVExportParams.h
@@ -90,5 +90,7 @@ struct PROCEDURALVEGETATION_API FPVExportParams
 		return FPaths::Combine(ContentBrowserFolder.Path, GetOutputSkeletonName());
 	}
 
+	UClass* GetMeshClass() const;
+
 	bool Validate(FString& OutError) const;
 };
\ No newline at end of file
diff --git a/Engine/Plugins/Experimental/ProceduralVegetationEditor/Source/ProceduralVegetationEditor/Private/Helpers/PVExportHelper.cpp b/Engine/Plugins/Experimental/ProceduralVegetationEditor/Source/ProceduralVegetationEditor/Private/Helpers/PVExportHelper.cpp
--- a/Engine/Plugins/Experimental/ProceduralVegetationEditor/Source/ProceduralVegetationEditor/Private/Helpers/PVExportHelper.cpp
+++ b/Engine/Plugins/Experimental/ProceduralVegetationEditor/Source/ProceduralVegetationEditor/Private/Helpers/PVExportHelper.cpp
@@ -200,6 +200,33 @@ namespace PV::Export::Internal
 		}
 	}
 
+	bool GetValidAssetPath(EPVAssetReplacePolicy AssetReplacePolicy, UClass* AssetClass, FString& InOutPackageName, FString& InOutAssetName)
+	{
+		if (!FPackageName::IsValidTextForLongPackageName(InOutPackageName))
+		{
+			return false;
+		}
+
+		if (FindPackage(nullptr, *InOutPackageName) != nullptr || FPackageName::DoesPackageExist(InOutPackageName))
+		{
+			if (AssetReplacePolicy == EPVAssetReplacePolicy::Append)
+			{
+				IAssetTools::Get().CreateUniqueAssetName(InOutPackageName, TEXT(""), InOutPackageName, InOutAssetName);
+			}
+			else if (AssetReplacePolicy == EPVAssetReplacePolicy::Ignore)
+			{
+				return false;
+			}
+		}
+
+		if (Utilities::DoesConflictingPackageExist(InOutPackageName, AssetClass))
+		{
+			return false;
+		}
+
+		return true;
+	}
+
 	UPackage* CreateAssetPackage(const FString& LongPackageName)
 	{
 		UPackage* Package = LoadPackage(nullptr, *LongPackageName, LOAD_None);
@@ -213,8 +240,6 @@ namespace PV::Export::Internal
 		return Package;
 	}
 
-	
-
 	TObjectPtr<UDynamicMesh> CollectionToDynamicMesh(const FManagedArrayCollection& Collection)
 	{
 		TRACE_CPUPROFILER_EVENT_SCOPE(PV::Export::CollectionToDynamicMesh);
@@ -1476,18 +1501,12 @@ namespace PV::Export
 
 		FString MeshName = ExportParams.MeshName.ToString();
 		FString MeshPackageName = ExportParams.GetOutputMeshPackagePath();
-		if (FPackageName::DoesPackageExist(MeshPackageName))
+
+		if (!Internal::GetValidAssetPath(ExportParams.ReplacePolicy, ExportParams.GetMeshClass(), MeshPackageName, MeshName))
 		{
-			if (ExportParams.ReplacePolicy == EPVAssetReplacePolicy::Append)
-			{
-				IAssetTools::Get().CreateUniqueAssetName(MeshPackageName, TEXT(""), MeshPackageName, MeshName);
-			}
-			else if (ExportParams.ReplacePolicy == EPVAssetReplacePolicy::Ignore)
-			{
-				return false;
-			}
+			return false;
 		}
-
+		
 		const bool bCreateNaniteFoliage = ExportParams.bCreateNaniteFoliage && NaniteAssembliesSupported();
 		const bool bShouldExportFoliage = ExportParams.bShouldExportFoliage;
 
@@ -1498,11 +1517,6 @@ namespace PV::Export
 
 		if (ExportParams.ExportMeshType == EPVExportMeshType::StaticMesh)
 		{
-			if (Utilities::DoesConflictingPackageExist(MeshPackageName, UStaticMesh::StaticClass()))
-			{
-				return false;
-			}
-
 			Analytics::SendExportStartEvent(ExportAnalyticsAttributes);
 
 			TObjectPtr<UPackage> MeshPackage = Internal::CreateAssetPackage(MeshPackageName);
@@ -1534,27 +1548,10 @@ namespace PV::Export
 		}
 		else if (ExportParams.ExportMeshType == EPVExportMeshType::SkeletalMesh)
 		{
-			if (Utilities::DoesConflictingPackageExist(*MeshPackageName, USkeletalMesh::StaticClass()))
-			{
-				return false;
-			}
-
 			FString SkeletonName = ExportParams.GetOutputSkeletonName();
 			FString SkeletonPackageName = ExportParams.GetOutputSkeletonPackagePath();
 
-			if (FPackageName::DoesPackageExist(SkeletonPackageName))
-			{
-				if (ExportParams.ReplacePolicy == EPVAssetReplacePolicy::Append)
-				{
-					IAssetTools::Get().CreateUniqueAssetName(SkeletonPackageName, TEXT(""), SkeletonPackageName, SkeletonName);
-				}
-				else if (ExportParams.ReplacePolicy == EPVAssetReplacePolicy::Ignore)
-				{
-					return false;
-				}
-			}
-			
-			if (Utilities::DoesConflictingPackageExist(SkeletonPackageName, USkeleton::StaticClass()))
+			if (!Internal::GetValidAssetPath(ExportParams.ReplacePolicy, USkeleton::StaticClass(), SkeletonPackageName, SkeletonName))
 			{
 				return false;
 			}
