[PVE] Fix for Crash on invalid data
Hidden the data asset editable controls under CVar
Validations added in nodes
Added tooltip for windsettings

diff --git a/Engine/Plugins/Experimental/ProceduralVegetationEditor/Source/ProceduralVegetation/Private/Facades/PVBoneFacade.cpp b/Engine/Plugins/Experimental/ProceduralVegetationEditor/Source/ProceduralVegetation/Private/Facades/PVBoneFacade.cpp
--- a/Engine/Plugins/Experimental/ProceduralVegetationEditor/Source/ProceduralVegetation/Private/Facades/PVBoneFacade.cpp
+++ b/Engine/Plugins/Experimental/ProceduralVegetationEditor/Source/ProceduralVegetation/Private/Facades/PVBoneFacade.cpp
@@ -7,6 +7,7 @@
 #include "Facades/PVPlantFacade.h"
 #include "Facades/PVPointFacade.h"
 #include "GeometryCollection/Facades/CollectionConstraintOverrideFacade.h"
+#include "Helpers/PVUtilities.h"
 
 class FGeometryCollection;
 
@@ -246,12 +247,17 @@ namespace PV::Facades
 		SetBranchSimulationGroupIndex(PointFacade, BranchFacade);
 	}
 	
-	void FBoneFacade::CreateBoneData(TArray<FBoneNode>& Bones,const float ReductionStrength)
+	bool FBoneFacade::CreateBoneData(TArray<FBoneNode>& Bones,const float ReductionStrength)
 	{
 		FPointFacade PointFacade(VertexPointIds.GetConstCollection());
 		FBranchFacade BranchFacade(VertexPointIds.GetConstCollection());
 		FPlantFacade PlantFacade(VertexPointIds.GetConstCollection());
 
+		if (!(Utilities::IsValidPVData(VertexPointIds.GetConstCollection()) && PlantFacade.IsValid()))
+		{
+			return false;	
+		}
+		
 		SetWindSimulationData(*VertexPointIds.GetCollection());
 		
 		const TManagedArray<FVector3f>& Positions = PointFacade.GetPositions();
@@ -393,6 +399,8 @@ namespace PV::Facades
 		}
 
 		ReindexBones(Bones);
+
+		return true;
 	}
 
 	FBoneNode* FBoneFacade::FindClosestBone(const FManagedArrayCollection& Collection, TArray<FBoneNode>& Bones, const int32 PointID)
diff --git a/Engine/Plugins/Experimental/ProceduralVegetationEditor/Source/ProceduralVegetation/Private/Helpers/PVUtilities.cpp b/Engine/Plugins/Experimental/ProceduralVegetationEditor/Source/ProceduralVegetation/Private/Helpers/PVUtilities.cpp
--- a/Engine/Plugins/Experimental/ProceduralVegetationEditor/Source/ProceduralVegetation/Private/Helpers/PVUtilities.cpp
+++ b/Engine/Plugins/Experimental/ProceduralVegetationEditor/Source/ProceduralVegetation/Private/Helpers/PVUtilities.cpp
@@ -2,6 +2,7 @@
 
 #include "Helpers/PVUtilities.h"
 #include "ProceduralVegetationModule.h"
+#include "ProceduralVegetationPreset.h"
 #include "AssetRegistry/AssetRegistryModule.h"
 #include "Components/InstancedSkinnedMeshComponent.h"
 #include "Engine/StaticMesh.h"
@@ -12,6 +13,8 @@
 #include "Nodes/PVImporterSettings.h"
 #include "Misc/PackageName.h"
 #include "AssetRegistry/AssetData.h"
+#include "Facades/PVBranchFacade.h"
+#include "Facades/PVPointFacade.h"
 #include "Misc/FileHelper.h"
 #include "Rendering/SkeletalMeshRenderData.h"
 
@@ -25,6 +28,7 @@ namespace PV::Utilities
 		FConsoleVariableDelegate::CreateLambda([](IConsoleVariable* ConsoleVariable)
 			{
 				UPVImporterSettings::StaticClass()->GetDefaultObject<UPVImporterSettings>()->bExposeToLibrary = ConsoleVariable->GetBool();
+				UProceduralVegetationPreset::ShowHideInternalProperties(ConsoleVariable->GetBool());
 			})
 	);
 #endif
@@ -246,4 +250,12 @@ namespace PV::Utilities
 
 		return NumTriangles;
 	}
+
+	bool IsValidPVData(const FManagedArrayCollection& Collection)
+	{
+		Facades::FPointFacade PointFacade = Facades::FPointFacade(Collection);
+		Facades::FBranchFacade BranchFacade = Facades::FBranchFacade(Collection);
+	
+		return PointFacade.IsValid() && BranchFacade.IsValid();
+	}
 }
diff --git a/Engine/Plugins/Experimental/ProceduralVegetationEditor/Source/ProceduralVegetation/Private/Nodes/PVBoneReductionSettings.cpp b/Engine/Plugins/Experimental/ProceduralVegetationEditor/Source/ProceduralVegetation/Private/Nodes/PVBoneReductionSettings.cpp
--- a/Engine/Plugins/Experimental/ProceduralVegetationEditor/Source/ProceduralVegetation/Private/Nodes/PVBoneReductionSettings.cpp
+++ b/Engine/Plugins/Experimental/ProceduralVegetationEditor/Source/ProceduralVegetation/Private/Nodes/PVBoneReductionSettings.cpp
@@ -69,7 +69,13 @@ bool FPVBoneReductionElement::ExecuteInternal(FPCGContext* InContext) const
 
 			TArray<PV::Facades::FBoneNode> BoneNodes;
 			float ReductionStrength = InputSettings->Strength;
-			BoneFacade.CreateBoneData(BoneNodes, ReductionStrength);
+
+			if (!BoneFacade.CreateBoneData(BoneNodes, ReductionStrength))
+			{
+				PCGLog::InputOutput::LogInvalidInputDataError(InContext);
+				return true;
+			}
+			
 			BoneFacade.SetBoneDataToCollection(BoneNodes);
 			
 			OutManagedArrayCollectionData->Initialize(MoveTemp(OutCollection));
diff --git a/Engine/Plugins/Experimental/ProceduralVegetationEditor/Source/ProceduralVegetation/Private/Nodes/PVImporterSettings.cpp b/Engine/Plugins/Experimental/ProceduralVegetationEditor/Source/ProceduralVegetation/Private/Nodes/PVImporterSettings.cpp
--- a/Engine/Plugins/Experimental/ProceduralVegetationEditor/Source/ProceduralVegetation/Private/Nodes/PVImporterSettings.cpp
+++ b/Engine/Plugins/Experimental/ProceduralVegetationEditor/Source/ProceduralVegetation/Private/Nodes/PVImporterSettings.cpp
@@ -58,8 +58,13 @@ bool FPVImporterElement::ExecuteInternal(FPCGContext* InContext) const
 	bool bSuccess = PV::LoadMegaPlantsJsonToCollection(Collection, SkeletonFilePath, ErrorMessage);
 	if (!bSuccess)
 	{
-		UE_LOG(LogProceduralVegetation, Error, TEXT("%s"), *ErrorMessage);
-		PCGLog::InputOutput::LogInvalidInputDataError(InContext);
+		PCGLog::LogErrorOnGraph(FText::FromString(ErrorMessage), InContext);
+		return true;
+	}
+
+	if (!PV::Utilities::IsValidPVData(Collection))
+	{
+		PCGLog::LogErrorOnGraph(FText::FromString(FString::Format(TEXT("Invalid data in the imported json {0}"), { SkeletonFilePath})), InContext);
 		return true;
 	}
 
diff --git a/Engine/Plugins/Experimental/ProceduralVegetationEditor/Source/ProceduralVegetation/Private/Nodes/PVPresetLoaderSettings.cpp b/Engine/Plugins/Experimental/ProceduralVegetationEditor/Source/ProceduralVegetation/Private/Nodes/PVPresetLoaderSettings.cpp
--- a/Engine/Plugins/Experimental/ProceduralVegetationEditor/Source/ProceduralVegetation/Private/Nodes/PVPresetLoaderSettings.cpp
+++ b/Engine/Plugins/Experimental/ProceduralVegetationEditor/Source/ProceduralVegetation/Private/Nodes/PVPresetLoaderSettings.cpp
@@ -5,6 +5,8 @@
 #include "ProceduralVegetationModule.h"
 #include "DataTypes/PVGrowthData.h"
 #include "ProceduralVegetationPreset.h"
+#include "Facades/PVBranchFacade.h"
+#include "Facades/PVPointFacade.h"
 
 #define LOCTEXT_NAMESPACE "PVPresetLoaderSettings"
 
@@ -56,8 +58,11 @@ TArray<FPCGPinProperties> UPVPresetLoaderSettings::OutputPinProperties() const
 	{
 		for (const auto& [VariantName, VariantData] : Preset->Variants)
 		{
-			FPCGPinProperties& Pin = Properties.Emplace_GetRef(*VariantName, GetOutputPinTypeIdentifier());
-			Pin.bAllowMultipleData = false;
+			if (PV::Utilities::IsValidPVData(VariantData))
+			{
+				FPCGPinProperties& Pin = Properties.Emplace_GetRef(*VariantName, GetOutputPinTypeIdentifier());
+				Pin.bAllowMultipleData = false;
+			}
 		}
 	}
 
@@ -107,16 +112,31 @@ bool FPVPresetLoaderElement::ExecuteInternal(FPCGContext* InContext) const
 		return true;
 	}
 
+	bool bValidDataExist = false;
 	for (const auto& [VariantName, VariantData] : Settings->Preset->Variants)
 	{
-		UPVGrowthData* OutVariantData = FPCGContext::NewObject_AnyThread<UPVGrowthData>(InContext);
-		OutVariantData->Initialize(CopyTemp(VariantData));
+		if (PV::Utilities::IsValidPVData(VariantData))
+		{
+			UPVGrowthData* OutVariantData = FPCGContext::NewObject_AnyThread<UPVGrowthData>(InContext);
+			OutVariantData->Initialize(CopyTemp(VariantData));
 
-		FPCGTaggedData& CollectionOutput = InContext->OutputData.TaggedData.Emplace_GetRef();
-		CollectionOutput.Data = OutVariantData;
-		CollectionOutput.Pin = *VariantName;
+			FPCGTaggedData& CollectionOutput = InContext->OutputData.TaggedData.Emplace_GetRef();
+			CollectionOutput.Data = OutVariantData;
+			CollectionOutput.Pin = *VariantName;
+			bValidDataExist = true;
+		}
+		else
+		{
+			PCGLog::LogWarningOnGraph(FText::FromString(FString::Format(TEXT("Skipping variation {0} due to invalid data."), {*VariantName})), InContext);
+		}
 	}
 
+	if (!bValidDataExist)
+	{
+		PCGLog::LogErrorOnGraph(FText::FromString(FString::Format(TEXT("No valid variation found in data asset {0}"), { Settings->Preset.GetFName().ToString()})), InContext);
+		return true;
+	}
+	
 	return true;
 }
 
diff --git a/Engine/Plugins/Experimental/ProceduralVegetationEditor/Source/ProceduralVegetation/Private/ProceduralVegetationPreset.cpp b/Engine/Plugins/Experimental/ProceduralVegetationEditor/Source/ProceduralVegetation/Private/ProceduralVegetationPreset.cpp
--- a/Engine/Plugins/Experimental/ProceduralVegetationEditor/Source/ProceduralVegetation/Private/ProceduralVegetationPreset.cpp
+++ b/Engine/Plugins/Experimental/ProceduralVegetationEditor/Source/ProceduralVegetation/Private/ProceduralVegetationPreset.cpp
@@ -8,6 +8,7 @@
 #include "Facades/PVBranchFacade.h"
 #include "Facades/PVProfileFacade.h"
 #include "Helpers/PVJSONHelper.h"
+#include "Helpers/PVUtilities.h"
 
 void FPVPresetVariationInfo::Fill(FName InName, FManagedArrayCollection Collection)
 {
@@ -132,6 +133,35 @@ void UProceduralVegetationPreset::FillVariationInfo()
 	}
 }
 
+void UProceduralVegetationPreset::ShowHideInternalProperties(bool bDebugEnable)
+{
+#if WITH_METADATA
+	FString CategoryName(TEXT("Preset Data"));
+	if (!bDebugEnable)
+	{
+		CategoryName = TEXT("Internal");
+	}
+
+	for (TFieldIterator<FProperty> It(StaticClass()); It; ++It)
+	{
+		FProperty* Property = *It;
+		if (Property && Property->HasMetaData(TEXT("DevelopmentOnly")))
+		{
+			Property->SetMetaData(TEXT("Category"), *CategoryName);
+		}
+	}
+
+	for (TFieldIterator<UFunction> It(StaticClass()); It; ++It)
+	{
+		UFunction* Function = *It;
+		if (Function && Function->HasMetaData(TEXT("DevelopmentOnly")))
+		{
+			Function->SetMetaData(TEXT("Category"), *CategoryName);
+		}
+	}
+#endif
+}
+
 void UProceduralVegetationPreset::CreateProfileDataAsset()
 {
 	FString ProfilePackageName(PlantProfileName);
diff --git a/Engine/Plugins/Experimental/ProceduralVegetationEditor/Source/ProceduralVegetation/Public/Facades/PVBoneFacade.h b/Engine/Plugins/Experimental/ProceduralVegetationEditor/Source/ProceduralVegetation/Public/Facades/PVBoneFacade.h
--- a/Engine/Plugins/Experimental/ProceduralVegetationEditor/Source/ProceduralVegetation/Public/Facades/PVBoneFacade.h
+++ b/Engine/Plugins/Experimental/ProceduralVegetationEditor/Source/ProceduralVegetation/Public/Facades/PVBoneFacade.h
@@ -52,7 +52,7 @@ namespace PV::Facades
 
 		void SetWindSimulationData(FManagedArrayCollection& Collection);
 		
-		void CreateBoneData(TArray<FBoneNode>& Bones,const float ReductionStrength);
+		bool CreateBoneData(TArray<FBoneNode>& Bones,const float ReductionStrength);
 
 		static FBoneNode* FindClosestBone(const FManagedArrayCollection& Collection, TArray<FBoneNode>& Bones, const int32 PointID);
 
diff --git a/Engine/Plugins/Experimental/ProceduralVegetationEditor/Source/ProceduralVegetation/Public/Helpers/PVUtilities.h b/Engine/Plugins/Experimental/ProceduralVegetationEditor/Source/ProceduralVegetation/Public/Helpers/PVUtilities.h
--- a/Engine/Plugins/Experimental/ProceduralVegetationEditor/Source/ProceduralVegetation/Public/Helpers/PVUtilities.h
+++ b/Engine/Plugins/Experimental/ProceduralVegetationEditor/Source/ProceduralVegetation/Public/Helpers/PVUtilities.h
@@ -24,4 +24,6 @@ namespace PV::Utilities
 	bool IsFileNameValid(FName FileName, FText& Reason);
 	bool PROCEDURALVEGETATION_API DoesConflictingPackageExist(const FString& PackageName, UClass* InClass);
 	int32 PROCEDURALVEGETATION_API GetMeshTriangles(const FString InMeshPath);
+
+	bool IsValidPVData(const FManagedArrayCollection& Collection);
 }
diff --git a/Engine/Plugins/Experimental/ProceduralVegetationEditor/Source/ProceduralVegetation/Public/PVExportParams.h b/Engine/Plugins/Experimental/ProceduralVegetationEditor/Source/ProceduralVegetation/Public/PVExportParams.h
--- a/Engine/Plugins/Experimental/ProceduralVegetationEditor/Source/ProceduralVegetation/Public/PVExportParams.h
+++ b/Engine/Plugins/Experimental/ProceduralVegetationEditor/Source/ProceduralVegetation/Public/PVExportParams.h
@@ -56,7 +56,7 @@ struct PROCEDURALVEGETATION_API FPVExportParams
 	UPROPERTY(EditAnywhere, Category = "Mesh | Nanite", Meta = (DisplayName = "Nanite Shape Preservation Method", Tooltip="Shape preservation setting to be used in nanite settings."))
 	ENaniteShapePreservation NaniteShapePreservation = ENaniteShapePreservation::Voxelize;
 
-	UPROPERTY(EditAnywhere, Category = "Mesh | Dynamic Wind", Meta = (EditCondition = "ExportMeshType == EPVExportMeshType::SkeletalMesh"))
+	UPROPERTY(EditAnywhere, Category = "Mesh | Dynamic Wind", Meta = (EditCondition = "ExportMeshType == EPVExportMeshType::SkeletalMesh", Tooltip="Wind Settings contains the wind simulation group data, Default presets are available for Saplings and Trees."))
 	TObjectPtr<UPVWindSettings> WindSettings;
 
 	void Initialize(const FString& InAssetPath, const FString& InName, const FString& InDefaultWindSettingsPath)
diff --git a/Engine/Plugins/Experimental/ProceduralVegetationEditor/Source/ProceduralVegetation/Public/ProceduralVegetationPreset.h b/Engine/Plugins/Experimental/ProceduralVegetationEditor/Source/ProceduralVegetation/Public/ProceduralVegetationPreset.h
--- a/Engine/Plugins/Experimental/ProceduralVegetationEditor/Source/ProceduralVegetation/Public/ProceduralVegetationPreset.h
+++ b/Engine/Plugins/Experimental/ProceduralVegetationEditor/Source/ProceduralVegetation/Public/ProceduralVegetationPreset.h
@@ -28,7 +28,7 @@ struct FPVPresetVariationInfo
 	void Fill(FName InName,FManagedArrayCollection Collection);
 };
 
-UCLASS(BlueprintType)
+UCLASS(BlueprintType, hidecategories="Internal")
 class PROCEDURALVEGETATION_API UProceduralVegetationPreset : public UDataAsset
 {
 	GENERATED_BODY()
@@ -37,26 +37,26 @@ public:
 	UProceduralVegetationPreset();
 
 	virtual void PostLoad() override;
-	
-	UPROPERTY(EditAnywhere, Category="Preset Data", meta = (DevelopmentOnly))
+
+	UPROPERTY(EditAnywhere, Category="Internal", meta = (DevelopmentOnly))
 	FDirectoryPath JsonDirectoryPath;
 	
-	UPROPERTY(EditAnywhere, Category = "Preset Data", Meta = (ContentDir, DevelopmentOnly))
+	UPROPERTY(EditAnywhere, Category = "Internal", Meta = (ContentDir, DevelopmentOnly))
 	bool bOverrideFolderPaths = false;
 	
-	UPROPERTY(EditAnywhere, Category = "Preset Data", Meta = (ContentDir, DevelopmentOnly , EditCondition = "bOverrideFolderPaths"))
+	UPROPERTY(EditAnywhere, Category = "Internal", Meta = (ContentDir, DevelopmentOnly , EditCondition = "bOverrideFolderPaths"))
 	FDirectoryPath FoliageFolder;
 	
-	UPROPERTY(EditAnywhere, Category = "Preset Data", Meta = (ContentDir, DevelopmentOnly , EditCondition = "bOverrideFolderPaths"))
+	UPROPERTY(EditAnywhere, Category = "Internal", Meta = (ContentDir, DevelopmentOnly , EditCondition = "bOverrideFolderPaths"))
 	FDirectoryPath MaterialsFolder;
 	
-	UPROPERTY(EditAnywhere, Category = "Preset Data", meta = (DevelopmentOnly))
+	UPROPERTY(EditAnywhere, Category = "Internal", meta = (DevelopmentOnly))
 	FString TrunkMaterialName;
 	
-	UPROPERTY(EditAnywhere, Category = "Preset Data", meta = (DevelopmentOnly))
+	UPROPERTY(EditAnywhere, Category = "Internal", meta = (DevelopmentOnly))
 	bool bCreateProfileDataAsset = false;
 	
-	UPROPERTY(EditAnywhere, Category = "Preset Data", meta = (DevelopmentOnly, EditCondition = "bCreateProfileDataAsset"))
+	UPROPERTY(EditAnywhere, Category = "Internal", meta = (DevelopmentOnly, EditCondition = "bCreateProfileDataAsset"))
 	FString PlantProfileName;
 	
 	UPROPERTY()
@@ -68,10 +68,12 @@ public:
 public:
 	void LoadFromVariantFiles(const TMap<FString, FString>& InVariantFiles);
 
-	UFUNCTION(CallInEditor, Category="Preset Data", meta = (DevelopmentOnly))
+	UFUNCTION(CallInEditor, Category="Internal", meta = (DevelopmentOnly))
 	void UpdateDataAsset();
 
 	void FillVariationInfo();
+
+	static void ShowHideInternalProperties(bool bDebugEnable);
 private:
 	void CreateProfileDataAsset();
 	
