Fix "conflicting asset exists" error being shown even when asset replace policy is compatible with the existing asset

diff --git a/Engine/Plugins/Experimental/ProceduralVegetationEditor/Source/ProceduralVegetation/Private/Helpers/PVUtilities.cpp b/Engine/Plugins/Experimental/ProceduralVegetationEditor/Source/ProceduralVegetation/Private/Helpers/PVUtilities.cpp
--- a/Engine/Plugins/Experimental/ProceduralVegetationEditor/Source/ProceduralVegetation/Private/Helpers/PVUtilities.cpp
+++ b/Engine/Plugins/Experimental/ProceduralVegetationEditor/Source/ProceduralVegetation/Private/Helpers/PVUtilities.cpp
@@ -178,13 +178,6 @@ namespace PV::Utilities
         	return false;
         }
 
-        const FString PackageName = FPaths::Combine(Path, AssetName);
-        if (DoesConflictingPackageExist(PackageName, InClass))
-        {
-			OutError = FString::Format(TEXT("Create a Unique name for the output as another asset of different type {0} exists with same name {1}"), { *InClass->GetName(), *PackageName});
-        	return false;
-        }
-
 		return true;
 	}
 
diff --git a/Engine/Plugins/Experimental/ProceduralVegetationEditor/Source/ProceduralVegetation/Private/PVExportParams.cpp b/Engine/Plugins/Experimental/ProceduralVegetationEditor/Source/ProceduralVegetation/Private/PVExportParams.cpp
--- a/Engine/Plugins/Experimental/ProceduralVegetationEditor/Source/ProceduralVegetation/Private/PVExportParams.cpp
+++ b/Engine/Plugins/Experimental/ProceduralVegetationEditor/Source/ProceduralVegetation/Private/PVExportParams.cpp
@@ -29,5 +29,27 @@ bool FPVExportParams::Validate(FString& OutError) const
 		}
 	}
 
+	if (ReplacePolicy == EPVAssetReplacePolicy::Replace)
+	{
+		const TCHAR* ConflictingAssetError = TEXT("Create a Unique name for the output as another asset of different type {0} exists with same name {1}");
+
+		const FString MeshPackagePath = GetOutputMeshPackagePath();
+		if (PV::Utilities::DoesConflictingPackageExist(MeshPackagePath, GetMeshClass()))
+		{
+			OutError = FString::Format(ConflictingAssetError, { *GetMeshClass()->GetName(), *MeshPackagePath });
+			return false;
+		}
+
+		if (ExportMeshType == EPVExportMeshType::SkeletalMesh)
+		{
+			const FString SkeletonPackagePath = GetOutputSkeletonPackagePath();
+			if (PV::Utilities::DoesConflictingPackageExist(SkeletonPackagePath, USkeleton::StaticClass()))
+			{
+				OutError = FString::Format(ConflictingAssetError, { *USkeleton::StaticClass()->GetName(), *SkeletonPackagePath });
+				return false;
+			}
+		}
+	}
+
 	return true;
 }
\ No newline at end of file
