[PVE] Fix usability issues with Mesher node settings

diff --git a/Engine/Plugins/Experimental/ProceduralVegetationEditor/Source/ProceduralVegetation/Private/Implementations/PVMeshBuilder.cpp b/Engine/Plugins/Experimental/ProceduralVegetationEditor/Source/ProceduralVegetation/Private/Implementations/PVMeshBuilder.cpp
--- a/Engine/Plugins/Experimental/ProceduralVegetationEditor/Source/ProceduralVegetation/Private/Implementations/PVMeshBuilder.cpp
+++ b/Engine/Plugins/Experimental/ProceduralVegetationEditor/Source/ProceduralVegetation/Private/Implementations/PVMeshBuilder.cpp
@@ -42,7 +42,7 @@ void MergeMaterials(const TArray<FTrunkGenerationMaterialSetup>& InMaterialSetup
 }
 
 void FPVMeshBuilder::GenerateGeometryCollection(const FManagedArrayCollection& InSkeletonCollection, const FPVMeshBuilderParams& MeshBuilderParams,
-                                           FGeometryCollection& OutGeometryCollection)
+                                                FGeometryCollection& OutGeometryCollection)
 {
 	InSkeletonCollection.CopyTo(&OutGeometryCollection);
 
@@ -89,7 +89,7 @@ void FPVMeshBuilder::GenerateGeometryCollection(const FManagedArrayCollection& I
 	const FName MaterialPathAttributeName("MaterialPath");
 	auto& MaterialArray = OutGeometryCollection.AddAttribute<FString>(MaterialPathAttributeName, FGeometryCollection::MaterialGroup);
 	auto& MaterialSetups = MeshBuilderParams.MaterialSettings.MaterialSetups;
-	
+
 	struct FSection
 	{
 		int32 MaterialID;
@@ -106,12 +106,12 @@ void FPVMeshBuilder::GenerateGeometryCollection(const FManagedArrayCollection& I
 
 	//Merge MaterialIds based on material Paths
 	MergeMaterials(MaterialSetups, MeshesDataCollection);
-	
+
 	//Sorting by material id so sections are created according to material ids
 	MeshesDataCollection.Sort([](const FLocalDynamicMeshData& A, const FLocalDynamicMeshData& B)
-	{
-		return A.MaterialID > B.MaterialID;
-	});
+		{
+			return A.MaterialID > B.MaterialID;
+		});
 
 	int32 SectionId = -1;
 	for (auto MeshIndex = MeshesDataCollection.Num() - 1; MeshIndex >= 0; --MeshIndex)
@@ -120,7 +120,7 @@ void FPVMeshBuilder::GenerateGeometryCollection(const FManagedArrayCollection& I
 		if (PrevMaterialID != MaterialID)
 		{
 			PrevMaterialID = MaterialID;
-			
+
 			TObjectPtr<UMaterialInterface> SectionMaterial;
 
 			if (MaterialSetups.IsValidIndex(MaterialID))
@@ -129,7 +129,8 @@ void FPVMeshBuilder::GenerateGeometryCollection(const FManagedArrayCollection& I
 			}
 			else
 			{
-				UE_LOG(LogProceduralVegetation, Warning, TEXT("Unassigned Material for section %i, MaterialID %i not found in MaterialSetups."), SectionId, MaterialID);
+				UE_LOG(LogProceduralVegetation, Warning, TEXT("Unassigned Material for section %i, MaterialID %i not found in MaterialSetups."),
+					SectionId, MaterialID);
 			}
 
 			FString MaterialPathName = "";
@@ -139,7 +140,8 @@ void FPVMeshBuilder::GenerateGeometryCollection(const FManagedArrayCollection& I
 			}
 			else
 			{
-				UE_LOG(LogProceduralVegetation, Warning, TEXT("Null Material for section %i, Assign Material slot no %i in MaterialSetups."), SectionId, MaterialID);
+				UE_LOG(LogProceduralVegetation, Warning, TEXT("Null Material for section %i, Assign Material slot no %i in MaterialSetups."),
+					SectionId, MaterialID);
 			}
 
 			int Element = OutGeometryCollection.AddElements(1, FGeometryCollection::MaterialGroup);
@@ -188,7 +190,7 @@ void FPVMeshBuilder::GenerateGeometryCollection(const FManagedArrayCollection& I
 	OutGeometryCollection.TangentU.Fill(FVector3f::ForwardVector);
 	OutGeometryCollection.TangentV.Fill(FVector3f::RightVector);
 	OutGeometryCollection.Color.Fill(FLinearColor::White);
-	
+
 	int32 SectionIndex = 0;
 	for (auto Section : Sections)
 	{
@@ -232,7 +234,7 @@ void FPVMeshBuilder::SetNjordPixelID(FGeometryCollection& OutGeometryCollection)
 }
 
 void FPVMeshBuilder::GenerateDynamicMesh(FManagedArrayCollection& Collection, const FPVMeshBuilderParams& MeshBuilderParams,
-                                    TObjectPtr<UDynamicMesh>& OutMesh)
+                                         TObjectPtr<UDynamicMesh>& OutMesh)
 {
 	TArray<FLocalDynamicMeshData> MeshesDataCollection;
 	GetLocalDynamicMeshData(Collection, MeshBuilderParams, MeshesDataCollection);
@@ -256,7 +258,7 @@ void FPVMeshBuilder::GenerateDynamicMesh(FManagedArrayCollection& Collection, co
 }
 
 TSet<int32> FPVMeshBuilder::CollectHardPoints(const PV::Facades::FBranchFacade& BranchFacade, const PV::Facades::FPointFacade& PointFacade,
-                                         const PV::Facades::FPlantFacade& PlantFacade)
+                                              const PV::Facades::FPlantFacade& PlantFacade)
 {
 	TSet<int32> HardPoints;
 	TSet<int32> BranchSourceBudNumbers;
@@ -331,9 +333,9 @@ TSet<int32> FPVMeshBuilder::CollectHardPoints(const PV::Facades::FBranchFacade&
 }
 
 TSet<int32> FPVMeshBuilder::ComputePointGradients(const PV::Facades::FPointFacade& PointFacade,
-                                             const FPVMeshBuilderParams& MeshBuilderParams, const TSet<int32>& HardPoints,
-                                             const float MaxPointScale, TArray<float>& OutMeshDivisionsGradients,
-                                             TArray<float>& OutDeltaModifiers)
+                                                  const FPVMeshBuilderParams& MeshBuilderParams, const TSet<int32>& HardPoints,
+                                                  const float MaxPointScale, TArray<float>& OutMeshDivisionsGradients,
+                                                  TArray<float>& OutDeltaModifiers)
 {
 	TSet<int32> PointsToRemove;
 	OutMeshDivisionsGradients.Reserve(PointFacade.GetElementCount());
@@ -404,7 +406,7 @@ TSet<int32> FPVMeshBuilder::ComputePointGradients(const PV::Facades::FPointFacad
 		float DeltaModifier = 1.0f;
 		if (!HardPoints.Contains(PointIndex))
 		{
-			DeltaModifier = MeshBuilderParams.SegmentRetentionImpact * (MaxPointScale / 100.0f) * DeltaModifierGradient;
+			DeltaModifier = MeshBuilderParams.SegmentRetentionImpact * MaxPointScale * DeltaModifierGradient;
 		}
 		OutDeltaModifiers.Add(DeltaModifier);
 	}
@@ -413,8 +415,9 @@ TSet<int32> FPVMeshBuilder::ComputePointGradients(const PV::Facades::FPointFacad
 }
 
 void FPVMeshBuilder::PerformPathSimplification(const PV::Facades::FBranchFacade& BranchFacade,
-                                          const PV::Facades::FPointFacade& PointFacade, const FPVMeshBuilderParams& MeshBuilderParams,
-                                          const TSet<int32>& HardPoints, const TArray<float>& DeltaModifiers, TSet<int32>& InOutPointsToRemove)
+                                               const PV::Facades::FPointFacade& PointFacade, const FPVMeshBuilderParams& MeshBuilderParams,
+                                               const float MaxPointScale, const TSet<int32>& HardPoints, const TArray<float>& DeltaModifiers,
+                                               TSet<int32>& InOutPointsToRemove)
 {
 	const int32 Iterations = MeshBuilderParams.Accuracy;
 	// Epsilon is the maximum allowed curvature error
@@ -430,7 +433,7 @@ void FPVMeshBuilder::PerformPathSimplification(const PV::Facades::FBranchFacade&
 
 		for (int32 j = 0; j < Iterations; ++j)
 		{
-			const float IterationEpsilon = Epsilon * (j + 1);
+			const float IterationEpsilon = Epsilon * (j + 1) * MaxPointScale;
 			for (int32 i = 1; i < BranchPoints.Num() - 1; ++i)
 			{
 				const int32 Pt0Index = BranchPoints[i - 1],
@@ -440,7 +443,8 @@ void FPVMeshBuilder::PerformPathSimplification(const PV::Facades::FBranchFacade&
 				              Pt1Position = FVector(PointFacade.GetPosition(Pt1Index)),
 				              Pt2Position = FVector(PointFacade.GetPosition(Pt2Index));
 				const float DeltaModifier = DeltaModifiers[Pt1Index];
-				const float Delta = FMath::PointDistToLine(Pt1Position, Pt2Position - Pt0Position, Pt0Position) + DeltaModifier;
+				const float Delta =
+					FMath::PointDistToLine(Pt1Position, Pt2Position - Pt0Position, Pt0Position) + DeltaModifier;
 
 				if (Delta < IterationEpsilon && !HardPoints.Contains(Pt1Index))
 				{
@@ -512,8 +516,8 @@ void FPVMeshBuilder::PerformPathSimplification(const PV::Facades::FBranchFacade&
 }
 
 TArray<int32> FPVMeshBuilder::ComputeMeshDivisions(const PV::Facades::FBranchFacade& BranchFacade,
-                                              const FPVMeshBuilderParams& MeshBuilderParams, const TArray<float>& MeshDivisionsGradients,
-                                              const int32 PointCount)
+                                                   const FPVMeshBuilderParams& MeshBuilderParams, const TArray<float>& MeshDivisionsGradients,
+                                                   const int32 PointCount)
 {
 	TArray<int32> TargetMeshDivisions;
 	static constexpr int32 MinimumDivisions = 3;
@@ -546,7 +550,7 @@ TArray<int32> FPVMeshBuilder::ComputeMeshDivisions(const PV::Facades::FBranchFac
 }
 
 void FPVMeshBuilder::TriangulateRings(const TArray<int32>& PreviousIndices, const TArray<int32>& CurrentIndices, int32& InOutPolyGroupIndex,
-                                 FLocalDynamicMeshData& OutMeshData)
+                                      FLocalDynamicMeshData& OutMeshData)
 {
 	if (!PreviousIndices.IsEmpty())
 	[[likely]]
@@ -648,9 +652,9 @@ TMap<int32, TArray<int32>> FPVMeshBuilder::GetPointsIndicesToFoliageIndicesMap(c
 }
 
 void FPVMeshBuilder::GenerateBranchMeshData(const bool bPrimitiveIsTrunk, const int32 GenerationNumber, const TArray<int32>& PrimitivePoints,
-                                       const FDisplacementData& DisplacementData, const FPVMeshBuilderParams& MeshBuilderParams,
-                                       const TArray<int32>& TargetMeshDivisions, const FManagedArrayCollection& Collection,
-                                       FLocalDynamicMeshData& OutLocalMeshData)
+                                            const FDisplacementData& DisplacementData, const FPVMeshBuilderParams& MeshBuilderParams,
+                                            const TArray<int32>& TargetMeshDivisions, const FManagedArrayCollection& Collection,
+                                            FLocalDynamicMeshData& OutLocalMeshData)
 {
 	if (PrimitivePoints.Num() < 2)
 	{
@@ -904,7 +908,7 @@ void FPVMeshBuilder::UpdateFoliagePivotPoints(const TSet<int32>& PointsToRemove,
 }
 
 void FPVMeshBuilder::GetLocalDynamicMeshData(FManagedArrayCollection& Collection, const FPVMeshBuilderParams& MeshBuilderParams,
-                                        TArray<FLocalDynamicMeshData>& MeshesData)
+                                             TArray<FLocalDynamicMeshData>& MeshesData)
 {
 	const PV::Facades::FBranchFacade BranchFacade(Collection);
 	const PV::Facades::FPointFacade PointFacade(Collection);
@@ -926,7 +930,7 @@ void FPVMeshBuilder::GetLocalDynamicMeshData(FManagedArrayCollection& Collection
 	TSet<int32> PointsToRemove = ComputePointGradients(PointFacade, MeshBuilderParams, HardPoints, MaxPointScale, MeshDivisionsGradientsForPoints,
 		DeltaModifiersForPoints);
 
-	PerformPathSimplification(BranchFacade, PointFacade, MeshBuilderParams, HardPoints, DeltaModifiersForPoints, PointsToRemove);
+	PerformPathSimplification(BranchFacade, PointFacade, MeshBuilderParams, MaxPointScale, HardPoints, DeltaModifiersForPoints, PointsToRemove);
 
 	// Point removal may update curvate of the branch.
 	// Foliage instance pivot/attachment points need to be updated to compensate for that.
diff --git a/Engine/Plugins/Experimental/ProceduralVegetationEditor/Source/ProceduralVegetation/Private/Implementations/PVMeshBuilder.h b/Engine/Plugins/Experimental/ProceduralVegetationEditor/Source/ProceduralVegetation/Private/Implementations/PVMeshBuilder.h
--- a/Engine/Plugins/Experimental/ProceduralVegetationEditor/Source/ProceduralVegetation/Private/Implementations/PVMeshBuilder.h
+++ b/Engine/Plugins/Experimental/ProceduralVegetationEditor/Source/ProceduralVegetation/Private/Implementations/PVMeshBuilder.h
@@ -30,70 +30,109 @@ struct FPVMeshBuilderParams
 {
 	GENERATED_BODY()
 
-	UPROPERTY(EditAnywhere, Category="Materials", meta=(PinHiddenByDefault, InlineEditConditionToggle, Tooltip="Toggles between default and custom material settings.\n\nDefault material setting is defined in Procedural Vegetation Preset.\n\nCustom material settings allow for editing."))
+	UPROPERTY(EditAnywhere, Category="Materials",
+		meta=(PinHiddenByDefault, InlineEditConditionToggle, Tooltip=
+			"Toggles between default and custom material settings.\n\nDefault material setting is defined in Procedural Vegetation Preset.\n\nCustom material settings allow for editing."
+		))
 	bool bOverrideMaterial = false;
 
 	UPROPERTY(EditAnywhere, Category="Materials", meta = (EditCondition="bOverrideMaterial"))
 	FPVMaterialSettings MaterialSettings;
 
 	UPROPERTY(EditAnywhere, Category="Mesh Settings", DisplayName="Point Removal",
-		meta=(ClampMin=0.0f, ClampMax=0.1f, UIMin=0.0f, UIMax=0.1f, Tooltip="Removes small points to simplify the mesh.\n\nDefines the strength of point removal. Higher values eliminate more minor vertices or detail points to reduce complexity at the cost of fine detail."))
+		meta=(ClampMin=0.0f, ClampMax=0.1f, UIMin=0.0f, UIMax=0.1f, Tooltip=
+			"Removes small points to simplify the mesh.\n\nDefines the strength of point removal. Higher values eliminate more minor vertices or detail points to reduce complexity at the cost of fine detail."
+		))
 	float PointRemoval = 0.0f;
 
 	UPROPERTY(EditAnywhere, Category="Mesh Settings", DisplayName="Segment Reduction",
-		meta=(ClampMin=0.0f, ClampMax=0.1f, UIMin=0.0f, UIMax=0.1f, Tooltip="Controls how aggressively segments are reduced.\n\nSets the strength of segment merging/decimation. Higher values reduce edge count for lighter meshes; lower values preserve detail."))
+		meta=(ClampMin=0.0f, ClampMax=1.0f, UIMin=0.0f, UIMax=0.1f, Tooltip=
+			"Controls how aggressively segments are reduced.\n\nSets the strength of segment merging/decimation. Higher values reduce edge count for lighter meshes; lower values preserve detail."
+		))
 	float SegmentReduction = 0.0f;
 
-	UPROPERTY(EditAnywhere, Category="Mesh Settings", DisplayName="Segment Retention Impact",
-		meta=(ClampMin=0.0f, ClampMax=2.0f, UIMin=0.0f, UIMax=2.0f, Tooltip="Control the impact of reduced segments based on other criteria\n\nSegment retention impact allows you to reintroduce points that have been reduced based on user defined criteria\'s. See; Hull Retention, Ground Retention and Trunk Retention."))
-	float SegmentRetentionImpact = 1.0f;
-
-	UPROPERTY(EditAnywhere, Category="Mesh Settings", DisplayName="Longest Segment Allowed",
-		meta=(ClampMin=0.0f, ClampMax=50.0f, UIMin=0.0f, UIMax=10.0f, Tooltip="Maximum allowed segment length.\n\nCaps segment length after simplification to maintain even geometry and prevent stretched polygons. Specified in mesh units."))
-	float LongestSegmentLength = 0.1f;
-
-	UPROPERTY(EditAnywhere, Category="Mesh Settings", DisplayName="Shortest Segment Allowed",
-		meta=(ClampMin=0.0f, ClampMax=1.0f, UIMin=0.0f, UIMax=1.0f, Tooltip="Minimum allowed segment length.\n\nSegments shorter than this are merged or removed to avoid unnecessary fine detail that inflates polycount."))
-	float ShortestSegmentLength = 0.0f;
-
 	UPROPERTY(EditAnywhere, Category="Mesh Settings", DisplayName="Min Divisions",
-		meta=(ClampMin=3, ClampMax=1024, UIMin=3, UIMax=36, Tooltip="Minimum radial divisions.\n\nThe lowest polygon count around cross-sections. Higher values keep sections smooth; lower values produce blockier geometry with better performance."))
+		meta=(ClampMin=3, ClampMax=1024, UIMin=3, UIMax=36, Tooltip=
+			"Minimum radial divisions.\n\nThe lowest polygon count around cross-sections. Higher values keep sections smooth; lower values produce blockier geometry with better performance."
+		))
 	int32 MinMeshDivisions = 6;
 
 	UPROPERTY(EditAnywhere, Category="Mesh Settings", DisplayName="Max Divisions",
-		meta=(ClampMin=3, ClampMax=1024, UIMin=3, UIMax=36, Tooltip="Maximum radial divisions.\n\nUpper limit for cross-section detail on trunks/branches. Balance with Min Divisions to control density across sizes."))
+		meta=(ClampMin=3, ClampMax=1024, UIMin=3, UIMax=36, Tooltip=
+			"Maximum radial divisions.\n\nUpper limit for cross-section detail on trunks/branches. Balance with Min Divisions to control density across sizes."
+		))
 	int32 MaxMeshDivisions = 12;
 
 	UPROPERTY(EditAnywhere, Category="Mesh Settings|Advanced", DisplayName="Reduction Accuracy",
-		meta=(ClampMin=3, ClampMax=15, UIMin=1, UIMax=15, Tooltip="Accuracy level for mesh reduction.\n\nControls how closely the simplified mesh follows the original. It specifies number of iterations that are used to evaluate the curve reduction path. The more number of iterations the more accurate the reduction will be at the cost of processing power."))
+		meta=(ClampMax=15, UIMin=1, UIMax=15, Tooltip=
+			"Accuracy level for mesh reduction.\n\nControls how closely the simplified mesh follows the original. It specifies number of iterations that are used to evaluate the curve reduction path. The more number of iterations the more accurate the reduction will be at the cost of processing power."
+		))
 	int32 Accuracy = 5;
 
+	UPROPERTY(EditAnywhere, Category="Mesh Settings|Advanced", DisplayName="Longest Segment Allowed (meters)",
+		meta=(ClampMin=0.0f, ClampMax=50.0f, UIMin=0.0f, UIMax=10.0f, Tooltip=
+			"Maximum allowed segment length.\n\nCaps segment length after simplification to maintain even geometry and prevent stretched polygons. Specified in mesh units."
+		))
+	float LongestSegmentLength = 10.0f;
+
+	UPROPERTY(EditAnywhere, Category="Mesh Settings|Advanced", DisplayName="Shortest Segment Allowed (meters)",
+		meta=(ClampMin=0.0f, ClampMax=50.0f, UIMin=0.0f, UIMax=1.0f, Tooltip=
+			"Minimum allowed segment length.\n\nSegments shorter than this are merged or removed to avoid unnecessary fine detail that inflates polycount."
+		))
+	float ShortestSegmentLength = 0.0f;
+
+	UPROPERTY(EditAnywhere, Category="Mesh Settings|Advanced", DisplayName="Segment Retention Impact",
+		meta=(ClampMin=0.0f, ClampMax=5.0f, UIMin=0.0f, UIMax=2.0f, Tooltip=
+			"Control the impact of reduced segments based on other criteria\n\nSegment retention impact allows you to reintroduce points that have been reduced based on user defined criteria\'s. See; Hull Retention, Ground Retention and Trunk Retention."
+		))
+	float SegmentRetentionImpact = 1.0f;
+
 	UPROPERTY(EditAnywhere, Category="Mesh Settings|Advanced", DisplayName="Hull Retention",
-		meta=(ClampMin=0.0f, ClampMax=1.0f, UIMin=0.0f, UIMax=1.0f, Tooltip="Strength of preserving the outer hull.\n\nAdjusts retention of the mesh silhouette during reduction. Higher values protect external shape; lower values allow more simplification."))
+		meta=(ClampMin=0.0f, ClampMax=1.0f, UIMin=0.0f, UIMax=1.0f, Tooltip=
+			"Strength of preserving the outer hull.\n\nAdjusts retention of the mesh silhouette during reduction. Higher values protect external shape; lower values allow more simplification."
+		))
 	float HullRetention = 0.0f;
 
 	UPROPERTY(EditAnywhere, Category="Mesh Settings|Advanced", DisplayName="Main Trunk Retention",
-		meta=(ClampMin=0.0f, ClampMax=1.0f, UIMin=0.0f, UIMax=1.0f, Tooltip="Preserves detail in the main trunk.\n\nControls how much trunk geometry is retained during simplification. High values keep trunk smooth and detailed; low values reduce complexity."))
+		meta=(ClampMin=0.0f, ClampMax=1.0f, UIMin=0.0f, UIMax=1.0f, Tooltip=
+			"Preserves detail in the main trunk.\n\nControls how much trunk geometry is retained during simplification. High values keep trunk smooth and detailed; low values reduce complexity."
+		))
 	float MainTrunkRetention = 0.0f;
 
 	UPROPERTY(EditAnywhere, Category="Mesh Settings|Advanced", DisplayName="Ground Retention",
-		meta=(ClampMin=0.0f, ClampMax=1.0f, UIMin=0.0f, UIMax=1.0f, Tooltip="Preserves detail near ground regions.\n\nAdjusts retention for areas close to the base of the plant. Use higher values to keep root/ground transitions detailed."))
+		meta=(UIMin=0.0f, UIMax=1.0f, Tooltip=
+			"Preserves detail near ground regions.\n\nAdjusts retention for areas close to the base of the plant. Use higher values to keep root/ground transitions detailed."
+		))
 	float GroundRetention = 0.0f;
 
-	UPROPERTY(EditAnywhere, Category="Mesh Settings|Advanced", DisplayName="Hull Retention Gradient", meta=(XAxisMin=0.0f, XAxisMax=1.0f, YAxisMin=0.0f, YAxisMax=1.0f, Tooltip="Gradient for preserving outer hull detail.\n\nA ramp that varies outer-hull retention across the model. Use to keep silhouettes crisp where needed while simplifying elsewhere."))
+	UPROPERTY(EditAnywhere, Category="Mesh Settings|Advanced", DisplayName="Hull Retention Gradient",
+		meta=(XAxisMin=0.0f, XAxisMax=1.0f, YAxisMin=0.0f, YAxisMax=1.0f, Tooltip=
+			"Gradient for preserving outer hull detail.\n\nA ramp that varies outer-hull retention across the model. Use to keep silhouettes crisp where needed while simplifying elsewhere."
+		))
 	FPVFloatRamp HullRetentionGradient;
 
-	UPROPERTY(EditAnywhere, Category="Mesh Settings|Advanced", DisplayName="Main Trunk Retention Gradient", meta=(XAxisMin=0.0f, XAxisMax=1.0f, YAxisMin=0.0f, YAxisMax=1.0f, Tooltip="Gradient controlling trunk detail retention.\n\nVaries trunk retention via a ramp. Useful to preserve key trunk regions while simplifying others."))
+	UPROPERTY(EditAnywhere, Category="Mesh Settings|Advanced", DisplayName="Main Trunk Retention Gradient",
+		meta=(XAxisMin=0.0f, XAxisMax=1.0f, YAxisMin=0.0f, YAxisMax=1.0f, Tooltip=
+			"Gradient controlling trunk detail retention.\n\nVaries trunk retention via a ramp. Useful to preserve key trunk regions while simplifying others."
+		))
 	FPVFloatRamp MainTrunkRetentionGradient;
 
-	UPROPERTY(EditAnywhere, Category="Mesh Settings|Advanced", DisplayName="Ground Retention Gradient", meta=(XAxisMin=0.0f, XAxisMax=1.0f, YAxisMin=0.0f, YAxisMax=1.0f, Tooltip="Gradient for preserving ground-region detail.\n\nA ramp that modulates retention around base regions. Tune to keep roots/basal features while allowing simplification above."))
+	UPROPERTY(EditAnywhere, Category="Mesh Settings|Advanced", DisplayName="Ground Retention Gradient",
+		meta=(XAxisMin=0.0f, XAxisMax=1.0f, YAxisMin=0.0f, YAxisMax=1.0f, Tooltip=
+			"Gradient for preserving ground-region detail.\n\nA ramp that modulates retention around base regions. Tune to keep roots/basal features while allowing simplification above."
+		))
 	FPVFloatRamp GroundRetentionGradient;
 
 	UPROPERTY(EditAnywhere, Category="Mesh Settings|Advanced", DisplayName="Scale Retention",
-		meta=(ClampMin=0.0f, ClampMax=1.0f, UIMin=0.0f, UIMax=1.0f, Tooltip="Preserves detail relative to branch size.\n\nRetains more detail on larger/thicker structures and simplifies smaller ones to prioritize visually important geometry."))
+		meta=(ClampMin=0.0f, ClampMax=1.0f, UIMin=0.0f, UIMax=1.0f, Tooltip=
+			"Preserves detail relative to branch size.\n\nRetains more detail on larger/thicker structures and simplifies smaller ones to prioritize visually important geometry."
+		))
 	float ScaleRetention = 0.0f;
 
-	UPROPERTY(EditAnywhere, Category="Mesh Settings|Advanced", DisplayName="Scale Retention Gradient", meta=(XAxisMin=0.0f, XAxisMax=1.0f, YAxisMin=0.0f, YAxisMax=1.0f, Tooltip="Gradient controlling size-based retention.\n\nUses a ramp to vary how size-dependent retention is applied across the model. Keep large structures crisp while simplifying finer parts."))
+	UPROPERTY(EditAnywhere, Category="Mesh Settings|Advanced", DisplayName="Scale Retention Gradient",
+		meta=(XAxisMin=0.0f, XAxisMax=1.0f, YAxisMin=0.0f, YAxisMax=1.0f, Tooltip=
+			"Gradient controlling size-based retention.\n\nUses a ramp to vary how size-dependent retention is applied across the model. Keep large structures crisp while simplifying finer parts."
+		))
 	FPVFloatRamp ScaleRetentionGradient;
 
 	UPROPERTY(EditAnywhere, Category="Plant Profile Settings", DisplayName="Profile",
@@ -107,7 +146,8 @@ struct FPVMeshBuilderParams
 	bool bIsPlantProfileDropdownEnabled = false;
 
 	UPROPERTY(EditAnywhere, Category="Plant Profile Settings", DisplayName="Profile Falloff",
-		meta=(XAxisMin=0.0f, XAxisMax=1.0f, YAxisMin=0.0f, YAxisMax=1.0f, EditCondition="bIsPlantProfileDropdownEnabled", EditConditionHides, HideEditConditionToggle))
+		meta=(XAxisMin=0.0f, XAxisMax=1.0f, YAxisMin=0.0f, YAxisMax=1.0f, EditCondition="bIsPlantProfileDropdownEnabled", EditConditionHides,
+			HideEditConditionToggle))
 	FPVFloatRamp PlantProfileFallOff;
 
 	UPROPERTY(EditAnywhere, Category="Plant Profile Settings", DisplayName="Profile Scale",
@@ -119,25 +159,36 @@ struct FPVMeshBuilderParams
 		meta=(EditCondition="bIsPlantProfileDropdownEnabled", EditConditionHides, HideEditConditionToggle))
 	bool bApplyProfileToBranches = false;
 
-	UPROPERTY(EditAnywhere, Category = "Displacement Settings", meta=(Tooltip="Texture used to drive mesh displacement.\n\nAssigns the texture map that controls displacement. Bright and dark areas in the texture push or pull the mesh surface to create additional geometric detail such as bark roughness, grooves, or ridges."))
+	UPROPERTY(EditAnywhere, Category = "Displacement Settings",
+		meta=(Tooltip=
+			"Texture used to drive mesh displacement.\n\nAssigns the texture map that controls displacement. Bright and dark areas in the texture push or pull the mesh surface to create additional geometric detail such as bark roughness, grooves, or ridges."
+		))
 	TObjectPtr<UTexture2D> DisplacementTexture = nullptr;
 
 	UPROPERTY(Transient, NonTransactional)
 	TArray<float> DisplacementValues;
 
 	UPROPERTY(EditAnywhere, Category="Displacement Settings", DisplayName="Displacement Scale",
-		meta=(ClampMin=0.0f, ClampMax=10.0f, UIMin=0.0f, UIMax=1.0f, Tooltip="Intensity of displacement effect.\n\nSets the overall strength of the displacement. Higher values exaggerate surface depth and detail, while lower values apply subtle variation. Use to balance realism with mesh stability."))
+		meta=(ClampMin=0.0f, ClampMax=10.0f, UIMin=0.0f, UIMax=1.0f, Tooltip=
+			"Intensity of displacement effect.\n\nSets the overall strength of the displacement. Higher values exaggerate surface depth and detail, while lower values apply subtle variation. Use to balance realism with mesh stability."
+		))
 	float DisplacementStrength = 0.5f;
 
 	UPROPERTY(EditAnywhere, Category="Displacement Settings", DisplayName="Displacement Bias",
-		meta=(ClampMin=0.0f, ClampMax=10.0f, UIMin=0.0f, UIMax=10.0f, Tooltip="Shifts the midpoint of displacement.\n\nOffsets the neutral level of the displacement map. Positive bias lifts the surface outward. Useful for correcting textures that displace unevenly or need centering."))
+		meta=(ClampMin=0.0f, ClampMax=10.0f, UIMin=0.0f, UIMax=10.0f, Tooltip=
+			"Shifts the midpoint of displacement.\n\nOffsets the neutral level of the displacement map. Positive bias lifts the surface outward. Useful for correcting textures that displace unevenly or need centering."
+		))
 	float DisplacementBias = 0.0f;
 
-	UPROPERTY(EditAnywhere, Category="Displacement Settings", DisplayName="Displacement UV Scale", meta=(Tooltip="Controls tiling of displacement texture.\n\nAdjusts the repetition of the displacement texture across the mesh in U and V directions. Higher values increase tiling frequency (smaller details repeated more often), while lower values stretch the texture over a larger area."))
+	UPROPERTY(EditAnywhere, Category="Displacement Settings", DisplayName="Displacement UV Scale",
+		meta=(Tooltip=
+			"Controls tiling of displacement texture.\n\nAdjusts the repetition of the displacement texture across the mesh in U and V directions. Higher values increase tiling frequency (smaller details repeated more often), while lower values stretch the texture over a larger area."
+		))
 	FIntVector2 DisplacementUVScale = FIntVector2(1.0f, 1.0f);
 
 	UPROPERTY(EditAnywhere, Category="Displacement Settings", DisplayName="Apply Displacement up to generation",
-		meta=(ClampMin=1, ClampMax=10, UIMin=1, UIMax=10, Tooltip="Controls which generations the displacement is applied on \n\nControls which generations the displacement is applied on."))
+		meta=(ClampMin=1, ClampMax=10, UIMin=1, UIMax=10, Tooltip=
+			"Controls which generations the displacement is applied on \n\nControls which generations the displacement is applied on."))
 	int32 DisplacementGenerationUpperLimit = 1;
 
 	FPVMeshBuilderParams()
@@ -264,7 +315,8 @@ private:
 
 	static void PerformPathSimplification(const PV::Facades::FBranchFacade& BranchFacade,
 	                                      const PV::Facades::FPointFacade& PointFacade, const FPVMeshBuilderParams& MeshBuilderParams,
-	                                      const TSet<int32>& HardPoints, const TArray<float>& DeltaModifiers, TSet<int32>& InOutPointsToRemove);
+	                                      const float MaxPointScale, const TSet<int32>& HardPoints, const TArray<float>& DeltaModifiers,
+	                                      TSet<int32>& InOutPointsToRemove);
 
 	static TArray<int32> ComputeMeshDivisions(const PV::Facades::FBranchFacade& BranchFacade, const FPVMeshBuilderParams& MeshBuilderParams,
 	                                          const TArray<float>& MeshDivisionsGradients, const int32 PointCount);
