[PVE] Crash on deleting PV Asset if PVE is open

diff --git a/Engine/Plugins/Experimental/ProceduralVegetationEditor/Source/ProceduralVegetationEditor/Private/PVEditor.cpp b/Engine/Plugins/Experimental/ProceduralVegetationEditor/Source/ProceduralVegetationEditor/Private/PVEditor.cpp
--- a/Engine/Plugins/Experimental/ProceduralVegetationEditor/Source/ProceduralVegetationEditor/Private/PVEditor.cpp
+++ b/Engine/Plugins/Experimental/ProceduralVegetationEditor/Source/ProceduralVegetationEditor/Private/PVEditor.cpp
@@ -164,8 +164,7 @@ void FPVEditor::OnClose()
 
 	if (ExecutionSource)
 	{
-		ExecutionSource->Sunset();
-		ExecutionSource = nullptr;
+		ExecutionSource->SetGraphInterface(nullptr);
 	}
 }
 
