Improve error handling of PVE export

- Add custom PVE MessageLog
- Clear previous export message before logging new messages (to prevent a error messages from previous export to trigger an error allert)
- Return result enum from ExportCollectionAsMesh so appropriate message can be logged
- Log location of exported assets
- Fix analytics event not firing properly if export fails

diff --git a/Engine/Plugins/Experimental/ProceduralVegetationEditor/Source/ProceduralVegetation/Private/Helpers/PVAnalyticsHelper.cpp b/Engine/Plugins/Experimental/ProceduralVegetationEditor/Source/ProceduralVegetation/Private/Helpers/PVAnalyticsHelper.cpp
--- a/Engine/Plugins/Experimental/ProceduralVegetationEditor/Source/ProceduralVegetation/Private/Helpers/PVAnalyticsHelper.cpp
+++ b/Engine/Plugins/Experimental/ProceduralVegetationEditor/Source/ProceduralVegetation/Private/Helpers/PVAnalyticsHelper.cpp
@@ -141,6 +141,7 @@ namespace PV::Analytics
 		TArray<FAnalyticsEventAttribute> ExportAttributes;
 		ExportAttributes.Add(FAnalyticsEventAttribute(TEXT("ExportSettings.MeshType"),  StaticEnum<EPVExportMeshType>()->GetNameStringByValue((uint8)InExportParam.ExportMeshType)));
 		ExportAttributes.Add(FAnalyticsEventAttribute(TEXT("ExportSettings.CreateNaniteFoliage"),  InExportParam.bCreateNaniteFoliage ? TEXT("True") : TEXT("False")));
+		ExportAttributes.Add(FAnalyticsEventAttribute(TEXT("ExportSettings.ReplacePolicy"), StaticEnum<EPVAssetReplacePolicy>()->GetNameStringByValue((uint8)InExportParam.ReplacePolicy)));
 		ExportAttributes.Add(FAnalyticsEventAttribute(TEXT("ExportSettings.NaniteShapePreservation"),  StaticEnum<ENaniteShapePreservation>()->GetNameStringByValue((uint8)InExportParam.NaniteShapePreservation)));
 		ExportAttributes.Add(FAnalyticsEventAttribute(TEXT("NaniteSettings.NaniteFoliageEnabled"),  GetDefault<URendererSettings>()->bEnableNaniteFoliage));
 		ExportAttributes.Add(FAnalyticsEventAttribute(TEXT("NaniteSettings.Enabled"),  GetDefault<URendererSettings>()->bNanite));
@@ -170,12 +171,28 @@ namespace PV::Analytics
 		}
 	}
 
-	void SendExportFinishedEvent(TArray<FAnalyticsEventAttribute>& InCommonAttributes, float InTotalTimeInSeconds)
+	void SendExportFinishedEvent(TArray<FAnalyticsEventAttribute> InCommonAttributes, float InTotalTimeInSeconds, EExportResult ExportResult)
 	{
 		if (FEngineAnalytics::IsAvailable())
 		{
-			InCommonAttributes.Add(FAnalyticsEventAttribute(TEXT("TotalTimeSeconds"),  InTotalTimeInSeconds));
+			InCommonAttributes.Add(FAnalyticsEventAttribute(TEXT("TotalTimeSeconds"), InTotalTimeInSeconds));
+
+			const static auto ExportResultEnumToString = [](EExportResult InExportResult)->FString {
+				switch (InExportResult)
+				{
+				case EExportResult::Failed:
+					return TEXT("Failed");
+				case EExportResult::Success:
+					return TEXT("Success");
+				case EExportResult::Skipped:
+					return TEXT("Skipped");
+				}
+
+				return TEXT("Unkown");
+			};
 			
+			InCommonAttributes.Add(FAnalyticsEventAttribute(TEXT("ExportResult"), ExportResultEnumToString(ExportResult)));
+
 			SendEvent(TEXT("ExportFinished"), InCommonAttributes);
 		}
 	}
diff --git a/Engine/Plugins/Experimental/ProceduralVegetationEditor/Source/ProceduralVegetation/Public/Helpers/PVAnalyticsHelper.h b/Engine/Plugins/Experimental/ProceduralVegetationEditor/Source/ProceduralVegetation/Public/Helpers/PVAnalyticsHelper.h
--- a/Engine/Plugins/Experimental/ProceduralVegetationEditor/Source/ProceduralVegetation/Public/Helpers/PVAnalyticsHelper.h
+++ b/Engine/Plugins/Experimental/ProceduralVegetationEditor/Source/ProceduralVegetation/Public/Helpers/PVAnalyticsHelper.h
@@ -19,5 +19,12 @@ namespace PV::Analytics
 	void SendWindSettingsChangeEvent(FString InPath);
 	TArray<FAnalyticsEventAttribute> PROCEDURALVEGETATION_API GatherExportCommonAttributes(const FManagedArrayCollection& InCollection, const FPVExportParams& InExportParam);
 	void PROCEDURALVEGETATION_API SendExportStartEvent(const TArray<FAnalyticsEventAttribute>& InCommonAttributes);
-	void PROCEDURALVEGETATION_API SendExportFinishedEvent(TArray<FAnalyticsEventAttribute>& InCommonAttributes, float InTotalTimeInSeconds);
+
+	enum class EExportResult
+	{
+		Failed,
+		Success,
+		Skipped
+	};
+	void PROCEDURALVEGETATION_API SendExportFinishedEvent(TArray<FAnalyticsEventAttribute> InCommonAttributes, float InTotalTimeInSeconds, EExportResult ExportResult);
 }
diff --git a/Engine/Plugins/Experimental/ProceduralVegetationEditor/Source/ProceduralVegetationEditor/Private/Helpers/PVExportHelper.cpp b/Engine/Plugins/Experimental/ProceduralVegetationEditor/Source/ProceduralVegetationEditor/Private/Helpers/PVExportHelper.cpp
--- a/Engine/Plugins/Experimental/ProceduralVegetationEditor/Source/ProceduralVegetationEditor/Private/Helpers/PVExportHelper.cpp
+++ b/Engine/Plugins/Experimental/ProceduralVegetationEditor/Source/ProceduralVegetationEditor/Private/Helpers/PVExportHelper.cpp
@@ -200,11 +200,19 @@ namespace PV::Export::Internal
 		}
 	}
 
-	bool GetValidAssetPath(EPVAssetReplacePolicy AssetReplacePolicy, UClass* AssetClass, FString& InOutPackageName, FString& InOutAssetName)
+	enum class EGetValidAssetPathResult
+	{
+		Success,
+		InvalidPath,
+		ConflictingPackageExist,
+		SkipAsset
+	};
+
+	EGetValidAssetPathResult GetValidAssetPath(EPVAssetReplacePolicy AssetReplacePolicy, UClass* AssetClass, FString& InOutPackageName, FString& InOutAssetName)
 	{
 		if (!FPackageName::IsValidTextForLongPackageName(InOutPackageName))
 		{
-			return false;
+			return EGetValidAssetPathResult::InvalidPath;
 		}
 
 		if (FindPackage(nullptr, *InOutPackageName) != nullptr || FPackageName::DoesPackageExist(InOutPackageName))
@@ -215,16 +223,16 @@ namespace PV::Export::Internal
 			}
 			else if (AssetReplacePolicy == EPVAssetReplacePolicy::Ignore)
 			{
-				return false;
+				return EGetValidAssetPathResult::SkipAsset;
 			}
 		}
 
 		if (Utilities::DoesConflictingPackageExist(InOutPackageName, AssetClass))
 		{
-			return false;
+			return EGetValidAssetPathResult::ConflictingPackageExist;
 		}
 
-		return true;
+		return EGetValidAssetPathResult::Success;
 	}
 
 	UPackage* CreateAssetPackage(const FString& LongPackageName)
@@ -1490,35 +1498,75 @@ namespace PV::Export::Internal
 
 namespace PV::Export
 {
-	bool ExportCollectionAsMesh(
+	EExportResult ExportCollectionAsMesh(
 		const TObjectPtr<UProceduralVegetation> InProceduralVegetation,
 		const FManagedArrayCollection& Collection,
 		const FPVExportParams& ExportParams,
+		TArray<FString>& OutCreatedAssets,
 		const FStatusReportCallback& StatusReportCallback
 	)
 	{
 		TRACE_CPUPROFILER_EVENT_SCOPE(PV::Export::ExportCollectionAsMesh);
 
+		TArray<UPackage*> ModifiedPackages;
+
+		const FDateTime ExportStartTime = FDateTime::Now();
+		TArray<FAnalyticsEventAttribute> AnalyticsEventAttribute = Analytics::GatherExportCommonAttributes(Collection, ExportParams);
+		Analytics::SendExportStartEvent(AnalyticsEventAttribute);
+
+		const auto ExitWithResult = [&](EExportResult InResult)->EExportResult
+		{
+			const static auto ExportResultToAnalyticsResult = [](EExportResult InExportResult)->PV::Analytics::EExportResult
+			{
+				switch (InExportResult)
+				{
+					case EExportResult::Success:
+						return PV::Analytics::EExportResult::Success;
+					case EExportResult::Fail:
+						return PV::Analytics::EExportResult::Failed;
+					case EExportResult::Canceled:
+						return PV::Analytics::EExportResult::Failed;
+					case EExportResult::Skipped:
+						return PV::Analytics::EExportResult::Skipped;
+				}
+				ensure(false);
+				return PV::Analytics::EExportResult::Failed;
+			};
+
+			if (InResult != EExportResult::Success)
+			{
+				Internal::CleanupPackages(ModifiedPackages);
+			}
+
+			const double ExportTime = (FDateTime::Now() - ExportStartTime).GetTotalSeconds();
+			Analytics::SendExportFinishedEvent(MoveTemp(AnalyticsEventAttribute), ExportTime, ExportResultToAnalyticsResult(InResult));
+			return InResult;
+		};
+
+		const auto ExitWithAssetPathResult = [&](Internal::EGetValidAssetPathResult InAssetPathResult)->EExportResult
+		{
+			if (InAssetPathResult == Internal::EGetValidAssetPathResult::SkipAsset)
+			{
+				return ExitWithResult(EExportResult::Skipped);
+			}
+
+			return ExitWithResult(EExportResult::Fail);
+		};
+
 		FString MeshName = ExportParams.MeshName.ToString();
 		FString MeshPackageName = ExportParams.GetOutputMeshPackagePath();
 
-		if (!Internal::GetValidAssetPath(ExportParams.ReplacePolicy, ExportParams.GetMeshClass(), MeshPackageName, MeshName))
+		Internal::EGetValidAssetPathResult AssetPathResult = Internal::GetValidAssetPath(ExportParams.ReplacePolicy, ExportParams.GetMeshClass(), MeshPackageName, MeshName);
+		if (AssetPathResult != Internal::EGetValidAssetPathResult::Success)
 		{
-			return false;
+			return ExitWithAssetPathResult(AssetPathResult);
 		}
 		
 		const bool bCreateNaniteFoliage = ExportParams.bCreateNaniteFoliage && NaniteAssembliesSupported();
 		const bool bShouldExportFoliage = ExportParams.bShouldExportFoliage;
 
-		TArray<UPackage*> ModifiedPackages;
-
-		FDateTime ExportStartTime = FDateTime::Now();
-		TArray<FAnalyticsEventAttribute> ExportAnalyticsAttributes = Analytics::GatherExportCommonAttributes(Collection, ExportParams);
-
 		if (ExportParams.ExportMeshType == EPVExportMeshType::StaticMesh)
 		{
-			Analytics::SendExportStartEvent(ExportAnalyticsAttributes);
-
 			TObjectPtr<UPackage> MeshPackage = Internal::CreateAssetPackage(MeshPackageName);
 			if (TObjectPtr<UStaticMesh> ExistingMeshAsset = Cast<UStaticMesh>(MeshPackage->FindAssetInPackage()))
 			{
@@ -1540,10 +1588,11 @@ namespace PV::Export
 				StatusReportCallback
 			))
 			{
-				Internal::CleanupPackages(ModifiedPackages);
-				return false;
+				return ExitWithResult(EExportResult::Canceled);
 			}
 
+			OutCreatedAssets.Add(MeshPackageName);
+
 			Internal::AttachProceduralVegetationLink(ExportMesh, InProceduralVegetation);
 		}
 		else if (ExportParams.ExportMeshType == EPVExportMeshType::SkeletalMesh)
@@ -1551,13 +1600,12 @@ namespace PV::Export
 			FString SkeletonName = ExportParams.GetOutputSkeletonName();
 			FString SkeletonPackageName = ExportParams.GetOutputSkeletonPackagePath();
 
-			if (!Internal::GetValidAssetPath(ExportParams.ReplacePolicy, USkeleton::StaticClass(), SkeletonPackageName, SkeletonName))
+			AssetPathResult = Internal::GetValidAssetPath(ExportParams.ReplacePolicy, USkeleton::StaticClass(), SkeletonPackageName, SkeletonName);
+			if (AssetPathResult != Internal::EGetValidAssetPathResult::Success)
 			{
-				return false;
+				return ExitWithAssetPathResult(AssetPathResult);
 			}
 
-			Analytics::SendExportStartEvent(ExportAnalyticsAttributes);
-
 			TObjectPtr<UPackage> MeshPackage = Internal::CreateAssetPackage(MeshPackageName);
 			TObjectPtr<UPackage> SkeletonPackage = Internal::CreateAssetPackage(SkeletonPackageName);
 
@@ -1592,10 +1640,12 @@ namespace PV::Export
 				StatusReportCallback
 			))
 			{
-				Internal::CleanupPackages(ModifiedPackages);
-				return false;
+				return ExitWithResult(EExportResult::Canceled);
 			}
 
+			OutCreatedAssets.Add(MeshPackageName);
+			OutCreatedAssets.Add(SkeletonPackageName);
+
 			Internal::AttachProceduralVegetationLink(ExportMesh, InProceduralVegetation);
 			Internal::AttachProceduralVegetationLink(ExportSkeleton, InProceduralVegetation);
 		}
@@ -1613,11 +1663,8 @@ namespace PV::Export
 				}
 			}
 		}
-		double ExportTime = (FDateTime::Now() - ExportStartTime).GetTotalSeconds();
 
-		Analytics::SendExportFinishedEvent(ExportAnalyticsAttributes, ExportTime);
-
-		return true;
+		return ExitWithResult(EExportResult::Success);
 	}
 }
 
diff --git a/Engine/Plugins/Experimental/ProceduralVegetationEditor/Source/ProceduralVegetationEditor/Private/Helpers/PVExportHelper.h b/Engine/Plugins/Experimental/ProceduralVegetationEditor/Source/ProceduralVegetationEditor/Private/Helpers/PVExportHelper.h
--- a/Engine/Plugins/Experimental/ProceduralVegetationEditor/Source/ProceduralVegetationEditor/Private/Helpers/PVExportHelper.h
+++ b/Engine/Plugins/Experimental/ProceduralVegetationEditor/Source/ProceduralVegetationEditor/Private/Helpers/PVExportHelper.h
@@ -24,10 +24,19 @@ namespace PV::Export
 {
 	typedef TFunction<bool(const FText& Stage, float Progress)> FStatusReportCallback;
 
-	bool ExportCollectionAsMesh(
+	enum class EExportResult
+	{
+		Success,
+		Fail,
+		Canceled,
+		Skipped
+	};
+
+	EExportResult ExportCollectionAsMesh(
 		const TObjectPtr<UProceduralVegetation> InProceduralVegetation,
 		const FManagedArrayCollection& Collection,
 		const FPVExportParams& ExportParams,
+		TArray<FString>& OutCreatedAssets,
 		const FStatusReportCallback& StatusReportCallback
 	);
 }
@@ -70,6 +79,7 @@ namespace PV::Export::Internal
 		const FManagedArrayCollection& Collection,
 		ENaniteShapePreservation InShapePreservation,
 		bool bBuildNaniteAssemblies,
+		bool bShouldExportFoliage,
 		TArray<UPackage*>& OutModifiedPackages,
 		const FStatusReportCallback& StatusReportCallback
 	);
@@ -80,6 +90,7 @@ namespace PV::Export::Internal
 		ENaniteShapePreservation InShapePreservation,
 		TObjectPtr<const UPVWindSettings> InWindSettings,
 		bool bBuildNaniteAssemblies,
+		bool bShouldExportFoliage,
 		TArray<UPackage*>& OutModifiedPackages,
 		const FStatusReportCallback& StatusReportCallback
 	);
diff --git a/Engine/Plugins/Experimental/ProceduralVegetationEditor/Source/ProceduralVegetationEditor/Private/PVEditor.cpp b/Engine/Plugins/Experimental/ProceduralVegetationEditor/Source/ProceduralVegetationEditor/Private/PVEditor.cpp
--- a/Engine/Plugins/Experimental/ProceduralVegetationEditor/Source/ProceduralVegetationEditor/Private/PVEditor.cpp
+++ b/Engine/Plugins/Experimental/ProceduralVegetationEditor/Source/ProceduralVegetationEditor/Private/PVEditor.cpp
@@ -6,6 +6,7 @@
 #include "PCGDefaultExecutionSource.h"
 #include "PCGEditorCommands.h"
 #include "ProceduralVegetation.h"
+#include "ProceduralVegetationEditorModule.h"
 #include "PVEditorCommands.h"
 #include "PVEditorSchema.h"
 #include "PVEditorSettings.h"
@@ -26,6 +27,9 @@
 #include "Misc/MessageDialog.h"
 #include "Misc/UObjectToken.h"
 
+#include "MessageLogModule.h"
+#include "IMessageLogListing.h"
+
 #include "Nodes/PCGEditorGraphNodeBase.h"
 #include "Nodes/PVOutputSettings.h"
 
@@ -337,31 +341,65 @@ FReply FPVEditor::ExportInternal()
 		return FReply::Handled();
 	}
 
-	struct FScopedErrorMessages
+	struct FScopedMessages
 	{
-		TArray<TSharedRef<FTokenizedMessage>> ErrorMessages;
+		TArray<TSharedRef<FTokenizedMessage>> Messages;
 
-		TSharedRef<FTokenizedMessage> AddError(const FText& ErrorMsg)
+		TSharedRef<FTokenizedMessage> Error(const FText& ErrorMsg)
+		{
+			TSharedRef<FTokenizedMessage> Message = FTokenizedMessage::Create(EMessageSeverity::Error, ErrorMsg);
+			Messages.Add(Message);
+			return Message;
+		}
+
+		TSharedRef<FTokenizedMessage> Warning(const FText& ErrorMsg)
+		{
+			TSharedRef<FTokenizedMessage> Message = FTokenizedMessage::Create(EMessageSeverity::Warning, ErrorMsg);
+			Messages.Add(Message);
+			return Message;
+		}
+
+		TSharedRef<FTokenizedMessage> Info(const FText& ErrorMsg)
+		{
+			TSharedRef<FTokenizedMessage> Message = FTokenizedMessage::Create(EMessageSeverity::Info, ErrorMsg);
+			Messages.Add(Message);
+			return Message;
+		}
+		
+		bool ContainsErrors() const
 		{
-			TSharedRef<FTokenizedMessage> Error = FTokenizedMessage::Create(EMessageSeverity::Error, ErrorMsg);
-			ErrorMessages.Add(Error);
-			return Error;
+			for (TSharedRef<FTokenizedMessage> Message : Messages)
+			{
+				if (Message->GetSeverity() == EMessageSeverity::Error)
+				{
+					return true;
+				}
+			}
+
+			return false;
 		}
 
-		~FScopedErrorMessages()
+		~FScopedMessages()
 		{
-			if (ErrorMessages.Num() > 0)
+			if (Messages.Num() > 0)
 			{
-				FMessageLog EditorErrors("EditorErrors");
-				EditorErrors.AddMessages(ErrorMessages);
-				EditorErrors.Notify(LOCTEXT("PVEMeshExportFailed", "PVE Mesh Export Failed"));
+				FMessageLogModule& MessageLogModule = FModuleManager::LoadModuleChecked<FMessageLogModule>("MessageLog");
+				TSharedRef<IMessageLogListing> PVEditorMessageLogListing = MessageLogModule.GetLogListing(PVEditor::MessageLogName);
+
+				PVEditorMessageLogListing->ClearMessages();
+				PVEditorMessageLogListing->AddMessages(Messages);
+				PVEditorMessageLogListing->NotifyIfAnyMessages(
+					ContainsErrors() ? LOCTEXT("PVEMeshExportFailed", "Mesh Export Failed") : LOCTEXT("PVEMeshExportComplete", "Mesh Export Completed"),
+					EMessageSeverity::Info,
+					true
+				);
 			}
 		}
 	};
 
-	FScopedErrorMessages ErrorMessages;
+	FScopedMessages Messages;
 
-	const static auto ValidateExportSettings = [](const TArray<UPCGNode*>& NodesToExport, FScopedErrorMessages& ErrorMessages)->bool
+	const static auto ValidateExportSettings = [](const TArray<UPCGNode*>& NodesToExport, FScopedMessages& OutMessages)->bool
 	{
 		// Check that the paths are valid
 		bool bHasInvalidPath = false;
@@ -372,7 +410,7 @@ FReply FPVEditor::ExportInternal()
 			FString PathError;
 			if (!OutputSettings->ExportSettings.Validate(PathError))
 			{
-				TSharedRef<FTokenizedMessage> Error = ErrorMessages.AddError(
+				TSharedRef<FTokenizedMessage> Error = OutMessages.Error(
 					FText::Format(
 						LOCTEXT("ExportFailed_InvalidPath", "Invalid export path for node \"{0}\": \"{1}\""),
 						Node->GetDefaultTitle(),
@@ -450,7 +488,7 @@ FReply FPVEditor::ExportInternal()
 		return UserResponse == DRT_Save;
 	};
 
-	const auto GetOutputNodeCollections = [this](TArray<FManagedArrayCollection>& OutCollections, const TArray<UPCGNode*>& NodesToExport, FScopedErrorMessages& ErrorMessages)->bool
+	const auto GetOutputNodeCollections = [this](TArray<FManagedArrayCollection>& OutCollections, const TArray<UPCGNode*>& NodesToExport, FScopedMessages& OutMessages)->bool
 	{
 		OutCollections.Reset(NodesToExport.Num());
 		for (UPCGNode* Node : NodesToExport)
@@ -459,7 +497,7 @@ FReply FPVEditor::ExportInternal()
 			FName PinName = NAME_None;
 			if (!GetCollectionFromNode(Node, PinName, OutCollection))
 			{
-				TSharedRef<FTokenizedMessage> Error = ErrorMessages.AddError(
+				TSharedRef<FTokenizedMessage> Error = OutMessages.Error(
 					FText::Format(
 						LOCTEXT("ExportFailed_NodeErrors", "Export failed for output node \"{0}\", Invalid output please fix the errors on the node."),
 						Node->GetDefaultTitle()
@@ -547,7 +585,7 @@ FReply FPVEditor::ExportInternal()
 		});
 	}
 
-	if (!ValidateExportSettings(NodesToExport, ErrorMessages))
+	if (!ValidateExportSettings(NodesToExport, Messages))
 	{
 		return FReply::Handled();
 	}
@@ -558,7 +596,7 @@ FReply FPVEditor::ExportInternal()
 	}
 	
 	TArray<FManagedArrayCollection> CollectionsToExport;
-	if (!GetOutputNodeCollections(CollectionsToExport, NodesToExport, ErrorMessages))
+	if (!GetOutputNodeCollections(CollectionsToExport, NodesToExport, Messages))
 	{
 		return FReply::Handled();
 	}
@@ -606,44 +644,64 @@ FReply FPVEditor::ExportInternal()
 			return !ProgressSlowTask.ShouldCancel();
 		};
 
-		const bool bExportSucceeded = PV::Export::ExportCollectionAsMesh(
+		TArray<FString> CreatedAssets;
+		const PV::Export::EExportResult ExportResult = PV::Export::ExportCollectionAsMesh(
 			ProceduralVegetationBeingEdited, 
 			CollectionsToExport[i], 
-			OutputSettings->ExportSettings, 
+			OutputSettings->ExportSettings,
+			CreatedAssets,
 			OnStatusUpdated
 		);
 
-		if (!bExportSucceeded)
+		if (ExportResult == PV::Export::EExportResult::Fail)
+		{
+			Messages.Error(FText::Format(
+				LOCTEXT("ExportFailed_UnknownError", "Failed to export mesh \"{0}\" due to unknown error, export aborted."),
+				FText::FromName(OutputSettings->ExportSettings.MeshName)
+			));
+			
+			break;
+		}
+		else if (ExportResult == PV::Export::EExportResult::Canceled)
 		{
-			if (ProgressSlowTask.ShouldCancel())
+			FString MeshNames = "{ ";
+			for (int32 j = i; j < NodesToExport.Num(); j++)
 			{
-				FString MeshNames = "{ ";
-				for (int32 j = i; j < NodesToExport.Num(); j++)
+				MeshNames += CastChecked<const UPVOutputSettings>(NodesToExport[j]->GetSettings())->ExportSettings.MeshName.ToString();
+				if (j != (NodesToExport.Num() - 1))
 				{
-					MeshNames += CastChecked<const UPVOutputSettings>(NodesToExport[j]->GetSettings())->ExportSettings.MeshName.ToString();
-					if (j != (NodesToExport.Num() - 1))
-					{
-						MeshNames += ", ";
-					}
+					MeshNames += ", ";
 				}
-				MeshNames += " }";
-
-				ErrorMessages.AddError(
-					FText::Format(
-						LOCTEXT("ExportFailed_ExportCanceled", "Failed to export meshes \"{0}\", operation canceled by user."),
-						FText::FromString(MeshNames)
-					));
 			}
-			else
-			{
-				ErrorMessages.AddError(FText::Format(
-					LOCTEXT("ExportFailed_UnknownError", "Failed to export mesh \"{0}\" due to unknown error, export aborted."),
-					FText::FromName(OutputSettings->ExportSettings.MeshName)
+			MeshNames += " }";
+
+			Messages.Warning(
+				FText::Format(
+					LOCTEXT("ExportFailed_ExportCanceled", "Failed to export meshes \"{0}\", operation canceled by user."),
+					FText::FromString(MeshNames)
 				));
-			}
-			
+
 			break;
 		}
+		else if (ExportResult == PV::Export::EExportResult::Skipped)
+		{
+			Messages.Info(FText::Format(
+				LOCTEXT("ExportSkipped", "Export of mesh \"{0}\" skipped due to conflicting asset at output location"),
+				FText::FromName(OutputSettings->ExportSettings.MeshName)
+			));
+		}
+		else
+		{
+			auto Message = Messages.Info(FText::Format(
+				LOCTEXT("ExportComplete", "Mesh exported successfully \"{0}\""),
+				FText::FromName(OutputSettings->ExportSettings.MeshName)
+			));
+
+			for (const FString& AssetPath : CreatedAssets)
+			{
+				Message->AddToken(FAssetNameToken::Create(AssetPath));
+			}
+		}
 	}
 	
 	return FReply::Handled();
diff --git a/Engine/Plugins/Experimental/ProceduralVegetationEditor/Source/ProceduralVegetationEditor/Private/ProceduralVegetationEditorModule.cpp b/Engine/Plugins/Experimental/ProceduralVegetationEditor/Source/ProceduralVegetationEditor/Private/ProceduralVegetationEditorModule.cpp
--- a/Engine/Plugins/Experimental/ProceduralVegetationEditor/Source/ProceduralVegetationEditor/Private/ProceduralVegetationEditorModule.cpp
+++ b/Engine/Plugins/Experimental/ProceduralVegetationEditor/Source/ProceduralVegetationEditor/Private/ProceduralVegetationEditorModule.cpp
@@ -40,6 +40,7 @@
 #include "Helpers/PVExportHelper.h"
 
 #include "Modules/ModuleManager.h"
+#include "MessageLogModule.h"
 
 #include "Templates/SharedPointer.h"
 
@@ -49,6 +50,8 @@
 
 DEFINE_LOG_CATEGORY(LogProceduralVegetationEditor);
 
+const FName PVEditor::MessageLogName = TEXT("PVEditor");
+
 void FProceduralVegetationEditorModule::StartupModule()
 {
 	// This code will execute after your module is loaded into memory; the exact timing is specified in the .uplugin file per-module
@@ -64,6 +67,9 @@ void FProceduralVegetationEditorModule::StartupModule()
 	FPropertyEditorModule& PropertyEditor = FModuleManager::LoadModuleChecked<FPropertyEditorModule>("PropertyEditor");
 	PropertyEditor.RegisterCustomPropertyTypeLayout(FPVExportParams::StaticStruct()->GetFName(), FOnGetPropertyTypeCustomizationInstance::CreateStatic(&FPVOutputSettingsCustomizations::MakeInstance));
 	PropertyEditor.RegisterCustomPropertyTypeLayout(FPVFloatRamp::StaticStruct()->GetFName(), FOnGetPropertyTypeCustomizationInstance::CreateStatic(&FPVFloatRampCustomization::MakeInstance));
+
+	FMessageLogModule& MessageLogModule = FModuleManager::LoadModuleChecked<FMessageLogModule>("MessageLog");
+	MessageLogModule.RegisterLogListing(PVEditor::MessageLogName, LOCTEXT("PVEditorLogLabel", "Procedural Vegetation Editor"));
 }
 
 void FProceduralVegetationEditorModule::ShutdownModule()
@@ -80,6 +86,11 @@ void FProceduralVegetationEditorModule::ShutdownModule()
 			PropertyEditor->UnregisterCustomPropertyTypeLayout(FPVFloatRamp::StaticStruct()->GetFName());
 		}
 	}
+
+	if (FMessageLogModule* MessageLogModule = FModuleManager::LoadModulePtr<FMessageLogModule>("MessageLog"))
+	{
+		MessageLogModule->UnregisterLogListing(PVEditor::MessageLogName);
+	}
 }
 
 void FProceduralVegetationEditorModule::RegisterMenus()
diff --git a/Engine/Plugins/Experimental/ProceduralVegetationEditor/Source/ProceduralVegetationEditor/Public/ProceduralVegetationEditorModule.h b/Engine/Plugins/Experimental/ProceduralVegetationEditor/Source/ProceduralVegetationEditor/Public/ProceduralVegetationEditorModule.h
--- a/Engine/Plugins/Experimental/ProceduralVegetationEditor/Source/ProceduralVegetationEditor/Public/ProceduralVegetationEditorModule.h
+++ b/Engine/Plugins/Experimental/ProceduralVegetationEditor/Source/ProceduralVegetationEditor/Public/ProceduralVegetationEditorModule.h
@@ -9,6 +9,11 @@ class FUICommandList;
 class FExtender;
 struct FToolMenuContext;
 
+namespace PVEditor
+{
+	extern const FName MessageLogName;
+};
+
 DECLARE_LOG_CATEGORY_EXTERN(LogProceduralVegetationEditor, Log, All);
 
 class PROCEDURALVEGETATIONEDITOR_API FProceduralVegetationEditorModule : public IModuleInterface
