Fix minor spelling mistake

diff --git a/Engine/Plugins/Experimental/ProceduralVegetationEditor/Source/ProceduralVegetation/Private/Helpers/PVAnalyticsHelper.cpp b/Engine/Plugins/Experimental/ProceduralVegetationEditor/Source/ProceduralVegetation/Private/Helpers/PVAnalyticsHelper.cpp
--- a/Engine/Plugins/Experimental/ProceduralVegetationEditor/Source/ProceduralVegetation/Private/Helpers/PVAnalyticsHelper.cpp
+++ b/Engine/Plugins/Experimental/ProceduralVegetationEditor/Source/ProceduralVegetation/Private/Helpers/PVAnalyticsHelper.cpp
@@ -188,7 +188,7 @@ namespace PV::Analytics
 					return TEXT("Skipped");
 				}
 
-				return TEXT("Unkown");
+				return TEXT("Unknown");
 			};
 			
 			InCommonAttributes.Add(FAnalyticsEventAttribute(TEXT("ExportResult"), ExportResultEnumToString(ExportResult)));
