[PVE] Ability to lock/unlock inspection with control + L
Refactored code to remove bAutoFocusViewport from PVBaseSettings

diff --git a/Engine/Plugins/Experimental/ProceduralVegetationEditor/Source/ProceduralVegetation/Private/Nodes/PVBoneReductionSettings.cpp b/Engine/Plugins/Experimental/ProceduralVegetationEditor/Source/ProceduralVegetation/Private/Nodes/PVBoneReductionSettings.cpp
--- a/Engine/Plugins/Experimental/ProceduralVegetationEditor/Source/ProceduralVegetation/Private/Nodes/PVBoneReductionSettings.cpp
+++ b/Engine/Plugins/Experimental/ProceduralVegetationEditor/Source/ProceduralVegetation/Private/Nodes/PVBoneReductionSettings.cpp
@@ -20,7 +20,8 @@ FText UPVBoneReductionSettings::GetNodeTooltipText() const
 { 
 	return LOCTEXT("NodeTooltip", 
 		"Allows the user to reduce the complexity of the skeletal bones generated. "
-		"Reduced bone complexity helps in optimizing performance but can effect the accuracy of wind simulation"
+		"Reduced bone complexity helps in optimizing performance but can effect the accuracy of wind simulation."
+		"\n\nPress Ctrl + L to lock/unlock node output"
 	);
 }
 #endif
diff --git a/Engine/Plugins/Experimental/ProceduralVegetationEditor/Source/ProceduralVegetation/Private/Nodes/PVCarveSettings.cpp b/Engine/Plugins/Experimental/ProceduralVegetationEditor/Source/ProceduralVegetation/Private/Nodes/PVCarveSettings.cpp
--- a/Engine/Plugins/Experimental/ProceduralVegetationEditor/Source/ProceduralVegetation/Private/Nodes/PVCarveSettings.cpp
+++ b/Engine/Plugins/Experimental/ProceduralVegetationEditor/Source/ProceduralVegetation/Private/Nodes/PVCarveSettings.cpp
@@ -23,6 +23,7 @@ FText UPVCarveSettings::GetNodeTooltipText() const
 		"From Bottom – relative height starting from the plant’s lowest point.\n"
 		"ZPosition – absolute world - space Z axis position.\n"
 		"Radius – Thickness of branch"
+		"\n\nPress Ctrl + L to lock/unlock node output"
 	); 
 }
 #endif
diff --git a/Engine/Plugins/Experimental/ProceduralVegetationEditor/Source/ProceduralVegetation/Private/Nodes/PVFoliageDistributorSettings.cpp b/Engine/Plugins/Experimental/ProceduralVegetationEditor/Source/ProceduralVegetation/Private/Nodes/PVFoliageDistributorSettings.cpp
--- a/Engine/Plugins/Experimental/ProceduralVegetationEditor/Source/ProceduralVegetation/Private/Nodes/PVFoliageDistributorSettings.cpp
+++ b/Engine/Plugins/Experimental/ProceduralVegetationEditor/Source/ProceduralVegetation/Private/Nodes/PVFoliageDistributorSettings.cpp
@@ -25,6 +25,7 @@ FText UPVFoliageDistributorSettings::GetNodeTooltipText() const
 		"EthyleneThreshold: Defines how sensitive the procedural system is to \"ethylene levels.\"\n"
 		"\tHigher Threshold = more buds/branches retained, making the plant denser and bushier.\n"
 		"\tLower Threshold = fewer buds/branches retained, resulting in a sparser, pruned structure."
+		"\n\nPress Ctrl + L to lock/unlock node output"
 	); 
 }
 #endif
diff --git a/Engine/Plugins/Experimental/ProceduralVegetationEditor/Source/ProceduralVegetation/Private/Nodes/PVFoliagePaletteSettings.cpp b/Engine/Plugins/Experimental/ProceduralVegetationEditor/Source/ProceduralVegetation/Private/Nodes/PVFoliagePaletteSettings.cpp
--- a/Engine/Plugins/Experimental/ProceduralVegetationEditor/Source/ProceduralVegetation/Private/Nodes/PVFoliagePaletteSettings.cpp
+++ b/Engine/Plugins/Experimental/ProceduralVegetationEditor/Source/ProceduralVegetation/Private/Nodes/PVFoliagePaletteSettings.cpp
@@ -27,6 +27,7 @@ FText UPVFoliagePaletteSettings::GetNodeTooltipText() const
 		"Allows the user to visualize and optionally change the foliage meshes. "
 		"Foliage mesh data is loaded from the Procedrual Vegetation preset. "
 		"The user can replace or remove foliage meshes. The placement of folliage meshes is driven by the procedural vegetaiton preset and cannot be modifed"
+		"\n\nPress Ctrl + L to lock/unlock node output"
 	);
 }
 
diff --git a/Engine/Plugins/Experimental/ProceduralVegetationEditor/Source/ProceduralVegetation/Private/Nodes/PVGravitySettings.cpp b/Engine/Plugins/Experimental/ProceduralVegetationEditor/Source/ProceduralVegetation/Private/Nodes/PVGravitySettings.cpp
--- a/Engine/Plugins/Experimental/ProceduralVegetationEditor/Source/ProceduralVegetation/Private/Nodes/PVGravitySettings.cpp
+++ b/Engine/Plugins/Experimental/ProceduralVegetationEditor/Source/ProceduralVegetation/Private/Nodes/PVGravitySettings.cpp
@@ -25,6 +25,7 @@ FText UPVGravitySettings::GetNodeTooltipText() const
 		"Strength Controls: Parameters adjust how strongly branches respond, allowing subtle or dramatic curvature.\n"
 		"Direction Vector: When in gravity mode, you can set the exact vector (e.g., world down, custom direction).\n"
 		"Angle Correction & Bias: Additional properties fine-tune how branches preserve their initial angles or bias between optimal light direction/shadow avoidance.\n"
+		"\n\nPress Ctrl + L to lock/unlock node output"
 	); 
 }
 #endif
diff --git a/Engine/Plugins/Experimental/ProceduralVegetationEditor/Source/ProceduralVegetation/Private/Nodes/PVMeshBuilderSettings.cpp b/Engine/Plugins/Experimental/ProceduralVegetationEditor/Source/ProceduralVegetation/Private/Nodes/PVMeshBuilderSettings.cpp
--- a/Engine/Plugins/Experimental/ProceduralVegetationEditor/Source/ProceduralVegetation/Private/Nodes/PVMeshBuilderSettings.cpp
+++ b/Engine/Plugins/Experimental/ProceduralVegetationEditor/Source/ProceduralVegetation/Private/Nodes/PVMeshBuilderSettings.cpp
@@ -29,6 +29,7 @@ FText UPVMeshBuilderSettings::GetNodeTooltipText() const
 		"- Material\n"
 		"- Mesh\n"
 		"- Displacement\n"
+		"\n\nPress Ctrl + L to lock/unlock node output"
 	); 
 }
 
diff --git a/Engine/Plugins/Experimental/ProceduralVegetationEditor/Source/ProceduralVegetation/Private/Nodes/PVPresetLoaderSettings.cpp b/Engine/Plugins/Experimental/ProceduralVegetationEditor/Source/ProceduralVegetation/Private/Nodes/PVPresetLoaderSettings.cpp
--- a/Engine/Plugins/Experimental/ProceduralVegetationEditor/Source/ProceduralVegetation/Private/Nodes/PVPresetLoaderSettings.cpp
+++ b/Engine/Plugins/Experimental/ProceduralVegetationEditor/Source/ProceduralVegetation/Private/Nodes/PVPresetLoaderSettings.cpp
@@ -22,6 +22,7 @@ FText UPVPresetLoaderSettings::GetNodeTooltipText() const
 		"This allows the user to load a procedural vegetation preset. "
 		"A procedural vegatation preset contains data for a specie of vegetaion. "
 		"It contains hormonal information, the main structure of the vegetation, data about the branches, where foliage will / can be attached, references to foliage meshes, materials to be applied"
+		"\n\nPress Ctrl + L to lock/unlock node output"
 	);
 }
 
diff --git a/Engine/Plugins/Experimental/ProceduralVegetationEditor/Source/ProceduralVegetation/Private/Nodes/PVRemoveBranchesSettings.cpp b/Engine/Plugins/Experimental/ProceduralVegetationEditor/Source/ProceduralVegetation/Private/Nodes/PVRemoveBranchesSettings.cpp
--- a/Engine/Plugins/Experimental/ProceduralVegetationEditor/Source/ProceduralVegetation/Private/Nodes/PVRemoveBranchesSettings.cpp
+++ b/Engine/Plugins/Experimental/ProceduralVegetationEditor/Source/ProceduralVegetation/Private/Nodes/PVRemoveBranchesSettings.cpp
@@ -18,7 +18,10 @@ FText UPVRemoveBranchesSettings::GetDefaultNodeTitle() const
 
 FText UPVRemoveBranchesSettings::GetNodeTooltipText() const
 { 
-	return LOCTEXT("NodeTooltip", "Removes certain branches on the basis of Length, Radius, Light (How exposed a branch is to the light), Age and Generation"); 
+	return LOCTEXT("NodeTooltip",
+		"Removes certain branches on the basis of Length, Radius, Light (How exposed a branch is to the light), Age and Generation"
+		"\n\nPress Ctrl + L to lock/unlock node output"
+	); 
 }
 #endif
 
diff --git a/Engine/Plugins/Experimental/ProceduralVegetationEditor/Source/ProceduralVegetation/Private/Nodes/PVScaleSettings.cpp b/Engine/Plugins/Experimental/ProceduralVegetationEditor/Source/ProceduralVegetation/Private/Nodes/PVScaleSettings.cpp
--- a/Engine/Plugins/Experimental/ProceduralVegetationEditor/Source/ProceduralVegetation/Private/Nodes/PVScaleSettings.cpp
+++ b/Engine/Plugins/Experimental/ProceduralVegetationEditor/Source/ProceduralVegetation/Private/Nodes/PVScaleSettings.cpp
@@ -18,7 +18,10 @@ FText UPVScaleSettings::GetDefaultNodeTitle() const
 
 FText UPVScaleSettings::GetNodeTooltipText() const
 { 
-	return LOCTEXT("NodeTooltip", "This node defines how large or small vegetation elements appear within the generated structure. Scaling affects trunks, branches, leaves, or entire sub-structures.");
+	return LOCTEXT("NodeTooltip",
+		"This node defines how large or small vegetation elements appear within the generated structure. Scaling affects trunks, branches, leaves, or entire sub-structures."
+		"\n\nPress Ctrl + L to lock/unlock node output"
+	);
 }
 #endif
 
diff --git a/Engine/Plugins/Experimental/ProceduralVegetationEditor/Source/ProceduralVegetation/Public/Nodes/PVBaseSettings.h b/Engine/Plugins/Experimental/ProceduralVegetationEditor/Source/ProceduralVegetation/Public/Nodes/PVBaseSettings.h
--- a/Engine/Plugins/Experimental/ProceduralVegetationEditor/Source/ProceduralVegetation/Public/Nodes/PVBaseSettings.h
+++ b/Engine/Plugins/Experimental/ProceduralVegetationEditor/Source/ProceduralVegetation/Public/Nodes/PVBaseSettings.h
@@ -13,11 +13,11 @@ class UPVBaseSettings : public UPCGSettings
 	GENERATED_BODY()
 
 public:
-	bool bAutoFocusViewport = false;
-
 	UPVBaseSettings();
 	
 #if WITH_EDITOR
+	bool bLockInspection = false;
+	
 	virtual EPCGSettingsType GetType() const override { return EPCGSettingsType::Spatial; }
 	virtual TArray<EPVRenderType> GetDefaultRenderType() const { return TArray<EPVRenderType>(); }
 	virtual TArray<EPVRenderType> GetSupportedRenderTypes() const { return GetDefaultRenderType(); }
diff --git a/Engine/Plugins/Experimental/ProceduralVegetationEditor/Source/ProceduralVegetationEditor/Private/PVEditor.cpp b/Engine/Plugins/Experimental/ProceduralVegetationEditor/Source/ProceduralVegetationEditor/Private/PVEditor.cpp
--- a/Engine/Plugins/Experimental/ProceduralVegetationEditor/Source/ProceduralVegetationEditor/Private/PVEditor.cpp
+++ b/Engine/Plugins/Experimental/ProceduralVegetationEditor/Source/ProceduralVegetationEditor/Private/PVEditor.cpp
@@ -285,40 +285,66 @@ void FPVEditor::BindCommands()
 	ToolkitCommands->MapAction(
 		EditorCommands.Export,
 		FExecuteAction::CreateSP(this, &FPVEditor::OnExport_Clicked));
+
+	ToolkitCommands->MapAction(
+		EditorCommands.LockNodeInspection,
+		FExecuteAction::CreateSP(this, &FPVEditor::OnLockNodeSelection));
+}
+
+UPVBaseSettings* GetNodeSettings(UPCGEditorGraphNodeBase* InNode)
+{
+	if (InNode && InNode->GetPCGNode())
+	{
+		return Cast<UPVBaseSettings>(InNode->GetPCGNode()->GetSettings());
+	}
+
+	return nullptr;
+}
+
+void FPVEditor:: ChangeNodeInspection(UPCGEditorGraphNodeBase* InNode)
+{
+	if (auto Settings = GetNodeSettings(InNode))
+	{
+		if (ViewportWidgets.Num() > 0)
+		{
+			ViewportWidgets[0]->OnNodeInspectionChanged(Settings);
+		}
+
+		NodeBeingInspected = InNode;
+		
+		SetNodeInspected(InNode, true);	
+	}
 }
 
 void FPVEditor::OnSelectedNodesChanged(const TSet<UObject*>& InNewSelection)
 {
 	FPCGEditor::OnSelectedNodesChanged(InNewSelection);
 
-	UPCGNode* SelectedNode = nullptr;
-	
+	UPCGEditorGraphNodeBase* EdNode = nullptr;
+
 	if (!InNewSelection.IsEmpty())
 	{
-		UPCGEditorGraphNodeBase* Node = Cast<UPCGEditorGraphNodeBase>(*InNewSelection.CreateConstIterator());
+		EdNode = Cast<UPCGEditorGraphNodeBase>(*InNewSelection.CreateConstIterator());
 
-		if (Node && Node->GetPCGNode())
+		if (EdNode && EdNode->GetPCGNode())
 		{
-			SelectedNode = Node->GetPCGNode();
+			UPCGNode* Node = EdNode->GetPCGNode();
 
-			UpdateCollectionSpreadSheet(SelectedNode);
+			UpdateCollectionSpreadSheet(Node);
 
-			if (UPVBaseSettings* Settings = Cast<UPVBaseSettings>(SelectedNode->GetSettings()))
+			if (UPVBaseSettings* Settings = Cast<UPVBaseSettings>(Node->GetSettings()))
 			{
-				Settings->bAutoFocusViewport = bFocusViewportOnSelection;
-				bFocusViewportOnSelection = false;
-
-				if (ViewportWidgets.Num() > 0)
+				const UPVBaseSettings* InspectionSettings = GetNodeSettings(NodeBeingInspected);
+				
+				if (!InspectionSettings || !InspectionSettings->bLockInspection)
 				{
-					ViewportWidgets[0]->OnNodeSelectionChanged(Settings);
+					ChangeNodeInspection(EdNode);
 				}
 			}
-			
-			SetNodeInspected(Node, true);
 		}
 	}
 
-	CurrentSelectedNode = SelectedNode;
+	SelectedNode = EdNode;
 }
 
 void FPVEditor::OnExport_Clicked()
@@ -334,7 +360,7 @@ FReply FPVEditor::ExportInternal()
 		SNew(SPVExportSelectionDialog)
 		.Title(LOCTEXT("Export Selection", "Choose outputs to export"))
 		.Graph(ProceduralVegetationGraph)
-		.SelectedNode(CurrentSelectedNode);
+		.SelectedNode(SelectedNode ? SelectedNode->GetPCGNode() : nullptr);
 
 	if (ExportSelectionWidget->ShowModal() != EAppReturnType::Ok)
 	{
@@ -569,11 +595,11 @@ FReply FPVEditor::ExportInternal()
 
 	if (EditorSettings->bExportType == 0)
 	{
-		if (CurrentSelectedNode != nullptr)
+		if (SelectedNode != nullptr && SelectedNode->GetPCGNode())
 		{
-			if (const UPVOutputSettings* OutputSettings = Cast<UPVOutputSettings>(CurrentSelectedNode->GetSettings()))
+			if (const UPVOutputSettings* OutputSettings = Cast<UPVOutputSettings>(SelectedNode->GetPCGNode()->GetSettings()))
 			{
-				NodesToExport.Add(CurrentSelectedNode);
+				NodesToExport.Add(SelectedNode->GetPCGNode());
 			}
 		}
 	}
@@ -707,6 +733,26 @@ FReply FPVEditor::ExportInternal()
 	return FReply::Handled();
 }
 
+void FPVEditor::OnLockNodeSelection()
+{
+	auto NodeBeingInspectedSettings = GetNodeSettings(NodeBeingInspected);
+	auto SelectedNodeSettings = GetNodeSettings(SelectedNode);
+
+	if (NodeBeingInspectedSettings)
+	{
+		NodeBeingInspectedSettings->bLockInspection = !NodeBeingInspectedSettings->bLockInspection;
+	}
+
+	if (NodeBeingInspected != SelectedNode)
+	{
+		if (SelectedNode)
+		{
+			SelectedNodeSettings->bLockInspection = true;
+			ChangeNodeInspection(SelectedNode);
+		}
+	}
+}
+
 TSharedRef<SPCGEditorViewport> FPVEditor::CreateViewportWidget()
 {
 	TSharedPtr<SPVEditorViewport> ViewportWidget;
@@ -763,6 +809,11 @@ const UPVData* FPVEditor::GetPVDataFromNode(const UPCGNode* InNode, FName& OutPi
 	return nullptr;
 }
 
+bool FPVEditor::CanToggleInspected() const
+{
+	return false;
+}
+
 bool FPVEditor::GetCollectionFromNode(const UPCGNode* InNode, FName& OutPinName, FManagedArrayCollection& OutCollection) const
 {
 	const UPVData* PVData = GetPVDataFromNode(InNode, OutPinName);
diff --git a/Engine/Plugins/Experimental/ProceduralVegetationEditor/Source/ProceduralVegetationEditor/Private/PVEditor.h b/Engine/Plugins/Experimental/ProceduralVegetationEditor/Source/ProceduralVegetationEditor/Private/PVEditor.h
--- a/Engine/Plugins/Experimental/ProceduralVegetationEditor/Source/ProceduralVegetationEditor/Private/PVEditor.h
+++ b/Engine/Plugins/Experimental/ProceduralVegetationEditor/Source/ProceduralVegetationEditor/Private/PVEditor.h
@@ -10,6 +10,7 @@ class SCollectionSpreadSheetWidget;
 struct FManagedArrayCollection;
 class UPCGDefaultExecutionSource;
 class UProceduralVegetation;
+class UPCGEditorGraphNodeBase;
 
 class FPVEditor : public FPCGEditor
 {
@@ -50,14 +51,12 @@ protected:
 
 private:
 	TObjectPtr<UProceduralVegetation> ProceduralVegetationBeingEdited = nullptr;
-	UPCGNode* CurrentSelectedNode = nullptr;
-	
+	TObjectPtr<UPCGEditorGraphNodeBase> SelectedNode = nullptr;
+	TObjectPtr<UPCGEditorGraphNodeBase> NodeBeingInspected = nullptr;
 	TArray<TSharedRef<SPVEditorViewport>> ViewportWidgets;
 	
 	FDateTime SessionStartTime;
 
-	bool bFocusViewportOnSelection = true;
-
 	static const FName CollectionSpreadSheetTabId;
 
 	TWeakPtr<SDockTab> CollectionSpreadSheetTab;
@@ -72,9 +71,14 @@ private:
 	bool GetCollectionFromNode(const UPCGNode* InNode, FName& OutPinName, FManagedArrayCollection& OutCollection) const;
 	void OnExport_Clicked();
 	FReply ExportInternal();
+	
+	void OnLockNodeSelection();
+
+	void ChangeNodeInspection(UPCGEditorGraphNodeBase* InNode);
 
 protected:
 	virtual TSharedRef<SPCGEditorViewport> CreateViewportWidget() override;
 	
 	const UPVData* GetPVDataFromNode(const UPCGNode* InNode, FName& OutPinName) const;
+	virtual bool CanToggleInspected() const override;
 };
diff --git a/Engine/Plugins/Experimental/ProceduralVegetationEditor/Source/ProceduralVegetationEditor/Private/PVEditorCommands.cpp b/Engine/Plugins/Experimental/ProceduralVegetationEditor/Source/ProceduralVegetationEditor/Private/PVEditorCommands.cpp
--- a/Engine/Plugins/Experimental/ProceduralVegetationEditor/Source/ProceduralVegetationEditor/Private/PVEditorCommands.cpp
+++ b/Engine/Plugins/Experimental/ProceduralVegetationEditor/Source/ProceduralVegetationEditor/Private/PVEditorCommands.cpp
@@ -21,6 +21,8 @@ void FPVEditorCommands::RegisterCommands()
 	UI_COMMAND(ShowMannequin, "Mannequin", "Show/Hide Mannequin", EUserInterfaceActionType::ToggleButton, FInputChord(EModifierKey::Control, EKeys::M));
 	UI_COMMAND(ShowScaleVisualization, "Scale Visualization", "Show/Hide Scale Visualization", EUserInterfaceActionType::ToggleButton, FInputChord(EModifierKey::Control, EKeys::V));
 	UI_COMMAND(AutoFocusViewport, "Auto focus viewport", "Auto focus viewport on node selection and visualization mode switch", EUserInterfaceActionType::ToggleButton, FInputChord());
+	UI_COMMAND(LockNodeInspection, "LockNodeInspection", "Lock node inspection.", EUserInterfaceActionType::Button, FInputChord(EModifierKey::Control, EKeys::L));
+
 }
 
 #undef LOCTEXT_NAMESPACE
diff --git a/Engine/Plugins/Experimental/ProceduralVegetationEditor/Source/ProceduralVegetationEditor/Private/PVEditorCommands.h b/Engine/Plugins/Experimental/ProceduralVegetationEditor/Source/ProceduralVegetationEditor/Private/PVEditorCommands.h
--- a/Engine/Plugins/Experimental/ProceduralVegetationEditor/Source/ProceduralVegetationEditor/Private/PVEditorCommands.h
+++ b/Engine/Plugins/Experimental/ProceduralVegetationEditor/Source/ProceduralVegetationEditor/Private/PVEditorCommands.h
@@ -20,4 +20,5 @@ public:
 	TSharedPtr<FUICommandInfo> ShowMannequin;
 	TSharedPtr<FUICommandInfo> ShowScaleVisualization;
 	TSharedPtr<FUICommandInfo> AutoFocusViewport;
+	TSharedPtr<FUICommandInfo> LockNodeInspection;;
 };
\ No newline at end of file
diff --git a/Engine/Plugins/Experimental/ProceduralVegetationEditor/Source/ProceduralVegetationEditor/Private/SPVEditorViewport.cpp b/Engine/Plugins/Experimental/ProceduralVegetationEditor/Source/ProceduralVegetationEditor/Private/SPVEditorViewport.cpp
--- a/Engine/Plugins/Experimental/ProceduralVegetationEditor/Source/ProceduralVegetationEditor/Private/SPVEditorViewport.cpp
+++ b/Engine/Plugins/Experimental/ProceduralVegetationEditor/Source/ProceduralVegetationEditor/Private/SPVEditorViewport.cpp
@@ -3,6 +3,7 @@
 #include "SPVEditorViewport.h"
 #include "AdvancedPreviewScene.h"
 #include "AdvancedPreviewSceneMenus.h"
+#include "PCGEditorSettings.h"
 #include "PVEditorCommands.h"
 #include "PVEditorSettings.h"
 #include "ViewportToolbar/UnrealEdViewportToolbar.h"
@@ -31,6 +32,18 @@ SPVEditorViewport::~SPVEditorViewport()
 	}
 }
 
+void SPVEditorViewport::Construct(const FArguments& InArgs)
+{
+	SPCGEditorViewport::Construct(SPCGEditorViewport::FArguments());
+
+	UPCGEditorSettings* PCGEditorSettings = GetMutableDefault<UPCGEditorSettings>();
+
+	if (PCGEditorSettings)
+	{
+		PCGEditorSettings->bAutoFocusViewport = IsAutoFocusViewportChecked();
+	}
+}
+
 TSharedPtr<SWidget> SPVEditorViewport::BuildViewportToolbar()
 {
 	check(AdvancedPreviewScene.IsValid());
@@ -222,6 +235,9 @@ void SPVEditorViewport::ToggleAutofocusViewport()
 	UPVEditorSettings* EditorSettings = GetMutableDefault<UPVEditorSettings>();
 	EditorSettings->bAutoFocusViewport = !EditorSettings->bAutoFocusViewport;
 	EditorSettings->SaveConfig();
+
+	UPCGEditorSettings* PCGEditorSettings = GetMutableDefault<UPCGEditorSettings>();
+	PCGEditorSettings->bAutoFocusViewport = IsAutoFocusViewportChecked();
 }
 
 bool SPVEditorViewport::IsAutoFocusViewportChecked() const
@@ -479,11 +495,11 @@ void SPVEditorViewport::OnVisualizationModeChanged(EPVVisualizationMode InMode)
 {
 	CurrentVisualizationMode = InMode;
 
-	if (CurrentNodeSettings)
+	if (InspectedNodeSettings)
 	{
-		const TArray<EPVRenderType> RenderTypes = VisualizationModeToRenderModes(CurrentVisualizationMode, CurrentNodeSettings->GetDefaultRenderType());
+		const TArray<EPVRenderType> RenderTypes = VisualizationModeToRenderModes(CurrentVisualizationMode, InspectedNodeSettings->GetDefaultRenderType());
 												
-		CurrentNodeSettings->SetCurrentRenderType(RenderTypes);
+		InspectedNodeSettings->SetCurrentRenderType(RenderTypes);
 
 		Invalidate();
 
@@ -494,13 +510,16 @@ void SPVEditorViewport::OnVisualizationModeChanged(EPVVisualizationMode InMode)
 	}
 }
 
-void SPVEditorViewport::OnNodeSelectionChanged(UPVBaseSettings* InSettings)
+void SPVEditorViewport::OnNodeInspectionChanged(UPVBaseSettings* InSettings)
 {
 	SupportedRenderTypes = InSettings->GetSupportedRenderTypes();
 	CurrentVisualizationMode = EPVVisualizationMode::Default;
-	CurrentNodeSettings = InSettings;
+	InspectedNodeSettings = InSettings;
 	
-	CurrentNodeSettings->SetCurrentRenderType(InSettings->GetDefaultRenderType());
+	InspectedNodeSettings->SetCurrentRenderType(InSettings->GetDefaultRenderType());
+
+	// TODO: @Tayyab add lock visualization based on the inspection lock state
+	// InSettings->bLockInspection
 
 	if (IsAutoFocusViewportChecked())
 	{
diff --git a/Engine/Plugins/Experimental/ProceduralVegetationEditor/Source/ProceduralVegetationEditor/Private/SPVEditorViewport.h b/Engine/Plugins/Experimental/ProceduralVegetationEditor/Source/ProceduralVegetationEditor/Private/SPVEditorViewport.h
--- a/Engine/Plugins/Experimental/ProceduralVegetationEditor/Source/ProceduralVegetationEditor/Private/SPVEditorViewport.h
+++ b/Engine/Plugins/Experimental/ProceduralVegetationEditor/Source/ProceduralVegetationEditor/Private/SPVEditorViewport.h
@@ -39,11 +39,16 @@ public:
 class SPVEditorViewport : public SPCGEditorViewport
 {
 public:
+	SLATE_BEGIN_ARGS(SPVEditorViewport) {}
+	SLATE_END_ARGS()
+
 	virtual ~SPVEditorViewport() override;
+
+	void Construct(const FArguments& InArgs);
 	FToolMenuEntry CreateSettingsToolbarMenu();
 	FToolMenuEntry CreateVisualizationModeToolbarMenu();
 
-	void OnNodeSelectionChanged(UPVBaseSettings* InSettings);
+	void OnNodeInspectionChanged(UPVBaseSettings* InSettings);
 
 protected:
 	//~ Begin SEditorViewport Interface
@@ -75,11 +80,11 @@ private:
 
 	TArray<EPVRenderType> SupportedRenderTypes;
 	EPVVisualizationMode CurrentVisualizationMode = EPVVisualizationMode::Default;
-	TObjectPtr<UPVBaseSettings> CurrentNodeSettings = nullptr;
+	TObjectPtr<UPVBaseSettings> InspectedNodeSettings = nullptr;
 
 	void OnVisualizationModeChanged(EPVVisualizationMode InMode);
 
-	bool bFocusOnNextUpdate = false;
+	bool bFocusOnNextUpdate = true;
 
 protected:
 	virtual void OnSetupScene() override;
diff --git a/Engine/Plugins/Experimental/ProceduralVegetationEditor/Source/ProceduralVegetationEditor/Private/Visualizations/PVDataVisualization.cpp b/Engine/Plugins/Experimental/ProceduralVegetationEditor/Source/ProceduralVegetationEditor/Private/Visualizations/PVDataVisualization.cpp
--- a/Engine/Plugins/Experimental/ProceduralVegetationEditor/Source/ProceduralVegetationEditor/Private/Visualizations/PVDataVisualization.cpp
+++ b/Engine/Plugins/Experimental/ProceduralVegetationEditor/Source/ProceduralVegetationEditor/Private/Visualizations/PVDataVisualization.cpp
@@ -84,13 +84,6 @@ FPCGSetupSceneFunc FPVDataVisualization::GetViewportSetupFunc(const UPCGSettings
 		InOutParams.Scene->SetFloorVisibility(true);
 		InOutParams.Scene->SetFloorOffset(1.0f);
 			
-		UPCGEditorSettings* PCGEditorSettings = GetMutableDefault<UPCGEditorSettings>();
-
-		if (PCGEditorSettings && WeakSettings.IsValid())
-		{
-			PCGEditorSettings->bAutoFocusViewport = WeakSettings->bAutoFocusViewport;
-		}
-			
 		const FManagedArrayCollection& Collection = WeakData->GetCollection();
 
 		FBoxSphereBounds Bounds(ForceInit);
