[PVE] Bug hawk fixes

diff --git a/Engine/Plugins/Experimental/ProceduralVegetationEditor/Source/ProceduralVegetation/Private/Helpers/PVAnalyticsHelper.cpp b/Engine/Plugins/Experimental/ProceduralVegetationEditor/Source/ProceduralVegetation/Private/Helpers/PVAnalyticsHelper.cpp
--- a/Engine/Plugins/Experimental/ProceduralVegetationEditor/Source/ProceduralVegetation/Private/Helpers/PVAnalyticsHelper.cpp
+++ b/Engine/Plugins/Experimental/ProceduralVegetationEditor/Source/ProceduralVegetation/Private/Helpers/PVAnalyticsHelper.cpp
@@ -19,14 +19,6 @@ namespace PV::Analytics
 		{
 			FString EventName = FString::Printf(TEXT("Editor.Usage.PVE.%s"), *InEventName);
 			FEngineAnalytics::GetProvider().RecordEvent(EventName, InAttributes);
-
-			FString AttributesData;
-			for (int32 i = 0; i < InAttributes.Num(); i++)
-			{
-				AttributesData += FString::Printf(TEXT("(%s,%s),"), *InAttributes[i].GetName(), *InAttributes[i].GetValue());
-			}
-			
-			UE_LOG(LogTemp, Display, TEXT("PV::SendEvent: %s : Attributes: {%s}"), *EventName, *AttributesData);
 		}
 	}
 
@@ -62,7 +54,7 @@ namespace PV::Analytics
 
 	void SendNodeTweakedEvent(const FString& InNodeType, const FProperty* InProperty)
 	{
-		if (FEngineAnalytics::IsAvailable())
+		if (FEngineAnalytics::IsAvailable() && InProperty)
 		{
 			FString PropertyName = InProperty->GetName();
 
diff --git a/Engine/Plugins/Experimental/ProceduralVegetationEditor/Source/ProceduralVegetation/Public/Nodes/PVBaseSettings.h b/Engine/Plugins/Experimental/ProceduralVegetationEditor/Source/ProceduralVegetation/Public/Nodes/PVBaseSettings.h
--- a/Engine/Plugins/Experimental/ProceduralVegetationEditor/Source/ProceduralVegetation/Public/Nodes/PVBaseSettings.h
+++ b/Engine/Plugins/Experimental/ProceduralVegetationEditor/Source/ProceduralVegetation/Public/Nodes/PVBaseSettings.h
@@ -47,7 +47,7 @@ private:
 	FPVDebugSettings DebugVisualizationSettings;
 
 	UPROPERTY()
-	TArray<EPVRenderType> CurrentRenderType = GetDefaultRenderType();
+	TArray<EPVRenderType> CurrentRenderType;
 #endif
 };
 
diff --git a/Engine/Plugins/Experimental/ProceduralVegetationEditor/Source/ProceduralVegetationEditor/Private/PVEditor.cpp b/Engine/Plugins/Experimental/ProceduralVegetationEditor/Source/ProceduralVegetationEditor/Private/PVEditor.cpp
--- a/Engine/Plugins/Experimental/ProceduralVegetationEditor/Source/ProceduralVegetationEditor/Private/PVEditor.cpp
+++ b/Engine/Plugins/Experimental/ProceduralVegetationEditor/Source/ProceduralVegetationEditor/Private/PVEditor.cpp
@@ -592,8 +592,8 @@ FReply FPVEditor::ExportInternal()
 	TArray<UPCGNode*> NodesToExport;
 
 	const UPVEditorSettings* EditorSettings = GetDefault<UPVEditorSettings>();
-
-	if (EditorSettings->bExportType == 0)
+	
+	if (EditorSettings->ExportType == EPVExportType::Selection)
 	{
 		if (SelectedNode != nullptr && SelectedNode->GetPCGNode())
 		{
diff --git a/Engine/Plugins/Experimental/ProceduralVegetationEditor/Source/ProceduralVegetationEditor/Private/PVEditor.h b/Engine/Plugins/Experimental/ProceduralVegetationEditor/Source/ProceduralVegetationEditor/Private/PVEditor.h
--- a/Engine/Plugins/Experimental/ProceduralVegetationEditor/Source/ProceduralVegetationEditor/Private/PVEditor.h
+++ b/Engine/Plugins/Experimental/ProceduralVegetationEditor/Source/ProceduralVegetationEditor/Private/PVEditor.h
@@ -50,26 +50,27 @@ protected:
 	TObjectPtr<UPCGDefaultExecutionSource> ExecutionSource;
 
 private:
+	static const FName CollectionSpreadSheetTabId;
+	
 	TObjectPtr<UProceduralVegetation> ProceduralVegetationBeingEdited = nullptr;
 	TObjectPtr<UPCGEditorGraphNodeBase> SelectedNode = nullptr;
 	TObjectPtr<UPCGEditorGraphNodeBase> NodeBeingInspected = nullptr;
-	TArray<TSharedRef<SPVEditorViewport>> ViewportWidgets;
 	
-	FDateTime SessionStartTime;
-
-	static const FName CollectionSpreadSheetTabId;
-
-	TWeakPtr<SDockTab> CollectionSpreadSheetTab;
+	TArray<TSharedRef<SPVEditorViewport>> ViewportWidgets;
 	TSharedPtr<SCollectionSpreadSheetWidget> CollectionSpreadSheetWidget;
 
+	TWeakPtr<SDockTab> CollectionSpreadSheetTab;
+	
+	FDateTime SessionStartTime;
+	
 	TSharedRef<SCollectionSpreadSheetWidget> CreateCollectionSpreadSheetWidget();
 
 	TSharedRef<SDockTab> SpawnTab_CollectionSpreadSheet(const FSpawnTabArgs& SpawnTabArgs);
 	
 	void UpdateCollectionSpreadSheet(const UPCGNode* InSelectedNode) const;
 	
-	bool GetCollectionFromNode(const UPCGNode* InNode, FName& OutPinName, FManagedArrayCollection& OutCollection) const;
 	void OnExport_Clicked();
+
 	FReply ExportInternal();
 	
 	void OnLockNodeSelection();
@@ -81,4 +82,5 @@ protected:
 	
 	const UPVData* GetPVDataFromNode(const UPCGNode* InNode, FName& OutPinName) const;
 	virtual bool CanToggleInspected() const override;
+	bool GetCollectionFromNode(const UPCGNode* InNode, FName& OutPinName, FManagedArrayCollection& OutCollection) const;
 };
diff --git a/Engine/Plugins/Experimental/ProceduralVegetationEditor/Source/ProceduralVegetationEditor/Private/Widgets/SPVExportSelectionDialog.cpp b/Engine/Plugins/Experimental/ProceduralVegetationEditor/Source/ProceduralVegetationEditor/Private/Widgets/SPVExportSelectionDialog.cpp
--- a/Engine/Plugins/Experimental/ProceduralVegetationEditor/Source/ProceduralVegetationEditor/Private/Widgets/SPVExportSelectionDialog.cpp
+++ b/Engine/Plugins/Experimental/ProceduralVegetationEditor/Source/ProceduralVegetationEditor/Private/Widgets/SPVExportSelectionDialog.cpp
@@ -101,8 +101,13 @@ void SPVExportSelectionDialog::Construct(const FArguments& InArgs)
 
 EPVExportType SPVExportSelectionDialog::GetExportType()
 {
-	int SelectionIndex = ExportTypes.IndexOfByKey(CurrentExportType);
-	auto ExportType = (EPVExportType)SelectionIndex;
+	EPVExportType ExportType = EPVExportType::Selection;
+	
+	if (ExportTypes.Contains(CurrentExportType))
+	{
+		int SelectionIndex = ExportTypes.IndexOfByKey(CurrentExportType);
+		ExportType = (EPVExportType)SelectionIndex;
+	}
 
 	return ExportType;
 }
@@ -185,14 +190,17 @@ void SPVExportSelectionDialog::OnExportTypeChanged(TSharedPtr<FString> Selection
 {
 	if (!Selection->IsEmpty())
 	{
-		int SelectionIndex = ExportTypes.IndexOfByKey(Selection);
+		const int SelectionIndex = ExportTypes.IndexOfByKey(Selection);
 
-		CurrentExportType = ExportTypes[SelectionIndex];
+		if (ExportTypes.IsValidIndex(SelectionIndex))
+		{
+			CurrentExportType = ExportTypes[SelectionIndex];
 		
-		UPVEditorSettings* EditorSettings = GetMutableDefault<UPVEditorSettings>();
-		EditorSettings->bExportType = SelectionIndex;
+			UPVEditorSettings* EditorSettings = GetMutableDefault<UPVEditorSettings>();
+			EditorSettings->ExportType = (EPVExportType)SelectionIndex;
 
-		AddExportItems();
+			AddExportItems();	
+		}
 	}
 }
 
@@ -257,11 +265,11 @@ void SPVExportSelectionDialog::OnOutputSelectionChanged(const FString ItemName,
 
 void SPVExportSelectionDialog::BuildExportTypeOptions()
 {
-	UEnum* EnumPtr = StaticEnum<EPVExportType>();
+	const UEnum* EnumPtr = StaticEnum<EPVExportType>();
 	                      
 	if (EnumPtr)
 	{
-		for (int32 i = 0; i < EnumPtr->NumEnums() - 1; i++) // last entry is _MAX or hidden
+		for (int32 i = 0; i < EnumPtr->NumEnums() - 1; i++)
 		{
 			FText ModeName = EnumPtr->GetDisplayNameTextByIndex(i);
 			ExportTypes.Add(MakeShared<FString>(ModeName.ToString()));
@@ -271,10 +279,11 @@ void SPVExportSelectionDialog::BuildExportTypeOptions()
 	if (ExportTypes.Num() > 0)
 	{
 		const UPVEditorSettings* EditorSettings = GetDefault<UPVEditorSettings>();
-
-		if (EditorSettings->bExportType >= 0 && EditorSettings->bExportType < ExportTypes.Num())
+		const int ExportTypeIndex = (int)EditorSettings->ExportType; 
+		
+		if (ExportTypeIndex >= 0 && ExportTypeIndex < ExportTypes.Num())
 		{
-			CurrentExportType = ExportTypes[EditorSettings->bExportType];
+			CurrentExportType = ExportTypes[ExportTypeIndex];
 		}
 	}
 }
diff --git a/Engine/Plugins/Experimental/ProceduralVegetationEditor/Source/ProceduralVegetationEditor/Private/Widgets/SPVExportSelectionDialog.h b/Engine/Plugins/Experimental/ProceduralVegetationEditor/Source/ProceduralVegetationEditor/Private/Widgets/SPVExportSelectionDialog.h
--- a/Engine/Plugins/Experimental/ProceduralVegetationEditor/Source/ProceduralVegetationEditor/Private/Widgets/SPVExportSelectionDialog.h
+++ b/Engine/Plugins/Experimental/ProceduralVegetationEditor/Source/ProceduralVegetationEditor/Private/Widgets/SPVExportSelectionDialog.h
@@ -3,19 +3,13 @@
 #include "CoreMinimal.h"
 #include "Input/Reply.h"
 #include "Widgets/SWindow.h"
-#include "SPVExportSelectionDialog.generated.h"
 
 class UPVOutputSettings;
 class SScrollBox;
 class UProceduralVegetationGraph;
 class UPCGNode;
 class SPrimaryButton;
-UENUM()
-enum class EPVExportType
-{
-	Selection	UMETA(DisplayName = "Selected output node"),
-	BatchExport UMETA(DisplayName = "Batch export")
-};
+enum class EPVExportType;
 
 class SPVExportSelectionDialog : public SWindow
 {
diff --git a/Engine/Plugins/Experimental/ProceduralVegetationEditor/Source/ProceduralVegetationEditor/Public/PVEditorSettings.h b/Engine/Plugins/Experimental/ProceduralVegetationEditor/Source/ProceduralVegetationEditor/Public/PVEditorSettings.h
--- a/Engine/Plugins/Experimental/ProceduralVegetationEditor/Source/ProceduralVegetationEditor/Public/PVEditorSettings.h
+++ b/Engine/Plugins/Experimental/ProceduralVegetationEditor/Source/ProceduralVegetationEditor/Public/PVEditorSettings.h
@@ -6,6 +6,13 @@
 
 #include "PVEditorSettings.generated.h"
 
+UENUM()
+enum class EPVExportType
+{
+	Selection	UMETA(DisplayName = "Selected output node"),
+	BatchExport UMETA(DisplayName = "Batch export")
+};
+
 UCLASS(MinimalAPI, config=EditorPerProjectUserSettings, meta = (DisplayName = "Procedural Vegetation Editor Settings"))
 class UPVEditorSettings : public UObject
 {
@@ -36,5 +43,5 @@ public:
 	bool bAutoFocusViewport = true;
 
 	UPROPERTY(EditAnywhere, config, Category = Settings)
-	int bExportType = 0;
+	EPVExportType ExportType = EPVExportType::Selection;
 };
