[PVE] Remove branches node: switched generation mode to hierarchical generation.
Fixed issue: Angle Correction is not working as intended in gravity node

diff --git a/Engine/Plugins/Experimental/ProceduralVegetationEditor/Source/ProceduralVegetation/Private/Facades/PVBranchFacade.cpp b/Engine/Plugins/Experimental/ProceduralVegetationEditor/Source/ProceduralVegetation/Private/Facades/PVBranchFacade.cpp
--- a/Engine/Plugins/Experimental/ProceduralVegetationEditor/Source/ProceduralVegetation/Private/Facades/PVBranchFacade.cpp
+++ b/Engine/Plugins/Experimental/ProceduralVegetationEditor/Source/ProceduralVegetation/Private/Facades/PVBranchFacade.cpp
@@ -185,7 +185,7 @@ namespace PV::Facades
 		return INDEX_NONE;
 	}
 
-	int32 FBranchFacade::GetGenerationNumber(const int32 Index) const
+	int32 FBranchFacade::GetHierarchyGenerationNumber(const int32 Index) const
 	{
 		// The generation attributes contain incorrect information at the moment 
 		// thus the number of parents can be used to derive the correct number
diff --git a/Engine/Plugins/Experimental/ProceduralVegetationEditor/Source/ProceduralVegetation/Private/Implementations/PVGravity.cpp b/Engine/Plugins/Experimental/ProceduralVegetationEditor/Source/ProceduralVegetation/Private/Implementations/PVGravity.cpp
--- a/Engine/Plugins/Experimental/ProceduralVegetationEditor/Source/ProceduralVegetation/Private/Implementations/PVGravity.cpp
+++ b/Engine/Plugins/Experimental/ProceduralVegetationEditor/Source/ProceduralVegetation/Private/Implementations/PVGravity.cpp
@@ -90,18 +90,19 @@ void FPVGravity::ApplyGravity(
 	FVector3f PivotPosition
 )
 {
-	static const auto CalculateDownForce = [](const FVector3f InRelativeDirection, const int InPointIndex, const float InPScaleGradient,
+	const auto CalculateDownForce = [](const FVector3f InRelativeDirection, const int InPointIndex, const float InPScaleGradient,
 	                                          const FPVGravityParams& InGravityParams,
 	                                          const TArray<FVector3f>& InPhototropicDirections) -> FQuat4f
 		{
 			const FVector3f PhototropicDirection = InPhototropicDirections.Num() > InPointIndex
 				? InPhototropicDirections[InPointIndex]
 				: FVector3f::UpVector;
+			
 			const FVector3f GravityDirection = InGravityParams.Mode == EGravityMode::Phototropic
 				? PhototropicDirection
 				: InGravityParams.Direction;
 
-			float YDot = 1.0f - FMath::Abs(FVector3f::DotProduct(InRelativeDirection, InGravityParams.Direction));
+			float YDot = 1.0f - FMath::Abs(FVector3f::DotProduct(InRelativeDirection, GravityDirection));
 			YDot = FMath::GetMappedRangeValueClamped(FVector2f(0, 1),
 				FVector2f(FMath::Max(0.1f, 1 - InGravityParams.AngleCorrection), 1),
 				YDot);
diff --git a/Engine/Plugins/Experimental/ProceduralVegetationEditor/Source/ProceduralVegetation/Private/Implementations/PVMeshBuilder.cpp b/Engine/Plugins/Experimental/ProceduralVegetationEditor/Source/ProceduralVegetation/Private/Implementations/PVMeshBuilder.cpp
--- a/Engine/Plugins/Experimental/ProceduralVegetationEditor/Source/ProceduralVegetation/Private/Implementations/PVMeshBuilder.cpp
+++ b/Engine/Plugins/Experimental/ProceduralVegetationEditor/Source/ProceduralVegetation/Private/Implementations/PVMeshBuilder.cpp
@@ -973,7 +973,7 @@ void FPVMeshBuilder::GetLocalDynamicMeshData(FManagedArrayCollection& Collection
 						});
 
 				const bool PrimitiveIsTrunk = PlantFacade.IsTrunkIndex(PrimitiveIndex);
-				const int32 GenerationNumber = BranchFacade.GetGenerationNumber(PrimitiveIndex);
+				const int32 GenerationNumber = BranchFacade.GetHierarchyGenerationNumber(PrimitiveIndex);
 
 				MeshesData[PrimitiveIndex].MaterialID = BranchFacade.GetBranchUVMaterial(PrimitiveIndex);
 
diff --git a/Engine/Plugins/Experimental/ProceduralVegetationEditor/Source/ProceduralVegetation/Private/Implementations/PVRemoveBranches.cpp b/Engine/Plugins/Experimental/ProceduralVegetationEditor/Source/ProceduralVegetation/Private/Implementations/PVRemoveBranches.cpp
--- a/Engine/Plugins/Experimental/ProceduralVegetationEditor/Source/ProceduralVegetation/Private/Implementations/PVRemoveBranches.cpp
+++ b/Engine/Plugins/Experimental/ProceduralVegetationEditor/Source/ProceduralVegetation/Private/Implementations/PVRemoveBranches.cpp
@@ -38,7 +38,8 @@ namespace PV::Facades
 				continue;
 
 			const int32 RootPointIndex = TrunkPoints[0];
-
+			const int32 RootHierarchyGeneration =  InBranchFacade.GetHierarchyGenerationNumber(RootPointIndex);
+			
 			const TArray<int>& RootBudDevelopment = InPointFacade.GetBudDevelopment(RootPointIndex);
 
 			float MaxValue = std::numeric_limits<float>::lowest();
@@ -52,6 +53,11 @@ namespace PV::Facades
 					continue;
 				}
 
+				if (InPlantFacade.IsTrunkIndex(BranchIndex))
+				{
+					continue;
+				}
+
 				if (InBasis == ERemoveBranchesBasis::Length)
 				{
 					const float BranchRootPointLFR = InPointFacade.GetLengthFromRoot(BranchPoints[0]);
@@ -77,24 +83,20 @@ namespace PV::Facades
 						? PointBudLightDetected[BranchBudLightDetectedIndex]
 						: -1.0f;
 				}
-				else if (InBasis == ERemoveBranchesBasis::Age || InBasis == ERemoveBranchesBasis::Generation)
+				else if (InBasis == ERemoveBranchesBasis::Age)
 				{
 					const TArray<int>& PointBudDevelopment = InPointFacade.GetBudDevelopment(BranchPoints[BranchPoints.Num() - 1]);
 
-					if (InBasis == ERemoveBranchesBasis::Age)
-					{
-						const int RootBudAge = RootBudDevelopment[BudDevelopmentAgeIndex];
-						OutRemovalData[BranchIndex] = RootBudAge > 0
-							? static_cast<float>(PointBudDevelopment[BudDevelopmentAgeIndex]) / static_cast<float>(RootBudAge)
-							: -1.0f;
-					}
-					else if (InBasis == ERemoveBranchesBasis::Generation)
-					{
-						const int PointBudGeneration = PointBudDevelopment[BudDevelopmentGenerationIndex];
-						OutRemovalData[BranchIndex] = PointBudGeneration > 0
-							? static_cast<float>(RootBudDevelopment[BudDevelopmentGenerationIndex]) / static_cast<float>(PointBudGeneration)
-							: -1.0f;
-					}
+					const int RootBudAge = RootBudDevelopment[BudDevelopmentAgeIndex];
+					OutRemovalData[BranchIndex] = RootBudAge > 0
+						? static_cast<float>(PointBudDevelopment[BudDevelopmentAgeIndex]) / static_cast<float>(RootBudAge)
+						: -1.0f;
+				}
+				else if (InBasis == ERemoveBranchesBasis::Generation)
+				{
+					const int32 BranchHierarchyGeneration =  InBranchFacade.GetHierarchyGenerationNumber(BranchIndex);
+					
+					OutRemovalData[BranchIndex] = RootHierarchyGeneration / static_cast<float>(BranchHierarchyGeneration);
 				}
 
 				MaxValue = FMath::Max(OutRemovalData[BranchIndex], MaxValue);
diff --git a/Engine/Plugins/Experimental/ProceduralVegetationEditor/Source/ProceduralVegetation/Public/Facades/PVBranchFacade.h b/Engine/Plugins/Experimental/ProceduralVegetationEditor/Source/ProceduralVegetation/Public/Facades/PVBranchFacade.h
--- a/Engine/Plugins/Experimental/ProceduralVegetationEditor/Source/ProceduralVegetation/Public/Facades/PVBranchFacade.h
+++ b/Engine/Plugins/Experimental/ProceduralVegetationEditor/Source/ProceduralVegetation/Public/Facades/PVBranchFacade.h
@@ -49,7 +49,7 @@ namespace PV::Facades
 		
 		int32 GetParentBranchIndex(const int32 BranchIndex) const;
 
-		int32 GetGenerationNumber(const int32 Index) const;
+		int32 GetHierarchyGenerationNumber(const int32 Index) const;
 
 		const TManagedArray<int32>& GetBranchNumbers() const;
 
