[PVE] Fix issues with Foliage distributor node and a potential crash in Mesher node

diff --git a/Engine/Plugins/Experimental/ProceduralVegetationEditor/Source/ProceduralVegetation/Private/Facades/PVFoliageFacade.cpp b/Engine/Plugins/Experimental/ProceduralVegetationEditor/Source/ProceduralVegetation/Private/Facades/PVFoliageFacade.cpp
--- a/Engine/Plugins/Experimental/ProceduralVegetationEditor/Source/ProceduralVegetation/Private/Facades/PVFoliageFacade.cpp
+++ b/Engine/Plugins/Experimental/ProceduralVegetationEditor/Source/ProceduralVegetation/Private/Facades/PVFoliageFacade.cpp
@@ -364,11 +364,11 @@ namespace PV::Facades
 		FFoliageEntryData Data = GetFoliageEntry(Id);
 		FString FoliageName = GetFoliageName(Data.NameId);
 
-		const FQuat UpVectorRotationQuat = FQuat::FindBetweenNormals(FVector::UpVector, FVector(Data.UpVector));
+		const FQuat UpVectorRotationQuat = FQuat::FindBetweenNormals(FVector::UpVector, FVector(Data.UpVector).GetSafeNormal());
 		const FVector RotatedVectorFace = UpVectorRotationQuat.RotateVector(FVector::RightVector);
 		float Dot = FVector::DotProduct(RotatedVectorFace, FVector(Data.NormalVector));
 		float Angle = FMath::Acos(FMath::Clamp(Dot, -1.0f, 1.0f));
-		const FQuat FaceRotationQuat = FQuat(FVector(Data.UpVector), Angle);
+		const FQuat FaceRotationQuat = FQuat(FVector(Data.UpVector).GetSafeNormal(), Angle);
 
 		const FTransform ScaleTransform(FQuat::Identity, FVector::ZeroVector, FVector(Data.Scale));
 		const FTransform FaceRotationTransform(FaceRotationQuat);
diff --git a/Engine/Plugins/Experimental/ProceduralVegetationEditor/Source/ProceduralVegetation/Private/Implementations/PVFoliage.cpp b/Engine/Plugins/Experimental/ProceduralVegetationEditor/Source/ProceduralVegetation/Private/Implementations/PVFoliage.cpp
--- a/Engine/Plugins/Experimental/ProceduralVegetationEditor/Source/ProceduralVegetation/Private/Implementations/PVFoliage.cpp
+++ b/Engine/Plugins/Experimental/ProceduralVegetationEditor/Source/ProceduralVegetation/Private/Implementations/PVFoliage.cpp
@@ -19,7 +19,7 @@ void FPVFoliage::DistributeFoliage(FManagedArrayCollection& OutCollection, const
 
 	if (const int32 NumFoliageMeshes = FoliageFacadeOutput.NumFoliageNames(); NumFoliageMeshes == 0)
 	{
-		UE_LOG(LogProceduralVegetation, Warning, TEXT("There are no foliage meshes available in input for distribution"));	
+		UE_LOG(LogProceduralVegetation, Warning, TEXT("There are no foliage meshes available in input for distribution"));
 		return;
 	}
 
@@ -30,13 +30,13 @@ void FPVFoliage::DistributeFoliage(FManagedArrayCollection& OutCollection, const
 	PV::Facades::FPointFacade PointFacadeSource(SourceCollection);
 
 	TManagedArrayAccessor<TArray<FVector3f>> BudDirectionAttribute(OutCollection, PV::AttributeNames::BudDirection,
-	                                                               PV::GroupNames::PointGroup);
+		PV::GroupNames::PointGroup);
 	TManagedArrayAccessor<float> BranchGradientsAttribute(OutCollection, PV::AttributeNames::BranchGradient,
-	                                                      PV::GroupNames::PointGroup);
+		PV::GroupNames::PointGroup);
 	TManagedArrayAccessor<TArray<float>> LeafPhyllotaxyAttribute(OutCollection, PV::AttributeNames::LeafPhyllotaxy,
-	                                                             PV::GroupNames::DetailsGroup);
+		PV::GroupNames::DetailsGroup);
 	TManagedArrayAccessor<TArray<float>> BudHormoneLevels(OutCollection, PV::AttributeNames::BudHormoneLevels,
-	                                                      PV::GroupNames::PointGroup);
+		PV::GroupNames::PointGroup);
 
 	float PhyllotaxyFormationLeaf = LeafPhyllotaxyAttribute[0][1];
 	float MinBudsLeaf = LeafPhyllotaxyAttribute[0][3];
@@ -92,7 +92,7 @@ void FPVFoliage::DistributeFoliage(FManagedArrayCollection& OutCollection, const
 
 		// Compute attachment points in terms of normalized distances
 		OutCollection.EmptyGroup(PV::GroupNames::FoliageGroup);
-		for (int32 BranchIndex = 1; BranchIndex < NumBranches; ++BranchIndex)
+		for (int32 BranchIndex = 0; BranchIndex < NumBranches; ++BranchIndex)
 		{
 			TArray<float> NormalizedAttachmentPoints;
 			TArray<float> NormalizedAttachmentPointsFinal;
@@ -106,7 +106,9 @@ void FPVFoliage::DistributeFoliage(FManagedArrayCollection& OutCollection, const
 			const float LastPointLengthFromRoot = PointFacadeSource.GetLengthFromRoot(BranchPoints.Last());
 			const float BranchLength = LastPointLengthFromRoot - FirstPointLengthFromRoot;
 			float Increment = DistributionSettings.InstanceSpacing / BranchLength;
-			const float AdjustedMaxPerBranch = DistributionSettings.MaxPerBranch == -1 ? BranchPoints.Num() : DistributionSettings.MaxPerBranch;
+			const float AdjustedMaxPerBranch = DistributionSettings.MaxPerBranch == -1
+				? BranchPoints.Num()
+				: DistributionSettings.MaxPerBranch;
 			float LookUp = 0.0f;
 
 			for (int32 i = 0; i < AdjustedMaxPerBranch; ++i)
@@ -117,18 +119,16 @@ void FPVFoliage::DistributeFoliage(FManagedArrayCollection& OutCollection, const
 				}
 
 				float InstanceSpacingRampValue = DistributionSettings.InstanceSpacingRamp
-					                                 ? DistributionSettings.InstanceSpacingRamp->GetRichCurveConst()->Eval(1.0f - LookUp)
-					                                 : Increment;
+					? DistributionSettings.InstanceSpacingRamp->GetRichCurveConst()->Eval(LookUp)
+					: Increment;
 				// The increment is nudged slightly before blending with the weighted ramp value, the number is chosen based on observation
-				float AdjustedIncrement = Increment + (InstanceSpacingRampValue * (Increment * 10));
+				float AdjustedIncrement = Increment + (InstanceSpacingRampValue * (Increment * 10.0f));
 				AdjustedIncrement = FMath::Lerp(Increment, AdjustedIncrement, DistributionSettings.InstanceSpacingRampEffect);
 
 				NormalizedAttachmentPoints.Add(1.0f - LookUp);
 				LookUp += AdjustedIncrement;
 			}
 
-			Algo::Reverse(NormalizedAttachmentPoints);
-
 			// Compute Position, Angles, Scale and LengthFromRoot for attachment points
 
 			// Compute based on first point attributes
@@ -178,7 +178,7 @@ void FPVFoliage::DistributeFoliage(FManagedArrayCollection& OutCollection, const
 
 						const float EthyleneLevel = FMath::Lerp(EthyleneLevelAtCurrentPoint, EthyleneLevelAtNextPoint, DistanceAlpha);
 
-						if (EthyleneLevel >= DistributionSettings.EthyleneThreshold)
+						if ((EthyleneLevel - 0.001f) >= DistributionSettings.EthyleneThreshold)
 						{
 							continue;
 						}
@@ -199,8 +199,8 @@ void FPVFoliage::DistributeFoliage(FManagedArrayCollection& OutCollection, const
 						const float BranchGradient = FMath::Clamp(
 							FMath::Lerp(CurrentPointBranchGradient, NextPointBranchGradient, DistanceAlpha), 0.0f, 1.0f);
 						float BranchGradientRampValue = ScaleSettings.ScaleRamp
-							                                ? ScaleSettings.ScaleRamp->GetRichCurveConst()->Eval(BranchGradient)
-							                                : 1.0f;
+							? ScaleSettings.ScaleRamp->GetRichCurveConst()->Eval(BranchGradient)
+							: 1.0f;
 
 						const float CurrentPointScale = PointFacadeSource.GetPointScale(CurrentPointIndex);
 						const float NextPointScale = PointFacadeSource.GetPointScale(NextPointIndex);
@@ -224,13 +224,13 @@ void FPVFoliage::DistributeFoliage(FManagedArrayCollection& OutCollection, const
 						const FVector3f CurrentPointApicalDirectionVector(BudDirectionAttribute[CurrentPointIndex][0]);
 						const FVector3f NextPointApicalDirectionVector(BudDirectionAttribute[NextPointIndex][0]);
 						const FVector3f BudApicalDirectionVector = FMath::Lerp(CurrentPointApicalDirectionVector,
-						                                                       NextPointApicalDirectionVector, DistanceAlpha);
+							NextPointApicalDirectionVector, DistanceAlpha);
 
 						int32 BudsToAdd = 1;
 
 						// Compute Up Vector
 						// Last point
-						if (Idx == NormalizedAttachmentPoints.Num() - 1)
+						if (Idx == 0)
 						{
 							AttachmentPointsUpVectors.Add(BudApicalDirectionVector);
 						}
@@ -255,15 +255,15 @@ void FPVFoliage::DistributeFoliage(FManagedArrayCollection& OutCollection, const
 								if (VectorSettings.OverrideAxilAngle)
 								{
 									float AxilAngleRampValue = VectorSettings.AxilAngleRamp
-										                           ? FMath::Clamp(
-											                           VectorSettings.AxilAngleRamp->GetRichCurveConst()->Eval(BranchGradient),
-											                           0.0f,
-											                           1.0f)
-										                           : 1.0f;
+										? FMath::Clamp(
+											VectorSettings.AxilAngleRamp->GetRichCurveConst()->Eval(BranchGradient),
+											0.0f,
+											1.0f)
+										: 1.0f;
 
 									AxilAngleRampValue = AxilAngleRampValue * VectorSettings.AxilAngleRampUpperValue;
 									const float AxilAngleBlended = FMath::Lerp(VectorSettings.AxilAngle, AxilAngleRampValue,
-									                                           VectorSettings.AxilAngleRampEffect);
+										VectorSettings.AxilAngleRampEffect);
 
 									const FVector3f RotationAxis = FVector3f::CrossProduct(
 										BudApicalDirectionVector.GetSafeNormal(), AxillaryDirectionModified.GetSafeNormal());
@@ -297,10 +297,10 @@ void FPVFoliage::DistributeFoliage(FManagedArrayCollection& OutCollection, const
 						const FVector3f CurrentPointLightOptimalVector(BudDirectionAttribute[CurrentPointIndex][2]);
 						const FVector3f NextPointLightOptimalVector(BudDirectionAttribute[NextPointIndex][2]);
 						const FVector3f BudLightOptimalVector = FMath::Lerp(CurrentPointLightOptimalVector, NextPointLightOptimalVector,
-						                                                    DistanceAlpha);
+							DistanceAlpha);
 
 						FQuat4f QuatRotation = FQuat4f::FindBetweenNormals(BudApicalDirectionVector.GetSafeNormal(),
-						                                                   BudLightOptimalVector.GetSafeNormal());
+							BudLightOptimalVector.GetSafeNormal());
 						const FVector3f ApicalDirectionRotated = QuatRotation.RotateVector(BudApicalDirectionVector);
 
 						for (int32 i = 0; i < BudsToAdd; ++i)
@@ -313,21 +313,21 @@ void FPVFoliage::DistributeFoliage(FManagedArrayCollection& OutCollection, const
 
 			// Set foliage entries
 			TManagedArrayAccessor<FVector3f> PivotPointsAttribute(OutCollection, PV::AttributeNames::FoliagePivotPoint,
-			                                                      PV::GroupNames::FoliageGroup);
+				PV::GroupNames::FoliageGroup);
 			const int32 NumFoliageMeshes = FoliageFacadeOutput.NumFoliageNames();
 			int32 CurrentFoliageMeshIndex = BranchIndex % NumFoliageMeshes;
 			for (int32 Idx = 0; Idx < AttachmentPointsLengthsFromRoot.Num(); ++Idx)
 			{
 				const int32 NewIndex = PivotPointsAttribute.AddElements(1);
 				FoliageFacadeOutput.SetFoliageEntry(NewIndex, {
-					                                    .NameId = CurrentFoliageMeshIndex,
-					                                    .BranchId = BranchIndex,
-					                                    .PivotPoint = AttachmentPointsPositions[Idx],
-					                                    .UpVector = AttachmentPointsUpVectors[Idx],
-					                                    .NormalVector = AttachmentPointsNormalVectors[Idx],
-					                                    .Scale = AttachmentPointsScales[Idx],
-					                                    .LengthFromRoot = AttachmentPointsLengthsFromRoot[Idx]
-				                                    });
+					.NameId = CurrentFoliageMeshIndex,
+					.BranchId = BranchIndex,
+					.PivotPoint = AttachmentPointsPositions[Idx],
+					.UpVector = AttachmentPointsUpVectors[Idx],
+					.NormalVector = AttachmentPointsNormalVectors[Idx],
+					.Scale = AttachmentPointsScales[Idx],
+					.LengthFromRoot = AttachmentPointsLengthsFromRoot[Idx]
+				});
 				CurrentFoliageMeshIndex = (CurrentFoliageMeshIndex + 1) % NumFoliageMeshes;
 			}
 		}
@@ -351,9 +351,17 @@ void FPVFoliage::DistributeFoliage(FManagedArrayCollection& OutCollection, const
 	{
 		TRACE_CPUPROFILER_EVENT_SCOPE(PVFoliageDistributorNode::EthyleneRemove);
 
+		// No need to apply Ethylene threshold if the distribution was overriden since that's already been accounted for
+		if (DistributionSettings.OverrideDistribution)
+		{
+			return;
+		}
+
 		// Compute Foliage instances to remove based on Ethylene threshold
-		TManagedArrayAccessor<FVector3f> PivotPointsAttribute(OutCollection, PV::AttributeNames::FoliagePivotPoint,
-		                                                      PV::GroupNames::FoliageGroup);
+		TManagedArrayAccessor<FVector3f> PivotPointsAttribute(
+			OutCollection,
+			PV::AttributeNames::FoliagePivotPoint,
+			PV::GroupNames::FoliageGroup);
 		TSet<int32> FoliageEntriesToRemove;
 
 		for (int32 BranchIndex = 0; BranchIndex < NumBranches; ++BranchIndex)
@@ -380,7 +388,7 @@ void FPVFoliage::DistributeFoliage(FManagedArrayCollection& OutCollection, const
 							FVector2f(EthyleneLevelAtNextPoint, EthyleneLevelAtCurrentPoint),
 							FoliageLengthFromRoot);
 
-						if (EthyleneLevel >= DistributionSettings.EthyleneThreshold)
+						if ((EthyleneLevel - 0.001f) >= DistributionSettings.EthyleneThreshold)
 						{
 							FoliageIDsToRemove.Add(FoliageIDs[FoliageIDsIndex]);
 						}
diff --git a/Engine/Plugins/Experimental/ProceduralVegetationEditor/Source/ProceduralVegetation/Private/Implementations/PVMeshBuilder.h b/Engine/Plugins/Experimental/ProceduralVegetationEditor/Source/ProceduralVegetation/Private/Implementations/PVMeshBuilder.h
--- a/Engine/Plugins/Experimental/ProceduralVegetationEditor/Source/ProceduralVegetation/Private/Implementations/PVMeshBuilder.h
+++ b/Engine/Plugins/Experimental/ProceduralVegetationEditor/Source/ProceduralVegetation/Private/Implementations/PVMeshBuilder.h
@@ -64,7 +64,7 @@ struct FPVMeshBuilderParams
 	int32 MaxMeshDivisions = 12;
 
 	UPROPERTY(EditAnywhere, Category="Mesh Settings|Advanced", DisplayName="Reduction Accuracy",
-		meta=(ClampMax=15, UIMin=1, UIMax=15, Tooltip=
+		meta=(ClampMin=1, ClampMax=15, UIMin=1, UIMax=15, Tooltip=
 			"Accuracy level for mesh reduction.\n\nControls how closely the simplified mesh follows the original. It specifies number of iterations that are used to evaluate the curve reduction path. The more number of iterations the more accurate the reduction will be at the cost of processing power."
 		))
 	int32 Accuracy = 5;
@@ -100,7 +100,7 @@ struct FPVMeshBuilderParams
 	float MainTrunkRetention = 0.0f;
 
 	UPROPERTY(EditAnywhere, Category="Mesh Settings|Advanced", DisplayName="Ground Retention",
-		meta=(UIMin=0.0f, UIMax=1.0f, Tooltip=
+		meta=(ClampMin=0.0f, ClampMax=1.0f, UIMin=0.0f, UIMax=1.0f, Tooltip=
 			"Preserves detail near ground regions.\n\nAdjusts retention for areas close to the base of the plant. Use higher values to keep root/ground transitions detailed."
 		))
 	float GroundRetention = 0.0f;
