Rotate foliage entry normal vector in gravity node

diff --git a/Engine/Plugins/Experimental/ProceduralVegetationEditor/Source/ProceduralVegetation/Private/Implementations/PVGravity.cpp b/Engine/Plugins/Experimental/ProceduralVegetationEditor/Source/ProceduralVegetation/Private/Implementations/PVGravity.cpp
--- a/Engine/Plugins/Experimental/ProceduralVegetationEditor/Source/ProceduralVegetation/Private/Implementations/PVGravity.cpp
+++ b/Engine/Plugins/Experimental/ProceduralVegetationEditor/Source/ProceduralVegetation/Private/Implementations/PVGravity.cpp
@@ -197,6 +197,7 @@ void FPVGravity::ApplyGravity(
 				// same as point positions and bud directions respectively
 				Entry.PivotPoint = PivotPosition + DownForce * RelativeDirection;
 				Entry.UpVector = TotalDownForce * Entry.UpVector;
+				Entry.NormalVector = TotalDownForce * Entry.NormalVector;
 				FoliageFacade.SetFoliageEntry(FoliageEntryId, Entry);
 			}
 		}
