Fix BudDirectionsLightOptimalIndex & BudDirectionsLightSubOptimal being updated by the gravity node

diff --git a/Engine/Plugins/Experimental/ProceduralVegetationEditor/Source/ProceduralVegetation/Private/Facades/PVBudVectorsFacade.h b/Engine/Plugins/Experimental/ProceduralVegetationEditor/Source/ProceduralVegetation/Private/Facades/PVBudVectorsFacade.h
--- a/Engine/Plugins/Experimental/ProceduralVegetationEditor/Source/ProceduralVegetation/Private/Facades/PVBudVectorsFacade.h
+++ b/Engine/Plugins/Experimental/ProceduralVegetationEditor/Source/ProceduralVegetation/Private/Facades/PVBudVectorsFacade.h
@@ -7,6 +7,9 @@
 
 namespace PV::Facades
 {
+	constexpr int BudDirectionsLightOptimalIndex = 2;
+	constexpr int BudDirectionsLightSubOptimal = 3;
+
 	/**
 	 * FBudVectorsFacade is used to access and manipulate the BudVectors data from Point Group with in the ProceduralVegetation's FManagedArrayCollection
 	 * BudVectors represent different growth related vectors in plant data and are associated with each point. 
diff --git a/Engine/Plugins/Experimental/ProceduralVegetationEditor/Source/ProceduralVegetation/Private/Implementations/PVGravity.cpp b/Engine/Plugins/Experimental/ProceduralVegetationEditor/Source/ProceduralVegetation/Private/Implementations/PVGravity.cpp
--- a/Engine/Plugins/Experimental/ProceduralVegetationEditor/Source/ProceduralVegetation/Private/Implementations/PVGravity.cpp
+++ b/Engine/Plugins/Experimental/ProceduralVegetationEditor/Source/ProceduralVegetation/Private/Implementations/PVGravity.cpp
@@ -14,9 +14,6 @@ void FPVGravity::GeneratePhototropicData(const FPVGravityParams& InGravityParams
 
 	OutPhototropicDirections.Init(PrevVec, InPointFacade.GetElementCount());
 
-	static constexpr int BudDirectionsLightOptimalIndex = 2;
-	static constexpr int BudDirectionsLightSubOptimal = 3;
-
 	for (int BranchIndex = 0; BranchIndex < InBranchFacade.GetElementCount(); BranchIndex++)
 	{
 		const TArray<int32>& BranchPoints = InBranchFacade.GetPoints(BranchIndex);
@@ -26,10 +23,10 @@ void FPVGravity::GeneratePhototropicData(const FPVGravityParams& InGravityParams
 			const auto CurrentPointIndex = BranchPoints[i];
 			const TArray<FVector3f>& BudDirections = InPointFacade.GetBudDirection(CurrentPointIndex);
 
-			if (BudDirections.Num() > BudDirectionsLightSubOptimal)
+			if (BudDirections.Num() > PV::Facades::BudDirectionsLightSubOptimal)
 			{
-				const FVector3f& LightOptimal = BudDirections[BudDirectionsLightOptimalIndex];
-				const FVector3f& LightSubOptimal = BudDirections[BudDirectionsLightSubOptimal];
+				const FVector3f& LightOptimal = BudDirections[PV::Facades::BudDirectionsLightOptimalIndex];
+				const FVector3f& LightSubOptimal = BudDirections[PV::Facades::BudDirectionsLightSubOptimal];
 
 				FVector3f PhototropicDirection = FMath::Lerp(LightOptimal, LightSubOptimal, InGravityParams.PhototropicBias);
 				PhototropicDirection.Normalize();
@@ -154,9 +151,14 @@ void FPVGravity::ApplyGravity(
 		PreviousPosition = Position;
 		Position = PivotPosition + DownForce * RelativeDirection;
 
-		for (FVector3f& BudDirection : CurrentBudDirections)
+		for (int32 i = 0; i < CurrentBudDirections.Num(); ++i)
 		{
-			BudDirection = TotalDownForce * BudDirection;
+			if (i == PV::Facades::BudDirectionsLightOptimalIndex || i == PV::Facades::BudDirectionsLightSubOptimal)
+			{
+				continue;
+			}
+				
+			CurrentBudDirections[i] = TotalDownForce * CurrentBudDirections[i];
 		}
 	}
 	else
@@ -230,9 +232,14 @@ void FPVGravity::ApplyGravity(
 		PreviousPosition = Position;
 		Position = PivotPosition + DownForce * RelativeDirection;
 
-		for (FVector3f& BudDirection : CurrentBudDirections)
+		for (int32 i = 0; i < CurrentBudDirections.Num(); ++i)
 		{
-			BudDirection = TotalDownForce * BudDirection;
+			if (i == PV::Facades::BudDirectionsLightOptimalIndex || i == PV::Facades::BudDirectionsLightSubOptimal)
+			{
+				continue;
+			}
+
+			CurrentBudDirections[i] = TotalDownForce * CurrentBudDirections[i];
 		}
 	}
 }
