[PVE] Logs added

diff --git a/Engine/Plugins/Experimental/ProceduralVegetationEditor/Source/ProceduralVegetation/Private/Helpers/PVAnalyticsHelper.cpp b/Engine/Plugins/Experimental/ProceduralVegetationEditor/Source/ProceduralVegetation/Private/Helpers/PVAnalyticsHelper.cpp
--- a/Engine/Plugins/Experimental/ProceduralVegetationEditor/Source/ProceduralVegetation/Private/Helpers/PVAnalyticsHelper.cpp
+++ b/Engine/Plugins/Experimental/ProceduralVegetationEditor/Source/ProceduralVegetation/Private/Helpers/PVAnalyticsHelper.cpp
@@ -3,6 +3,7 @@
 #include "Helpers/PVAnalyticsHelper.h"
 
 #include "EngineAnalytics.h"
+#include "ProceduralVegetationModule.h"
 #include "PVExportParams.h"
 #include "Engine/RendererSettings.h"
 #include "Facades/PVFoliageFacade.h"
@@ -13,6 +14,17 @@
 
 namespace PV::Analytics
 {
+	void LogEvent(const FString& InEventName, const TArray<FAnalyticsEventAttribute>& InAttributes = TArray<FAnalyticsEventAttribute>())
+	{
+		FString AttributesData;
+		for (int32 i = 0; i < InAttributes.Num(); i++)
+		{
+			AttributesData += FString::Printf(TEXT("(%s,%s),"), *InAttributes[i].GetName(), *InAttributes[i].GetValue());
+		}
+			
+		UE_LOG(LogProceduralVegetation, Display, TEXT("Event: %s : Attributes: {%s}"), *InEventName, *AttributesData);
+	}
+	
 	void SendEvent(const FString& InEventName, const TArray<FAnalyticsEventAttribute>& InAttributes = TArray<FAnalyticsEventAttribute>())
 	{
 		if (FEngineAnalytics::IsAvailable())
@@ -48,7 +60,8 @@ namespace PV::Analytics
 			TArray<FAnalyticsEventAttribute> Attributes;
 			Attributes.Add(FAnalyticsEventAttribute(TEXT("Name"),  InNodeType));
 		
-			SendEvent(TEXT("NodeAdded"), Attributes);	
+			SendEvent(TEXT("NodeAdded"), Attributes);
+			LogEvent(TEXT("NodeAdded"), Attributes);
 		}
 	}
 
@@ -69,6 +82,7 @@ namespace PV::Analytics
 			Attributes.Add(FAnalyticsEventAttribute(TEXT("PropertyName"),  PropertyName));
 		
 			SendEvent(TEXT("NodeTweaked"), Attributes);
+			LogEvent(TEXT("NodeTweaked"), Attributes);
 		}
 	}
 
@@ -80,6 +94,7 @@ namespace PV::Analytics
 			Attributes.Add(FAnalyticsEventAttribute(TEXT("Path"),  InMatPath));
 			
 			SendEvent(TEXT("MeshMaterialChanged"), Attributes);
+			LogEvent(TEXT("MeshMaterialChanged"), Attributes);
 		}
 	}
 
@@ -91,6 +106,7 @@ namespace PV::Analytics
 			Attributes.Add(FAnalyticsEventAttribute(TEXT("Path"),  InFoliageMeshPath));
 			
 			SendEvent(TEXT("FoliageMeshChanged"), Attributes);
+			LogEvent(TEXT("FoliageMeshChanged"), Attributes);
 		}
 	}
 
@@ -102,6 +118,7 @@ namespace PV::Analytics
 			Attributes.Add(FAnalyticsEventAttribute(TEXT("Path"),  InPath));
 			
 			SendEvent(TEXT("WindPresetChanged"), Attributes);
+			LogEvent(TEXT("WindPresetChanged"), Attributes);
 		}
 	}
 
@@ -160,6 +177,7 @@ namespace PV::Analytics
 		if (FEngineAnalytics::IsAvailable())
 		{
 			SendEvent(TEXT("ExportStarted"), InCommonAttributes);
+			LogEvent(TEXT("ExportStarted"), InCommonAttributes);
 		}
 	}
 
@@ -186,6 +204,7 @@ namespace PV::Analytics
 			InCommonAttributes.Add(FAnalyticsEventAttribute(TEXT("ExportResult"), ExportResultEnumToString(ExportResult)));
 
 			SendEvent(TEXT("ExportFinished"), InCommonAttributes);
+			LogEvent(TEXT("ExportFinished"), InCommonAttributes);
 		}
 	}
 }
diff --git a/Engine/Plugins/Experimental/ProceduralVegetationEditor/Source/ProceduralVegetation/Private/ProceduralVegetationPreset.cpp b/Engine/Plugins/Experimental/ProceduralVegetationEditor/Source/ProceduralVegetation/Private/ProceduralVegetationPreset.cpp
--- a/Engine/Plugins/Experimental/ProceduralVegetationEditor/Source/ProceduralVegetation/Private/ProceduralVegetationPreset.cpp
+++ b/Engine/Plugins/Experimental/ProceduralVegetationEditor/Source/ProceduralVegetation/Private/ProceduralVegetationPreset.cpp
@@ -54,6 +54,7 @@ void UProceduralVegetationPreset::LoadFromVariantFiles(const TMap<FString, FStri
 {
 	for (const auto& [VariantName, VariantFile] : InVariantFiles)
 	{
+		UE_LOG(LogProceduralVegetation, Log, TEXT("Loaded variant : %s"), *VariantName);
 		FManagedArrayCollection Collection;
 		if (FString OutError; PV::LoadMegaPlantsJsonToCollection(Collection, VariantFile, OutError))
 		{
@@ -64,6 +65,7 @@ void UProceduralVegetationPreset::LoadFromVariantFiles(const TMap<FString, FStri
 
 void UProceduralVegetationPreset::UpdateDataAsset()
 {
+	UE_LOG(LogProceduralVegetation, Log, TEXT("UpdateDataAsset Clicked"));
 	if (JsonDirectoryPath.Path.IsEmpty())
 	{
 		UE_LOG(LogProceduralVegetation, Warning, TEXT("Please specify a JSON path"));
@@ -85,6 +87,7 @@ void UProceduralVegetationPreset::UpdateDataAsset()
 		FString OutError;
 		if ( PV::LoadMegaPlantsJsonToCollection(Collection, FullPath, OutError))
 		{
+			UE_LOG(LogProceduralVegetation, Log, TEXT("Variant %s loaded from file"), *FPaths::GetBaseFilename(FullPath));
 			Variants.Add(FPaths::GetBaseFilename(FullPath), Collection);
 		}
 		else if (TSharedPtr<FJsonObject> MetaJson = PV::LoadMetaFileIntoJsonObject(FullPath, OutError))
diff --git a/Engine/Plugins/Experimental/ProceduralVegetationEditor/Source/ProceduralVegetationEditor/Private/PVEditor.cpp b/Engine/Plugins/Experimental/ProceduralVegetationEditor/Source/ProceduralVegetationEditor/Private/PVEditor.cpp
--- a/Engine/Plugins/Experimental/ProceduralVegetationEditor/Source/ProceduralVegetationEditor/Private/PVEditor.cpp
+++ b/Engine/Plugins/Experimental/ProceduralVegetationEditor/Source/ProceduralVegetationEditor/Private/PVEditor.cpp
@@ -14,6 +14,7 @@
 #include "ToolMenus.h"
 #include "ToolMenuSection.h"
 #include "PackagesDialog.h"
+#include "ProceduralVegetationEditorModule.h"
 
 #include "Dataflow/DataflowCollectionSpreadSheetWidget.h"
 
@@ -56,6 +57,7 @@ void FPVEditor::Initialize(const EToolkitMode::Type InMode, const TSharedPtr<cla
 	
 	check(InProceduralVegetation);
 
+	UE_LOG(LogProceduralVegetationEditor, Log, TEXT("PVEditor initialized for %s"), *InProceduralVegetation->GetFName().ToString());
 	ProceduralVegetationBeingEdited = InProceduralVegetation;
 	
 	UProceduralVegetationGraph* ProceduralVegetationGraph = InProceduralVegetation->GetGraph();
@@ -161,6 +163,8 @@ FString FPVEditor::GetWorldCentricTabPrefix() const
 
 void FPVEditor::OnClose()
 {
+	UE_LOG(LogProceduralVegetationEditor, Log, TEXT("PVEditor Closed"));
+	
 	double TotalSeconds = (FDateTime::Now() - SessionStartTime).GetTotalSeconds();
 	PV::Analytics::SendSessionEndedEvent(TotalSeconds);
 	
@@ -367,6 +371,7 @@ FReply FPVEditor::ExportInternal()
 		return FReply::Handled();
 	}
 
+	UE_LOG(LogProceduralVegetationEditor, Log, TEXT("Export Started"));
 	struct FScopedMessages
 	{
 		TArray<TSharedRef<FTokenizedMessage>> Messages;
@@ -595,6 +600,7 @@ FReply FPVEditor::ExportInternal()
 	
 	if (EditorSettings->ExportType == EPVExportType::Selection)
 	{
+		UE_LOG(LogProceduralVegetationEditor, Log, TEXT("Exporting the selected nodes"));
 		if (SelectedNode != nullptr && SelectedNode->GetPCGNode())
 		{
 			if (const UPVOutputSettings* OutputSettings = Cast<UPVOutputSettings>(SelectedNode->GetPCGNode()->GetSettings()))
@@ -605,6 +611,7 @@ FReply FPVEditor::ExportInternal()
 	}
 	else
 	{
+		UE_LOG(LogProceduralVegetationEditor, Log, TEXT("Batch Exporting nodes"));
 		NodesToExport = ProceduralVegetationGraph->GetNodes().FilterByPredicate([](UPCGNode* InNode)->bool {
 			const UPVOutputSettings* OutputSettings = Cast<const UPVOutputSettings>(InNode->GetSettings());
 			return OutputSettings && OutputSettings->ExportSettings.bShouldExport;
diff --git a/Engine/Plugins/Experimental/ProceduralVegetationEditor/Source/ProceduralVegetationEditor/Private/ProceduralVegetationFactory.cpp b/Engine/Plugins/Experimental/ProceduralVegetationEditor/Source/ProceduralVegetationEditor/Private/ProceduralVegetationFactory.cpp
--- a/Engine/Plugins/Experimental/ProceduralVegetationEditor/Source/ProceduralVegetationEditor/Private/ProceduralVegetationFactory.cpp
+++ b/Engine/Plugins/Experimental/ProceduralVegetationEditor/Source/ProceduralVegetationEditor/Private/ProceduralVegetationFactory.cpp
@@ -3,6 +3,7 @@
 #include "ProceduralVegetationFactory.h"
 
 #include "ProceduralVegetation.h"
+#include "ProceduralVegetationEditorModule.h"
 
 #include "Widgets/SPVCreateDialog.h"
 
@@ -26,11 +27,13 @@ UObject* UProceduralVegetationFactory::FactoryCreateNew(UClass* InClass, UObject
 
 		TObjectPtr<UProceduralVegetationGraph> SampleGraph = SampleProceduralVegetation->GetGraph();
 		NewProceduralVegetation->CreateGraph(SampleGraph);
+		UE_LOG(LogProceduralVegetationEditor, Log, TEXT("New Procedural Vegetation Asset created %s from Preset %s"), *NewProceduralVegetation->GetFName().ToString(), *SampleProceduralVegetation.GetFName().ToString());
 	}
 	else
 	{
 		NewProceduralVegetation = NewObject<UProceduralVegetation>(InParent, InClass, InName, InFlags, InContext);
 		NewProceduralVegetation->CreateGraph();
+		UE_LOG(LogProceduralVegetationEditor, Log, TEXT("New Procedural Vegetation Asset created %s"), *NewProceduralVegetation->GetFName().ToString());
 	}
 		
 	return NewProceduralVegetation;
