[PVE] Displacement issues fix
Added Validations for Dislacement Texture
Fixed profile uv issue

diff --git a/Engine/Plugins/Experimental/ProceduralVegetationEditor/Source/ProceduralVegetation/Private/Implementations/PVMeshBuilder.cpp b/Engine/Plugins/Experimental/ProceduralVegetationEditor/Source/ProceduralVegetation/Private/Implementations/PVMeshBuilder.cpp
--- a/Engine/Plugins/Experimental/ProceduralVegetationEditor/Source/ProceduralVegetation/Private/Implementations/PVMeshBuilder.cpp
+++ b/Engine/Plugins/Experimental/ProceduralVegetationEditor/Source/ProceduralVegetation/Private/Implementations/PVMeshBuilder.cpp
@@ -6,6 +6,7 @@
 
 #include "ProceduralVegetationModule.h"
 #include "Algo/MaxElement.h"
+#include "Elements/PCGPrintElement.h"
 #include "Facades/PVBoneFacade.h"
 #include "Facades/PVBranchFacade.h"
 #include "Facades/PVBudVectorsFacade.h"
@@ -660,7 +661,7 @@ void FPVMeshBuilder::GenerateBranchMeshData(const bool bPrimitiveIsTrunk, const
 	{
 		return;
 	}
-
+	
 	const PV::Facades::FPointFacade PointFacade(Collection);
 	const PV::Facades::FBudVectorsFacade BudVectorFacade(Collection);
 	const PV::Facades::FPlantProfileFacade PlantProfileFacade(Collection);
@@ -703,6 +704,9 @@ void FPVMeshBuilder::GenerateBranchMeshData(const bool bPrimitiveIsTrunk, const
 	FVector3f PreviousUpVector = FVector3f::ForwardVector;
 	FVector3f PreviousAimVector = FVector3f::UpVector;
 
+	float Displacement_UV_V = 0.0f;
+	FVector3f PreviousPosition = PointPositions[PrimitivePoints[0]];
+	
 	for (int32 Index = 0; Index < PrimitivePoints.Num(); ++Index)
 	{
 		const int32& PointIndex = PrimitivePoints[Index];
@@ -723,6 +727,10 @@ void FPVMeshBuilder::GenerateBranchMeshData(const bool bPrimitiveIsTrunk, const
 		// The AimVector needs to be adjusted in case curvature has changed during mesh simplification
 		const FVector3f& PointPosition = PointPositions[PointIndex];
 		FVector3f UpdatedAimVector;
+		
+		const float PointScale = PointScales[PointIndex];
+		Displacement_UV_V += FVector3f::Dist(PointPosition, PreviousPosition) * MaxScaleRatio * (MaxPointScale / PointScale);
+		
 		if (Index == PrimitivePoints.Num() - 1)
 		[[unlikely]]
 		{
@@ -765,8 +773,7 @@ void FPVMeshBuilder::GenerateBranchMeshData(const bool bPrimitiveIsTrunk, const
 
 		const float PlantGradient = PointFacade.GetPlantGradient(PointIndex);
 		const float ProfileFallOff = MeshBuilderParams.PlantProfileFallOff.GetRichCurveConst()->Eval(PlantGradient);
-		if (const float PointScale = PointScales[PointIndex];
-			PointScale == 0.0f)
+		if (PointScale == 0.0f)
 		[[unlikely]]
 		{
 			CurrentIndices.Add(OutLocalMeshData.Vertices.Emplace(PointPosition, UpdatedAimVector, FVector2f(0.0f, TextureCoordV), NjordPixelId));
@@ -781,10 +788,12 @@ void FPVMeshBuilder::GenerateBranchMeshData(const bool bPrimitiveIsTrunk, const
 				float UV_U = (1 - (static_cast<float>(Division) / PointDivisions)) + TextureCoordUOffset;
 				UV_U = FMath::GetMappedRangeValueUnclamped(OldRange, NewRange, FMath::Wrap(UV_U, OldRange.X, OldRange.Y));
 
+				float Displacement_UV_U = (static_cast<float>(Division) / PointDivisions);
+					
 				float ProfileMultiplier = 1.0f;
 				if (bShouldApplyProfile)
 				{
-					ProfileMultiplier = GetProfileMultiplier(ProfilePoints, UV_U);
+					ProfileMultiplier = GetProfileMultiplier(ProfilePoints, Displacement_UV_U);
 					ProfileMultiplier = FMath::Lerp(1.0f, ProfileMultiplier, ProfileFallOff);
 					ProfileMultiplier = ProfileMultiplier * MeshBuilderParams.PlantProfileScale;
 				}
@@ -796,13 +805,13 @@ void FPVMeshBuilder::GenerateBranchMeshData(const bool bPrimitiveIsTrunk, const
 				{
 					const int32 X = FMath::Clamp(
 						FMath::FloorToInt(
-							UV_U * MeshBuilderParams.DisplacementUVScale.X * (DisplacementData.TextureWidth - 1)
+							Displacement_UV_U * MeshBuilderParams.DisplacementUVScale.X * (DisplacementData.TextureWidth - 1)
 						) % DisplacementData.TextureWidth,
 						0, DisplacementData.TextureWidth - 1
 					);
 					const int32 Y = FMath::Clamp(
 						FMath::FloorToInt(
-							TextureCoordV * MeshBuilderParams.DisplacementUVScale.Y * (DisplacementData.TextureHeight - 1)
+							Displacement_UV_V * MeshBuilderParams.DisplacementUVScale.Y * (DisplacementData.TextureHeight - 1)
 						) % DisplacementData.TextureHeight,
 						0, DisplacementData.TextureHeight - 1
 					);
@@ -827,6 +836,8 @@ void FPVMeshBuilder::GenerateBranchMeshData(const bool bPrimitiveIsTrunk, const
 				TextureCoordV),
 			NjordPixelId));
 
+		PreviousPosition = PointPosition;
+		
 		TriangulateRings(PreviousIndices, CurrentIndices, PolyGroupIndex, OutLocalMeshData);
 
 		PreviousIndices = MoveTemp(CurrentIndices);
@@ -907,6 +918,15 @@ void FPVMeshBuilder::UpdateFoliagePivotPoints(const TSet<int32>& PointsToRemove,
 	}
 }
 
+bool IsDisplacementTextureValid(const TObjectPtr<UTexture2D>& Texture)
+{
+#if WITH_EDITOR
+	return Texture && Texture->Source.IsValid();
+#else
+	return false;
+#endif
+}
+
 void FPVMeshBuilder::GetLocalDynamicMeshData(FManagedArrayCollection& Collection, const FPVMeshBuilderParams& MeshBuilderParams,
                                              TArray<FLocalDynamicMeshData>& MeshesData)
 {
@@ -947,12 +967,14 @@ void FPVMeshBuilder::GetLocalDynamicMeshData(FManagedArrayCollection& Collection
 	int32 DisplacementWidth = 0;
 	int32 DisplacementHeight = 0;
 	if (MeshBuilderParams.DisplacementTexture
+		&& IsDisplacementTextureValid(MeshBuilderParams.DisplacementTexture)
 		&& MeshBuilderParams.DisplacementValues.Num() > 0
 		&& MeshBuilderParams.DisplacementStrength > 0.0f)
 	{
-		DisplacementWidth = MeshBuilderParams.DisplacementTexture->GetSizeX();
-		DisplacementHeight = MeshBuilderParams.DisplacementTexture->GetSizeY();
-
+#if WITH_EDITOR
+		DisplacementWidth = MeshBuilderParams.DisplacementTexture->Source.GetSizeX();
+		DisplacementHeight = MeshBuilderParams.DisplacementTexture->Source.GetSizeY();
+#endif
 		if (MeshBuilderParams.DisplacementTexture->SRGB)
 		{
 			UE_LOG(LogProceduralVegetation, Warning,
@@ -982,22 +1004,31 @@ void FPVMeshBuilder::GetLocalDynamicMeshData(FManagedArrayCollection& Collection
 			});
 }
 
-void FPVMeshBuilder::ExtractDisplacementData(const TObjectPtr<UTexture2D>& Texture, TArray<float>& OutValues)
+bool FPVMeshBuilder::ExtractDisplacementData(const TObjectPtr<UTexture2D>& Texture, TArray<float>& OutValues, FString& OutError)
 {
 	OutValues.Empty();
 
 #if WITH_EDITOR
 	if (!Texture)
 	{
-		UE_LOG(LogProceduralVegetation, Warning, TEXT("ExtractDisplacementData: Texture is null."));
-		return;
+		OutError = "ExtractDisplacementData: Texture is null.";
+		UE_LOG(LogProceduralVegetation, Warning, TEXT("%s"), *OutError);
+		return false;
 	}
 
 	FTextureSource& Source = Texture->Source;
 	if (!Source.IsValid())
 	{
-		UE_LOG(LogProceduralVegetation, Warning, TEXT("ExtractDisplacementData: No source data found on texture %s."), *Texture->GetName());
-		return;
+		OutError = FString::Printf(TEXT("ExtractDisplacementData: No source data found on texture %s."), *Texture->GetName());
+		UE_LOG(LogProceduralVegetation, Warning, TEXT("%s"), *OutError);
+		return false;
+	}
+
+	if (!Source.AreAllBlocksPowerOfTwo())
+	{
+		OutError = FString::Printf(TEXT("ExtractDisplacementData: Texture %s is not power of two, Only Power of 2 textures are supported."), *Texture->GetName());
+		UE_LOG(LogProceduralVegetation, Warning, TEXT("%s"), *OutError);
+		return false;
 	}
 
 	const int32 Width = Source.GetSizeX();
@@ -1014,17 +1045,34 @@ void FPVMeshBuilder::ExtractDisplacementData(const TObjectPtr<UTexture2D>& Textu
 			const float* FloatData = reinterpret_cast<const float*>(RawData);
 			for (int32 i = 0; i < NumPix; ++i)
 			{
-				OutValues[i] = FloatData[i * 4 + 0];
+				OutValues[i] = FloatData[i * 4];
+			}
+		}
+		break;
+	case TSF_R32F:
+		{
+			const float* FloatData = reinterpret_cast<const float*>(RawData);
+			for (int32 i = 0; i < NumPix; ++i)
+			{
+				OutValues[i] = FloatData[i];
 			}
 		}
 		break;
-
 	case TSF_RGBA16F:
 		{
 			const FFloat16* HalfData = reinterpret_cast<const FFloat16*>(RawData);
 			for (int32 i = 0; i < NumPix; ++i)
 			{
-				OutValues[i] = HalfData[i * 4 + 0].GetFloat();
+				OutValues[i] = HalfData[i * 4].GetFloat();
+			}
+		}
+		break;
+	case TSF_R16F:
+		{
+			const FFloat16* HalfData = reinterpret_cast<const FFloat16*>(RawData);
+			for (int32 i = 0; i < NumPix; ++i)
+			{
+				OutValues[i] = HalfData[i].GetFloat();
 			}
 		}
 		break;
@@ -1038,16 +1086,29 @@ void FPVMeshBuilder::ExtractDisplacementData(const TObjectPtr<UTexture2D>& Textu
 			}
 		}
 		break;
-
+		case TSF_G8:
+		{
+			const uint8* ColorData = reinterpret_cast<const uint8*>(RawData);
+			constexpr float Inv255 = 1.0f / 255.0f;
+			for (int32 i = 0; i < NumPix; ++i)
+			{
+				OutValues[i] = *ColorData * Inv255;
+				ColorData++;
+			}
+		}
+		break;
 	default:
-		UE_LOG(LogProceduralVegetation, Warning, TEXT("ExtractDisplacementData: Unsupported source format %d on texture %s."),
-			(int32)Source.GetFormat(), *Texture->GetName());
+		OutError = FString::Printf(TEXT("ExtractDisplacementData: Unsupported source format %d on texture %s."), Source.GetFormat(), *Texture->GetName());
+		UE_LOG(LogProceduralVegetation, Warning, TEXT("%s"),*OutError);
 		break;
 	}
 
 
 	Source.UnlockMip(0);
+	return true;
 #else
-	UE_LOG(LogProceduralVegetation, Warning, TEXT("ExtractDisplacementData only works in editor builds (WITH_EDITOR)."));
+	OutError = "ExtractDisplacementData only works in editor builds (WITH_EDITOR).";
+	UE_LOG(LogProceduralVegetation, Warning, TEXT("%s"), *OutError);
+	return false;
 #endif
 }
diff --git a/Engine/Plugins/Experimental/ProceduralVegetationEditor/Source/ProceduralVegetation/Private/Implementations/PVMeshBuilder.h b/Engine/Plugins/Experimental/ProceduralVegetationEditor/Source/ProceduralVegetation/Private/Implementations/PVMeshBuilder.h
--- a/Engine/Plugins/Experimental/ProceduralVegetationEditor/Source/ProceduralVegetation/Private/Implementations/PVMeshBuilder.h
+++ b/Engine/Plugins/Experimental/ProceduralVegetationEditor/Source/ProceduralVegetation/Private/Implementations/PVMeshBuilder.h
@@ -161,7 +161,7 @@ struct FPVMeshBuilderParams
 
 	UPROPERTY(EditAnywhere, Category = "Displacement Settings",
 		meta=(Tooltip=
-			"Texture used to drive mesh displacement.\n\nAssigns the texture map that controls displacement. Bright and dark areas in the texture push or pull the mesh surface to create additional geometric detail such as bark roughness, grooves, or ridges."
+			"Texture used to drive mesh displacement.\n\nAssigns the texture map that controls displacement. Bright and dark areas in the texture push or pull the mesh surface to create additional geometric detail such as bark roughness, grooves, or ridges.\n\nOnly power of 2 textures with source texture formats (TSF_RGBA32F, TSF_RGBA16F, TSF_BGRA8, TSF_R32F, TSF_R16F, TSF_G8) are supported.\n\nFor formats TSF_RGBA32F, TSF_RGBA16F and TSF_BGRA8 data from R channel is used."
 		))
 	TObjectPtr<UTexture2D> DisplacementTexture = nullptr;
 
@@ -184,7 +184,7 @@ struct FPVMeshBuilderParams
 		meta=(Tooltip=
 			"Controls tiling of displacement texture.\n\nAdjusts the repetition of the displacement texture across the mesh in U and V directions. Higher values increase tiling frequency (smaller details repeated more often), while lower values stretch the texture over a larger area."
 		))
-	FIntVector2 DisplacementUVScale = FIntVector2(1.0f, 1.0f);
+	FVector2f DisplacementUVScale = FVector2f(1.0f, 1.0f);
 
 	UPROPERTY(EditAnywhere, Category="Displacement Settings", DisplayName="Apply Displacement up to generation",
 		meta=(ClampMin=1, ClampMax=10, UIMin=1, UIMax=10, Tooltip=
@@ -303,7 +303,7 @@ struct FPVMeshBuilder
 	static void GenerateDynamicMesh(FManagedArrayCollection& Collection, const FPVMeshBuilderParams& MeshBuilderParams,
 	                                TObjectPtr<UDynamicMesh>& OutMesh);
 
-	static void ExtractDisplacementData(const TObjectPtr<UTexture2D>& Texture, TArray<float>& OutValues);
+	static bool ExtractDisplacementData(const TObjectPtr<UTexture2D>& Texture, TArray<float>& OutValues, FString& OutError);
 
 private:
 	static TSet<int32> CollectHardPoints(const PV::Facades::FBranchFacade& BranchFacade, const PV::Facades::FPointFacade& PointFacade,
diff --git a/Engine/Plugins/Experimental/ProceduralVegetationEditor/Source/ProceduralVegetation/Private/Nodes/PVMeshBuilderSettings.cpp b/Engine/Plugins/Experimental/ProceduralVegetationEditor/Source/ProceduralVegetation/Private/Nodes/PVMeshBuilderSettings.cpp
--- a/Engine/Plugins/Experimental/ProceduralVegetationEditor/Source/ProceduralVegetation/Private/Nodes/PVMeshBuilderSettings.cpp
+++ b/Engine/Plugins/Experimental/ProceduralVegetationEditor/Source/ProceduralVegetation/Private/Nodes/PVMeshBuilderSettings.cpp
@@ -40,9 +40,10 @@ void UPVMeshBuilderSettings::PostEditChangeProperty(FPropertyChangedEvent& Prope
 	if (PropertyChangedEvent.Property
 		&& PropertyChangedEvent.Property->GetFName() == GET_MEMBER_NAME_CHECKED(FPVMeshBuilderParams, DisplacementTexture))
 	{
+		DisplacementWarnings.Empty();
 		if (MesherSettings.DisplacementTexture)
 		{
-			FPVMeshBuilder::ExtractDisplacementData(MesherSettings.DisplacementTexture, MesherSettings.DisplacementValues);
+			FPVMeshBuilder::ExtractDisplacementData(MesherSettings.DisplacementTexture, MesherSettings.DisplacementValues, DisplacementWarnings);
 		}
 		else
 		{
@@ -102,9 +103,10 @@ void UPVMeshBuilderSettings::PostLoad()
 {
 	Super::PostLoad();
 
+	DisplacementWarnings.Empty();
 	if (MesherSettings.DisplacementTexture)
 	{
-		FPVMeshBuilder::ExtractDisplacementData(MesherSettings.DisplacementTexture, MesherSettings.DisplacementValues);
+		FPVMeshBuilder::ExtractDisplacementData(MesherSettings.DisplacementTexture, MesherSettings.DisplacementValues, DisplacementWarnings);
 	}
 	else
 	{
@@ -192,6 +194,11 @@ bool FPVMeshBuilderElement::ExecuteInternal(FPCGContext* InContext) const
 				MutableSettings->MesherSettings.bIsPlantProfileDropdownEnabled = false;
 			}
 
+			if (!Settings->DisplacementWarnings.IsEmpty())
+			{
+				PCGLog::LogWarningOnGraph(FText::FromString(FString::Format(TEXT("{0} "), {*Settings->DisplacementWarnings})), InContext);
+			}
+
 			UPVMeshData* OutManagedArrayCollectionData = FPCGContext::NewObject_AnyThread<UPVMeshData>(InContext);
 
 			FGeometryCollection OutGeometryCollection;
diff --git a/Engine/Plugins/Experimental/ProceduralVegetationEditor/Source/ProceduralVegetation/Private/Nodes/PVMeshBuilderSettings.h b/Engine/Plugins/Experimental/ProceduralVegetationEditor/Source/ProceduralVegetation/Private/Nodes/PVMeshBuilderSettings.h
--- a/Engine/Plugins/Experimental/ProceduralVegetationEditor/Source/ProceduralVegetation/Private/Nodes/PVMeshBuilderSettings.h
+++ b/Engine/Plugins/Experimental/ProceduralVegetationEditor/Source/ProceduralVegetation/Private/Nodes/PVMeshBuilderSettings.h
@@ -47,13 +47,15 @@ public:
 
 	UPROPERTY()
 	bool bCustomMaterialSet = false;
+
+	UPROPERTY()
+	FString DisplacementWarnings;
 	
 	UFUNCTION()
 	TArray<FString> GetPlantProfileOptions() const
 	{
 		return MesherSettings.PlantProfileOptions;
 	}
-	
 };
 
 class FPVMeshBuilderElement : public FPVBaseElement
