[PVE] BugHawk fix

diff --git a/Engine/Plugins/Experimental/ProceduralVegetationEditor/Source/ProceduralVegetation/Private/Implementations/PVMeshBuilder.cpp b/Engine/Plugins/Experimental/ProceduralVegetationEditor/Source/ProceduralVegetation/Private/Implementations/PVMeshBuilder.cpp
--- a/Engine/Plugins/Experimental/ProceduralVegetationEditor/Source/ProceduralVegetation/Private/Implementations/PVMeshBuilder.cpp
+++ b/Engine/Plugins/Experimental/ProceduralVegetationEditor/Source/ProceduralVegetation/Private/Implementations/PVMeshBuilder.cpp
@@ -1038,6 +1038,7 @@ bool FPVMeshBuilder::ExtractDisplacementData(const TObjectPtr<UTexture2D>& Textu
 
 	const uint8* RawData = Source.LockMipReadOnly(0);
 
+	bool bSuccess = true;
 	switch (Source.GetFormat())
 	{
 	case TSF_RGBA32F:
@@ -1100,12 +1101,13 @@ bool FPVMeshBuilder::ExtractDisplacementData(const TObjectPtr<UTexture2D>& Textu
 	default:
 		OutError = FString::Printf(TEXT("ExtractDisplacementData: Unsupported source format %d on texture %s."), Source.GetFormat(), *Texture->GetName());
 		UE_LOG(LogProceduralVegetation, Warning, TEXT("%s"),*OutError);
+		bSuccess = false;
 		break;
 	}
 
 
 	Source.UnlockMip(0);
-	return true;
+	return bSuccess;
 #else
 	OutError = "ExtractDisplacementData only works in editor builds (WITH_EDITOR).";
 	UE_LOG(LogProceduralVegetation, Warning, TEXT("%s"), *OutError);
