[PVE] Multiselect support in export panel. Also fixed bug hawk issue.

diff --git a/Engine/Plugins/Experimental/ProceduralVegetationEditor/Source/ProceduralVegetation/Private/Implementations/PVRemoveBranches.cpp b/Engine/Plugins/Experimental/ProceduralVegetationEditor/Source/ProceduralVegetation/Private/Implementations/PVRemoveBranches.cpp
--- a/Engine/Plugins/Experimental/ProceduralVegetationEditor/Source/ProceduralVegetation/Private/Implementations/PVRemoveBranches.cpp
+++ b/Engine/Plugins/Experimental/ProceduralVegetationEditor/Source/ProceduralVegetation/Private/Implementations/PVRemoveBranches.cpp
@@ -95,8 +95,11 @@ namespace PV::Facades
 				else if (InBasis == ERemoveBranchesBasis::Generation)
 				{
 					const int32 BranchHierarchyGeneration =  InBranchFacade.GetHierarchyGenerationNumber(BranchIndex);
-					
-					OutRemovalData[BranchIndex] = RootHierarchyGeneration / static_cast<float>(BranchHierarchyGeneration);
+
+					if (BranchHierarchyGeneration > 0)
+					{
+						OutRemovalData[BranchIndex] = RootHierarchyGeneration / static_cast<float>(BranchHierarchyGeneration);
+					}
 				}
 
 				MaxValue = FMath::Max(OutRemovalData[BranchIndex], MaxValue);
diff --git a/Engine/Plugins/Experimental/ProceduralVegetationEditor/Source/ProceduralVegetation/Private/Nodes/PVOutputSettings.cpp b/Engine/Plugins/Experimental/ProceduralVegetationEditor/Source/ProceduralVegetation/Private/Nodes/PVOutputSettings.cpp
--- a/Engine/Plugins/Experimental/ProceduralVegetationEditor/Source/ProceduralVegetation/Private/Nodes/PVOutputSettings.cpp
+++ b/Engine/Plugins/Experimental/ProceduralVegetationEditor/Source/ProceduralVegetation/Private/Nodes/PVOutputSettings.cpp
@@ -31,6 +31,7 @@ FText UPVOutputSettings::GetNodeTooltipText() const
 		"- Skeletal Mesh with Nanite Foliage enabled\n\n"
 
 		"If the foliage meshes do not contain skeletal data and Skeletal mesh export is selected, it will build the skeletal data for foliage meshes"
+		"\n\nPress Ctrl + L to lock/unlock node output"
 	); 
 }
 #endif
@@ -39,7 +40,7 @@ UPVOutputSettings::UPVOutputSettings()
 {
 	if (PCGHelpers::IsNewObjectAndNotDefault(this))
 	{
-		static FString DefaultMeshName = "ProceduralVegetationOutput";
+		static FString DefaultMeshName = "Output";
 		static FString DefaultAssetPath = "/Game/ProceduralVegetation/";
 		static FString DefaultWindSettingsPath = "/ProceduralVegetationEditor/SampleAssets/WindSettings/DefaultTreeWindSettings.DefaultTreeWindSettings";
 		ExportSettings.Initialize(DefaultAssetPath, DefaultMeshName, DefaultWindSettingsPath);
diff --git a/Engine/Plugins/Experimental/ProceduralVegetationEditor/Source/ProceduralVegetationEditor/Private/PVEditor.cpp b/Engine/Plugins/Experimental/ProceduralVegetationEditor/Source/ProceduralVegetationEditor/Private/PVEditor.cpp
--- a/Engine/Plugins/Experimental/ProceduralVegetationEditor/Source/ProceduralVegetationEditor/Private/PVEditor.cpp
+++ b/Engine/Plugins/Experimental/ProceduralVegetationEditor/Source/ProceduralVegetationEditor/Private/PVEditor.cpp
@@ -326,29 +326,32 @@ void FPVEditor::OnSelectedNodesChanged(const TSet<UObject*>& InNewSelection)
 
 	UPCGEditorGraphNodeBase* EdNode = nullptr;
 
-	if (!InNewSelection.IsEmpty())
+	SelectedNodes.Empty();
+
+	for(UObject* Selection: InNewSelection)
 	{
-		EdNode = Cast<UPCGEditorGraphNodeBase>(*InNewSelection.CreateConstIterator());
+		EdNode = Cast<UPCGEditorGraphNodeBase>(Selection);
 
 		if (EdNode && EdNode->GetPCGNode())
 		{
-			UPCGNode* Node = EdNode->GetPCGNode();
-
-			UpdateCollectionSpreadSheet(Node);
+			SelectedNodes.Add(EdNode);
+		}
+	}
 
-			if (UPVBaseSettings* Settings = Cast<UPVBaseSettings>(Node->GetSettings()))
-			{
-				const UPVBaseSettings* InspectionSettings = GetNodeSettings(NodeBeingInspected);
+	if (UPCGNode* SelectedNode = GetSelectedNode() ? GetSelectedNode()->GetPCGNode() : nullptr)
+	{
+		UpdateCollectionSpreadSheet(SelectedNode);
+		
+		if (UPVBaseSettings* Settings = Cast<UPVBaseSettings>(SelectedNode->GetSettings()))
+		{
+			const UPVBaseSettings* InspectionSettings = GetNodeSettings(NodeBeingInspected);
 				
-				if (!InspectionSettings || !InspectionSettings->bLockInspection)
-				{
-					ChangeNodeInspection(EdNode);
-				}
+			if (!InspectionSettings || !InspectionSettings->bLockInspection)
+			{
+				ChangeNodeInspection(EdNode);
 			}
-		}
+		}	
 	}
-
-	SelectedNode = EdNode;
 }
 
 void FPVEditor::OnExport_Clicked()
@@ -364,7 +367,7 @@ FReply FPVEditor::ExportInternal()
 		SNew(SPVExportSelectionDialog)
 		.Title(LOCTEXT("Export Selection", "Choose outputs to export"))
 		.Graph(ProceduralVegetationGraph)
-		.SelectedNode(SelectedNode ? SelectedNode->GetPCGNode() : nullptr);
+		.SelectedNodes(SelectedNodes);
 
 	if (ExportSelectionWidget->ShowModal() != EAppReturnType::Ok)
 	{
@@ -601,11 +604,15 @@ FReply FPVEditor::ExportInternal()
 	if (EditorSettings->ExportType == EPVExportType::Selection)
 	{
 		UE_LOG(LogProceduralVegetationEditor, Log, TEXT("Exporting the selected nodes"));
-		if (SelectedNode != nullptr && SelectedNode->GetPCGNode())
+		
+		for (int i = 0; i < SelectedNodes.Num(); i++)
 		{
-			if (const UPVOutputSettings* OutputSettings = Cast<UPVOutputSettings>(SelectedNode->GetPCGNode()->GetSettings()))
+			if (UPCGNode* Node = SelectedNodes[i] ?  SelectedNodes[i]->GetPCGNode() : nullptr)
 			{
-				NodesToExport.Add(SelectedNode->GetPCGNode());
+				if (const UPVOutputSettings* OutputSettings = Cast<UPVOutputSettings>(Node->GetSettings()))
+				{
+					NodesToExport.Add(Node);
+				}	
 			}
 		}
 	}
@@ -743,19 +750,19 @@ FReply FPVEditor::ExportInternal()
 void FPVEditor::OnLockNodeSelection()
 {
 	auto NodeBeingInspectedSettings = GetNodeSettings(NodeBeingInspected);
-	auto SelectedNodeSettings = GetNodeSettings(SelectedNode);
+	auto SelectedNodeSettings = GetNodeSettings(GetSelectedNode());
 
 	if (NodeBeingInspectedSettings)
 	{
 		NodeBeingInspectedSettings->bLockInspection = !NodeBeingInspectedSettings->bLockInspection;
 	}
 
-	if (NodeBeingInspected != SelectedNode)
+	if (NodeBeingInspected != GetSelectedNode())
 	{
-		if (SelectedNode)
+		if (GetSelectedNode())
 		{
 			SelectedNodeSettings->bLockInspection = true;
-			ChangeNodeInspection(SelectedNode);
+			ChangeNodeInspection(GetSelectedNode());
 		}
 	}
 }
@@ -821,6 +828,16 @@ bool FPVEditor::CanToggleInspected() const
 	return false;
 }
 
+UPCGEditorGraphNodeBase* FPVEditor::GetSelectedNode()
+{
+	if (SelectedNodes.Num())
+	{
+		return SelectedNodes[0];
+	}
+
+	return nullptr;
+}
+
 bool FPVEditor::GetCollectionFromNode(const UPCGNode* InNode, FName& OutPinName, FManagedArrayCollection& OutCollection) const
 {
 	const UPVData* PVData = GetPVDataFromNode(InNode, OutPinName);
diff --git a/Engine/Plugins/Experimental/ProceduralVegetationEditor/Source/ProceduralVegetationEditor/Private/PVEditor.h b/Engine/Plugins/Experimental/ProceduralVegetationEditor/Source/ProceduralVegetationEditor/Private/PVEditor.h
--- a/Engine/Plugins/Experimental/ProceduralVegetationEditor/Source/ProceduralVegetationEditor/Private/PVEditor.h
+++ b/Engine/Plugins/Experimental/ProceduralVegetationEditor/Source/ProceduralVegetationEditor/Private/PVEditor.h
@@ -53,7 +53,8 @@ private:
 	static const FName CollectionSpreadSheetTabId;
 	
 	TObjectPtr<UProceduralVegetation> ProceduralVegetationBeingEdited = nullptr;
-	TObjectPtr<UPCGEditorGraphNodeBase> SelectedNode = nullptr;
+
+	TArray<TObjectPtr<UPCGEditorGraphNodeBase>> SelectedNodes;
 	TObjectPtr<UPCGEditorGraphNodeBase> NodeBeingInspected = nullptr;
 	
 	TArray<TSharedRef<SPVEditorViewport>> ViewportWidgets;
@@ -83,4 +84,6 @@ protected:
 	const UPVData* GetPVDataFromNode(const UPCGNode* InNode, FName& OutPinName) const;
 	virtual bool CanToggleInspected() const override;
 	bool GetCollectionFromNode(const UPCGNode* InNode, FName& OutPinName, FManagedArrayCollection& OutCollection) const;
+
+	UPCGEditorGraphNodeBase* GetSelectedNode();
 };
diff --git a/Engine/Plugins/Experimental/ProceduralVegetationEditor/Source/ProceduralVegetationEditor/Private/Widgets/SPVExportSelectionDialog.cpp b/Engine/Plugins/Experimental/ProceduralVegetationEditor/Source/ProceduralVegetationEditor/Private/Widgets/SPVExportSelectionDialog.cpp
--- a/Engine/Plugins/Experimental/ProceduralVegetationEditor/Source/ProceduralVegetationEditor/Private/Widgets/SPVExportSelectionDialog.cpp
+++ b/Engine/Plugins/Experimental/ProceduralVegetationEditor/Source/ProceduralVegetationEditor/Private/Widgets/SPVExportSelectionDialog.cpp
@@ -7,6 +7,7 @@
 #include "SPVExportItem.h"
 #include "ProceduralVegetation.h"
 #include "PVEditorSettings.h"
+#include "Nodes/PCGEditorGraphNodeBase.h"
 #include "Nodes/PVOutputSettings.h"
 #include "Widgets/Input/STextComboBox.h"
 #include "Widgets/Layout/SScrollBox.h"
@@ -16,7 +17,7 @@
 void SPVExportSelectionDialog::Construct(const FArguments& InArgs)
 {
 	Graph = InArgs._Graph;
-	SelectedNode = InArgs._SelectedNode;
+	SelectedNodes = InArgs._SelectedNodes;
 
 	BuildExportTypeOptions();
 	
@@ -146,25 +147,49 @@ void SPVExportSelectionDialog::AddExportItems()
 	}
 	else
 	{
-		FString OutputName = "None";
-		
-		if (SelectedNode)
+		FString OutputName;
+
+		if (SelectedNodes.IsEmpty())
 		{
-			const UPVOutputSettings* OutputSettings = Cast<UPVOutputSettings>(SelectedNode->GetSettings());
+			OutputName = "None";
+		}
 
-			if (OutputSettings)
+		for (int i = 0; i < SelectedNodes.Num(); i++)
+		{
+			if (const UPCGNode* SelectedNode = SelectedNodes[i] ? SelectedNodes[i]->GetPCGNode() : nullptr)
 			{
-				OutputName = SelectedNode->GetNodeTitle(EPCGNodeTitleType::FullTitle).ToString();
-			}	
+				if (const UPVOutputSettings* OutputSettings = Cast<UPVOutputSettings>(SelectedNode->GetSettings()))
+				{
+					if (OutputName.IsEmpty())
+					{
+						OutputName = FString::Printf(TEXT("%s"), *SelectedNode->GetNodeTitle(EPCGNodeTitleType::FullTitle).ToString());
+					}
+					else
+					{
+						OutputName = FString::Printf(TEXT("%s\n%s"), *OutputName,*SelectedNode->GetNodeTitle(EPCGNodeTitleType::FullTitle).ToString());
+					}
+				}	
+			}
 		}
 		
-		FText Label = FText::FromString(FString::Printf(TEXT("Current output: %s"), *OutputName));
+		FText Outputs = FText::FromString(FString::Printf(TEXT("%s"), *OutputName));
 			
 		ScrollBox->AddSlot()
 		.Padding(5)
 		[
-			SNew(STextBlock)
-			.Text(Label)
+			SNew(SHorizontalBox)
+			+ SHorizontalBox::Slot()
+			.AutoWidth()
+			[
+				SNew(STextBlock)
+				.Text(FText::FromString(FString::Printf(TEXT("Current outputs(s): "))))	
+			]
+			+ SHorizontalBox::Slot()
+			.AutoWidth()
+			[
+				SNew(STextBlock)
+				.Text(Outputs)	
+			]
 		];
 	}
 
@@ -225,14 +250,17 @@ void SPVExportSelectionDialog::UpdateExportButtonState()
 	}
 	else
 	{
-		if (SelectedNode)
+		for (int i = 0; i < SelectedNodes.Num(); i++)
 		{
-			const UPVOutputSettings* OutputSettings = Cast<UPVOutputSettings>(SelectedNode->GetSettings());
-
-			if (OutputSettings)
+			if (const UPCGNode* SelectedNode = SelectedNodes[i] ? SelectedNodes[i]->GetPCGNode() : nullptr)
 			{
-				HasValidExport = true;
-			}	
+				const UPVOutputSettings* OutputSettings = Cast<UPVOutputSettings>(SelectedNode->GetSettings());
+
+				if (OutputSettings)
+				{
+					HasValidExport = true;
+				}
+			}
 		}
 	}
 
diff --git a/Engine/Plugins/Experimental/ProceduralVegetationEditor/Source/ProceduralVegetationEditor/Private/Widgets/SPVExportSelectionDialog.h b/Engine/Plugins/Experimental/ProceduralVegetationEditor/Source/ProceduralVegetationEditor/Private/Widgets/SPVExportSelectionDialog.h
--- a/Engine/Plugins/Experimental/ProceduralVegetationEditor/Source/ProceduralVegetationEditor/Private/Widgets/SPVExportSelectionDialog.h
+++ b/Engine/Plugins/Experimental/ProceduralVegetationEditor/Source/ProceduralVegetationEditor/Private/Widgets/SPVExportSelectionDialog.h
@@ -4,6 +4,7 @@
 #include "Input/Reply.h"
 #include "Widgets/SWindow.h"
 
+class UPCGEditorGraphNodeBase;
 class UPVOutputSettings;
 class SScrollBox;
 class UProceduralVegetationGraph;
@@ -18,7 +19,7 @@ class SPVExportSelectionDialog : public SWindow
 	SLATE_ARGUMENT(FText, Title)
 	SLATE_ARGUMENT(const UProceduralVegetationGraph*, Graph)
 	SLATE_ARGUMENT(TArray<FText>, ExportItems)
-	SLATE_ARGUMENT(UPCGNode*, SelectedNode)
+	SLATE_ARGUMENT(TArray<TObjectPtr<UPCGEditorGraphNodeBase>>, SelectedNodes)
 	SLATE_END_ARGS()
 
 public:
@@ -35,7 +36,7 @@ private:
 	TSharedPtr<SScrollBox> ScrollBox;
 	TSharedPtr<SPrimaryButton> ExportButton;
 	const UProceduralVegetationGraph* Graph = nullptr;
-	const UPCGNode* SelectedNode = nullptr;
+	TArray<TObjectPtr<UPCGEditorGraphNodeBase>> SelectedNodes;
 
 	TArray<TSharedPtr<FString>> ExportTypes;
 	TSharedPtr<FString> CurrentExportType;
diff --git a/Engine/Plugins/Experimental/ProceduralVegetationEditor/Source/ProceduralVegetationEditor/Public/PVEditorSettings.h b/Engine/Plugins/Experimental/ProceduralVegetationEditor/Source/ProceduralVegetationEditor/Public/PVEditorSettings.h
--- a/Engine/Plugins/Experimental/ProceduralVegetationEditor/Source/ProceduralVegetationEditor/Public/PVEditorSettings.h
+++ b/Engine/Plugins/Experimental/ProceduralVegetationEditor/Source/ProceduralVegetationEditor/Public/PVEditorSettings.h
@@ -9,7 +9,7 @@
 UENUM()
 enum class EPVExportType
 {
-	Selection	UMETA(DisplayName = "Selected output node"),
+	Selection	UMETA(DisplayName = "Selected output(s) node"),
 	BatchExport UMETA(DisplayName = "Batch export")
 };
 
