Add PVE Slope node

diff --git a/Engine/Plugins/Experimental/ProceduralVegetationEditor/Source/ProceduralVegetation/Private/Facades/PVFoliageFacade.cpp b/Engine/Plugins/Experimental/ProceduralVegetationEditor/Source/ProceduralVegetation/Private/Facades/PVFoliageFacade.cpp
--- a/Engine/Plugins/Experimental/ProceduralVegetationEditor/Source/ProceduralVegetation/Private/Facades/PVFoliageFacade.cpp
+++ b/Engine/Plugins/Experimental/ProceduralVegetationEditor/Source/ProceduralVegetation/Private/Facades/PVFoliageFacade.cpp
@@ -269,6 +269,29 @@ namespace PV::Facades
 		return DefaultPivotPoint;
 	}
 
+	const FVector3f& FFoliageFacade::GetUpVector(const int32 Index) const
+	{
+		if (UpVectorsAttribute.IsValid() && UpVectorsAttribute.IsValidIndex(Index))
+		{
+			return UpVectorsAttribute[Index];
+		}
+
+		static const FVector3f DefaultUpVector = FVector3f::Zero();
+		return DefaultUpVector;
+	}
+
+	const FVector3f& FFoliageFacade::GetNormalVector(const int32 Index) const
+	{
+		if (NormalVectorsAttribute.IsValid() && NormalVectorsAttribute.IsValidIndex(Index))
+		{
+			return NormalVectorsAttribute[Index];
+		}
+
+		static const FVector3f DefaultNormalVector = FVector3f::Zero();
+		return DefaultNormalVector;
+	}
+
+
 	void FFoliageFacade::SetPivotPoint(const int32 Index, const FVector3f& Input)
 	{
 		check(!IsConst());
@@ -278,6 +301,24 @@ namespace PV::Facades
 		}
 	}
 
+	void FFoliageFacade::SetUpVector(const int32 Index, const FVector3f& Input)
+	{
+		check(!IsConst());
+		if (IsValid() && UpVectorsAttribute.IsValidIndex(Index))
+		{
+			UpVectorsAttribute.ModifyAt(Index, Input);
+		}
+	}
+
+	void FFoliageFacade::SetNormalVector(const int32 Index, const FVector3f& Input)
+	{
+		check(!IsConst());
+		if (IsValid() && NormalVectorsAttribute.IsValidIndex(Index))
+		{
+			NormalVectorsAttribute.ModifyAt(Index, Input);
+		}
+	}
+
 	void FFoliageFacade::SetScale(const int32 Index, const float& Input)
 	{
 		check(!IsConst());
@@ -361,8 +402,7 @@ namespace PV::Facades
 
 	FTransform FFoliageFacade::GetFoliageTransform(int32 Id) const
 	{
-		FFoliageEntryData Data = GetFoliageEntry(Id);
-		FString FoliageName = GetFoliageName(Data.NameId);
+		const FFoliageEntryData Data = GetFoliageEntry(Id);
 
 		const FQuat UpVectorRotationQuat = FQuat::FindBetweenNormals(FVector::UpVector, FVector(Data.UpVector).GetSafeNormal());
 		const FVector RotatedVectorFace = UpVectorRotationQuat.RotateVector(FVector::RightVector);
diff --git a/Engine/Plugins/Experimental/ProceduralVegetationEditor/Source/ProceduralVegetation/Private/Implementations/PVSlope.cpp b/Engine/Plugins/Experimental/ProceduralVegetationEditor/Source/ProceduralVegetation/Private/Implementations/PVSlope.cpp
new file mode 100644
--- /dev/null
+++ b/Engine/Plugins/Experimental/ProceduralVegetationEditor/Source/ProceduralVegetation/Private/Implementations/PVSlope.cpp
@@ -0,0 +1,285 @@
+// Copyright Epic Games, Inc. All Rights Reserved.
+
+#include "PVSlope.h"
+#include "Facades/PVBranchFacade.h"
+#include "Facades/PVBudVectorsFacade.h"
+#include "Facades/PVFoliageFacade.h"
+#include "Facades/PVPlantFacade.h"
+#include "Facades/PVPointFacade.h"
+
+namespace PVSlope
+{
+	FQuat4f CalculateBranchRotation(float SlopeDirection, float SlopeAngle)
+	{
+		const float SlopeDirectionInRadians = FMath::DegreesToRadians(SlopeDirection);
+		const float SlopeAngleInRadians = FMath::DegreesToRadians(SlopeAngle);
+
+		const FQuat4f YawRotation = FQuat4f(FVector3f::UpVector, SlopeDirectionInRadians);
+		const FQuat4f LeanRotation = FQuat4f(FVector3f::RightVector, SlopeAngleInRadians);
+		return YawRotation * LeanRotation;
+	}
+
+	TArray<int32> GetBranchChildIndices(int32 BranchIndex, const PV::Facades::FBranchFacade& BranchFacade)
+	{
+		const TArray<int32>& CurrentBranchChildren = BranchFacade.GetChildren(BranchIndex);
+		const TManagedArray<int32>& BranchNumbers = BranchFacade.GetBranchNumbers();
+
+		TArray<int32> ChildrenIndices;
+		ChildrenIndices.Reserve(CurrentBranchChildren.Num());
+
+		for (const int32 Child : CurrentBranchChildren)
+		{
+			const int32 ChildIndex = BranchNumbers.Find(Child);
+			const TArray<int32>& Parents = BranchFacade.GetParents(ChildIndex);
+
+			if (ChildIndex != INDEX_NONE && !Parents.IsEmpty() && Parents[Parents.Num() - 1] == BranchNumbers[BranchIndex])
+			{
+				ChildrenIndices.Add(ChildIndex);
+			}
+		}
+
+		return ChildrenIndices;
+	}
+
+	TArray<int32> GetBranchSegmentChildren(
+		const PV::Facades::FBranchFacade& BranchFacade, 
+		const PV::Facades::FPointFacade& PointFacade,
+		int32 BranchIndex, 
+		int32 BranchPointIndex)
+	{
+		const TArray<int32>& BranchPoints = BranchFacade.GetPoints(BranchIndex);
+		const TArray<int32>& BranchChildren = BranchFacade.GetChildren(BranchIndex);
+		const TManagedArray<int32>& BranchNumbers = BranchFacade.GetBranchNumbers();
+
+		check(BranchPoints.Find(BranchPointIndex) != INDEX_NONE);
+
+		const float BranchPointLengthFromRoot = PointFacade.GetLengthFromRoot(BranchPointIndex);
+		const float PreviousBranchPointLengthFromRoot = PointFacade.GetLengthFromRoot(BranchPointIndex - 1);
+
+		TArray<int32> ChildrenIndices;
+		ChildrenIndices.Reserve(BranchChildren.Num());
+
+		for (const int32 BranchChild : BranchChildren)
+		{
+			const int32 ChildIndex = BranchNumbers.Find(BranchChild);
+			if (ChildIndex == INDEX_NONE)
+			{
+				continue;
+			}
+
+			const TArray<int32>& Parents = BranchFacade.GetParents(ChildIndex);
+			if (Parents.IsEmpty() || Parents[Parents.Num() - 1] != BranchNumbers[BranchIndex])
+			{
+				continue;
+			}
+
+			const int32 FirstChildPointIndex = BranchFacade.GetPoints(ChildIndex)[0];
+			const float ChildLengthFromRoot = PointFacade.GetLengthFromRoot(FirstChildPointIndex);
+			if (ChildLengthFromRoot > BranchPointLengthFromRoot || ChildLengthFromRoot <= PreviousBranchPointLengthFromRoot)
+			{
+				continue;
+			}
+
+			ChildrenIndices.Add(ChildIndex);
+		}
+
+		return ChildrenIndices;
+	}
+
+	TArray<int32> GetBranchSegmentFoliageEntryIds(
+		PV::Facades::FFoliageFacade& FoliageFacade,
+		const PV::Facades::FPointFacade& PointFacade,
+		int32 BranchIndex,
+		int32 BranchPointIndex)
+	{
+		TArray<int32> OutFoliageEntryIds;
+
+		const float BranchPointLengthFromRoot = PointFacade.GetLengthFromRoot(BranchPointIndex);
+		const float PreviousBranchPointLengthFromRoot = PointFacade.GetLengthFromRoot(BranchPointIndex - 1);
+
+		TArray<int32> FoliageEntryIds = FoliageFacade.GetFoliageEntryIdsForBranch(BranchIndex);
+		OutFoliageEntryIds.Reserve(FoliageEntryIds.Num());
+
+		for (const int32 FoliageEntryId : FoliageEntryIds)
+		{
+			const float FoliageEntryLengthFromRoot = FoliageFacade.GetLengthFromRoot(FoliageEntryId);
+
+			if (FoliageEntryLengthFromRoot <= BranchPointLengthFromRoot && FoliageEntryLengthFromRoot > PreviousBranchPointLengthFromRoot)
+			{
+				OutFoliageEntryIds.Add(FoliageEntryId);
+			}
+		}
+
+		return OutFoliageEntryIds;
+	}
+
+	void CalcBranchLengthSquared_Recursive(
+		const PV::Facades::FBranchFacade& BranchFacade,
+		const PV::Facades::FPointFacade& PointFacade,
+		int32 BranchIndex,
+		const FVector3f& BranchStartPosition,
+		float& OutLengthSquared)
+	{
+		const TArray<int32>& CurrentBranchPoints = BranchFacade.GetPoints(BranchIndex);
+		const TManagedArray<FVector3f>& PointPositions = PointFacade.GetPositions();
+
+		for (int32 Index = 0; Index < CurrentBranchPoints.Num(); ++Index)
+		{
+			const int32 PointIndex = CurrentBranchPoints[Index];
+			const FVector3f& PointPosition = PointPositions[PointIndex];
+
+			const float PointDistSquared = (PointPosition - BranchStartPosition).SizeSquared();
+			if (PointDistSquared > OutLengthSquared)
+			{
+				OutLengthSquared = PointDistSquared;
+			}
+
+			const TArray<int32> BranchSegmentChildren = GetBranchSegmentChildren(BranchFacade, PointFacade, BranchIndex, PointIndex);
+			for (const int32 ChildIndex : BranchSegmentChildren)
+			{
+				CalcBranchLengthSquared_Recursive(BranchFacade, PointFacade, ChildIndex, BranchStartPosition, OutLengthSquared);
+			}
+		}
+	}
+
+	float CalcBranchLength(
+		const PV::Facades::FBranchFacade& BranchFacade, 
+		const PV::Facades::FPointFacade& PointFacade, 
+		int32 BranchIndex)
+	{
+		const TArray<int32>& CurrentBranchPoints = BranchFacade.GetPoints(BranchIndex);
+		if (CurrentBranchPoints.Num() == 0)
+		{
+			return 0;
+		}
+
+		const TManagedArray<FVector3f>& PointPositions = PointFacade.GetPositions();
+		const FVector3f& RootPosition = PointPositions[CurrentBranchPoints[0]];
+
+		float BranchLengthSquared = 0;
+		FVector3f DebugLargetsPoint = FVector3f::ZeroVector;
+		CalcBranchLengthSquared_Recursive(BranchFacade, PointFacade, BranchIndex, RootPosition, BranchLengthSquared);
+
+		return BranchLengthSquared > 0 ? FMath::Sqrt(BranchLengthSquared) : 0.f;
+	}
+
+	void ApplySlope_Recursive(
+		PV::Facades::FBranchFacade& BranchFacade,
+		PV::Facades::FPointFacade& PointFacade,
+		PV::Facades::FBudVectorsFacade& BudVectorFacade,
+		PV::Facades::FFoliageFacade& FoliageFacade,
+		const int32 BranchIndex,
+		const FPVSlopeParams& SlopeParams,
+		float TreeHeight,
+		const FVector3f& TrunkPosition,
+		FVector3f ParentPosition,
+		FVector3f ParentIdentityPosition
+	)
+	{
+		TManagedArray<FVector3f>& PointPositions = PointFacade.ModifyPositions();
+		TManagedArray<TArray<FVector3f>>& PointBudDirections = BudVectorFacade.ModifyBudDirections();
+
+		const TArray<int32>& CurrentBranchPoints = BranchFacade.GetPoints(BranchIndex);
+
+		for (int32 Index = 0; Index < CurrentBranchPoints.Num(); ++Index)
+		{
+			const int32 PointIndex = CurrentBranchPoints[Index];
+			FVector3f& Position = PointPositions[PointIndex];
+
+			const float TreeAngleAlpha = 1.0 - (TreeHeight > 0 ? ((Position - TrunkPosition).Size() / TreeHeight) : 0);
+			const float AngleStrength = FMath::Clamp(FMath::Pow(TreeAngleAlpha, SlopeParams.BendStrength), 0.0, 1.0);
+
+			const FQuat4f SlopeRotation = CalculateBranchRotation(SlopeParams.SlopeDirection, SlopeParams.SlopeAngle * AngleStrength);
+
+			TArray<FVector3f>& CurrentBudDirections = PointBudDirections[PointIndex];
+			for (int32 i = 0; i < CurrentBudDirections.Num(); ++i)
+			{
+				if (i == PV::Facades::BudDirectionsLightOptimalIndex || i == PV::Facades::BudDirectionsLightSubOptimal)
+				{
+					continue;
+				}
+
+				CurrentBudDirections[i] = SlopeRotation.RotateVector(CurrentBudDirections[i]);
+			}
+
+			const TArray<int32> BranchSegmentFoliageEntryIds = GetBranchSegmentFoliageEntryIds(FoliageFacade, PointFacade, BranchIndex, PointIndex);
+			for (int32 FoliageEntryId : BranchSegmentFoliageEntryIds)
+			{
+				const FVector3f FoliageEntryPivotPoint = FoliageFacade.GetPivotPoint(FoliageEntryId);
+				const FVector3f NewFoliageEntryPivotPoint = SlopeRotation.RotateVector(FoliageEntryPivotPoint - ParentIdentityPosition) + ParentPosition;
+				FoliageFacade.SetPivotPoint(FoliageEntryId, NewFoliageEntryPivotPoint);
+
+				const FVector3f FoliageEntryUpVector = FoliageFacade.GetUpVector(FoliageEntryId);
+				const FVector3f NewFoliageEntryUpVector = SlopeRotation.RotateVector(FoliageEntryUpVector);
+				FoliageFacade.SetUpVector(FoliageEntryId, NewFoliageEntryUpVector);
+
+				const FVector3f FoliageEntryNormalVector = FoliageFacade.GetNormalVector(FoliageEntryId);
+				const FVector3f NewFoliageEntryNormalVector = SlopeRotation.RotateVector(FoliageEntryNormalVector);
+				FoliageFacade.SetNormalVector(FoliageEntryId, NewFoliageEntryNormalVector);
+			}
+
+			const FVector3f NewPosition = SlopeRotation.RotateVector(Position - ParentIdentityPosition) + ParentPosition;
+
+			ParentIdentityPosition = Position;
+			Position = NewPosition;
+			ParentPosition = Position;
+
+			const TArray<int32> BranchSegmentChildren = GetBranchSegmentChildren(BranchFacade, PointFacade, BranchIndex, PointIndex);
+			for (const int32 ChildIndex : BranchSegmentChildren)
+			{
+				ApplySlope_Recursive(
+					BranchFacade,
+					PointFacade,
+					BudVectorFacade,
+					FoliageFacade,
+					ChildIndex,
+					SlopeParams,
+					TreeHeight,
+					TrunkPosition,
+					ParentPosition,
+					ParentIdentityPosition
+				);
+			}
+		}
+	}
+}
+
+void FPVSlope::ApplySlope(const FPVSlopeParams& InSlopeParams, FManagedArrayCollection& OutCollection)
+{
+	PV::Facades::FBranchFacade BranchFacade = PV::Facades::FBranchFacade(OutCollection);
+	PV::Facades::FPointFacade PointFacade = PV::Facades::FPointFacade(OutCollection);
+	PV::Facades::FBudVectorsFacade BudVectorFacade = PV::Facades::FBudVectorsFacade(OutCollection);
+	PV::Facades::FFoliageFacade FoliageFacade = PV::Facades::FFoliageFacade(OutCollection);
+	const PV::Facades::FPlantFacade PlantFacade(OutCollection);
+
+	if (!PointFacade.IsValid() || !BranchFacade.IsValid())
+	{
+		return;
+	}
+
+	for (const int32 TrunkIndex : PlantFacade.GetTrunkIndices())
+	{
+		const TArray<int32>& CurrentBranchPoints = BranchFacade.GetPoints(TrunkIndex);
+		if (CurrentBranchPoints.Num() == 0)
+		{
+			continue;
+		}
+
+		const float BranchLength = PVSlope::CalcBranchLength(BranchFacade, PointFacade, TrunkIndex);
+		const FVector3f TrunkPosition = PointFacade.GetPosition(CurrentBranchPoints[0]);
+		const FVector3f PivotPoint = InSlopeParams.TrunkPivotPoint == EPVSlopeTrunkPivotPoint::Origin ? FVector3f::ZeroVector : TrunkPosition;
+
+		PVSlope::ApplySlope_Recursive(
+			BranchFacade,
+			PointFacade,
+			BudVectorFacade,
+			FoliageFacade,
+			TrunkIndex,
+			InSlopeParams,
+			BranchLength,
+			TrunkPosition,
+			PivotPoint,
+			PivotPoint
+		);
+	}
+}
diff --git a/Engine/Plugins/Experimental/ProceduralVegetationEditor/Source/ProceduralVegetation/Private/Implementations/PVSlope.h b/Engine/Plugins/Experimental/ProceduralVegetationEditor/Source/ProceduralVegetation/Private/Implementations/PVSlope.h
new file mode 100644
--- /dev/null
+++ b/Engine/Plugins/Experimental/ProceduralVegetationEditor/Source/ProceduralVegetation/Private/Implementations/PVSlope.h
@@ -0,0 +1,47 @@
+// Copyright Epic Games, Inc. All Rights Reserved.
+
+#pragma once
+
+#include "CoreMinimal.h"
+#include "GeometryCollection/ManagedArrayCollection.h"
+#include "PVSlope.generated.h"
+
+namespace PV::Facades
+{
+	class FPointFacade;
+	class FBranchFacade;
+}
+
+UENUM()
+enum class EPVSlopeTrunkPivotPoint
+{
+	Origin,
+	Trunk
+};
+
+USTRUCT()
+struct FPVSlopeParams
+{
+	GENERATED_BODY()
+
+	/** The angle of the ground. Measured in degrees. */
+	UPROPERTY(EditAnywhere, Category = "Slope", meta = (Units = Degrees, ClampMin = -90.0f, ClampMax = 90.0f, UIMin = -90.0f, UIMax = 90.0f))
+	float SlopeAngle = 22.5;
+
+	/** The direction of the slope. Measured in degrees. */
+	UPROPERTY(EditAnywhere, Category = "Slope", meta = (Units = Degrees, ClampMin = -180.0f, ClampMax = 180.0f, UIMin = -180.0f, UIMax = 180.0f))
+	float SlopeDirection = 0;
+
+	/** How strongly the vegetation should bend towards the sky as it grows. 0 = no bending. */
+	UPROPERTY(EditAnywhere, Category = "Slope", meta = (ClampMin = 0.0f, ClampMax = 100.0f, UIMin = 0.0f, UIMax = 10.0f))
+	float BendStrength = 2.0;
+
+	/** The pivot point of the tunk. */
+	UPROPERTY(EditAnywhere, Category = "Slope")
+	EPVSlopeTrunkPivotPoint TrunkPivotPoint;
+};
+
+struct FPVSlope
+{
+	static void ApplySlope(const FPVSlopeParams& InSlopeParams, FManagedArrayCollection& OutCollection);
+};
diff --git a/Engine/Plugins/Experimental/ProceduralVegetationEditor/Source/ProceduralVegetation/Private/Nodes/PVSlopeSettings.cpp b/Engine/Plugins/Experimental/ProceduralVegetationEditor/Source/ProceduralVegetation/Private/Nodes/PVSlopeSettings.cpp
new file mode 100644
--- /dev/null
+++ b/Engine/Plugins/Experimental/ProceduralVegetationEditor/Source/ProceduralVegetation/Private/Nodes/PVSlopeSettings.cpp
@@ -0,0 +1,78 @@
+// Copyright Epic Games, Inc. All Rights Reserved.
+
+#include "PVSlopeSettings.h"
+#include "DataTypes/PVData.h"
+#include "DataTypes/PVGrowthData.h"
+#include "Data/PCGBasePointData.h"
+#include "GeometryCollection/ManagedArrayCollection.h"
+#include "Helpers/PCGSettingsHelpers.h"
+#include "Implementations/PVSlope.h"
+
+#define LOCTEXT_NAMESPACE "PVSlopeSettings"
+
+#if WITH_EDITOR
+FText UPVSlopeSettings::GetDefaultNodeTitle() const
+{ 
+	return LOCTEXT("NodeTitle", "Slope"); 
+}
+
+FText UPVSlopeSettings::GetNodeTooltipText() const
+{ 
+	return LOCTEXT("NodeTooltip", 
+		"Simulates vegetation growing on a slope, e.g it will start growing perpendicular to the slope and slowly bend towards the sky as it grows larger."
+		"\n\nPress Ctrl + L to lock/unlock node output"
+	); 
+}
+#endif
+
+FPCGDataTypeIdentifier UPVSlopeSettings::GetInputPinTypeIdentifier() const
+{
+	return FPCGDataTypeIdentifier{ FPVDataTypeInfoGrowth::AsId() };
+}
+
+FPCGDataTypeIdentifier UPVSlopeSettings::GetOutputPinTypeIdentifier() const
+{
+	return FPCGDataTypeIdentifier{ FPVDataTypeInfoGrowth::AsId() };
+}
+
+FPCGElementPtr UPVSlopeSettings::CreateElement() const
+{
+	return MakeShared<FPVSlopeElement>();
+}
+
+bool FPVSlopeElement::ExecuteInternal(FPCGContext* InContext) const
+{
+	TRACE_CPUPROFILER_EVENT_SCOPE(FPVSlopeElement::Execute);
+
+	check(InContext);
+
+	const UPVSlopeSettings* Settings = InContext->GetInputSettings<UPVSlopeSettings>();
+	check(Settings);
+
+	for (const FPCGTaggedData& Input : InContext->InputData.GetInputsByPin(PCGPinConstants::DefaultInputLabel))
+	{
+		if (const UPVGrowthData* InputData = Cast<UPVGrowthData>(Input.Data))
+		{
+			FManagedArrayCollection OutCollection = InputData->GetCollection();
+
+			UPVGrowthData* OutManagedArrayCollectionData = FPCGContext::NewObject_AnyThread<UPVGrowthData>(InContext);
+
+			if (Settings->SlopeParams.SlopeAngle != 0.0f)
+			{
+				FPVSlope::ApplySlope(Settings->SlopeParams, OutCollection);
+			}
+
+			OutManagedArrayCollectionData->Initialize(MoveTemp(OutCollection));
+			InContext->OutputData.TaggedData.Emplace(OutManagedArrayCollectionData);
+		}
+		else
+		{
+			PCGLog::InputOutput::LogInvalidInputDataError(InContext);
+			return true;
+		}
+	}
+
+	return true;
+}
+
+#undef LOCTEXT_NAMESPACE
\ No newline at end of file
diff --git a/Engine/Plugins/Experimental/ProceduralVegetationEditor/Source/ProceduralVegetation/Private/Nodes/PVSlopeSettings.h b/Engine/Plugins/Experimental/ProceduralVegetationEditor/Source/ProceduralVegetation/Private/Nodes/PVSlopeSettings.h
new file mode 100644
--- /dev/null
+++ b/Engine/Plugins/Experimental/ProceduralVegetationEditor/Source/ProceduralVegetation/Private/Nodes/PVSlopeSettings.h
@@ -0,0 +1,39 @@
+// Copyright Epic Games, Inc. All Rights Reserved.
+#pragma once
+
+#include "CoreMinimal.h"
+#include "Nodes/PVBaseSettings.h"
+#include "Implementations/PVSlope.h"
+#include "PVSlopeSettings.generated.h"
+
+UCLASS(BlueprintType, ClassGroup = (Procedural))
+class UPVSlopeSettings : public UPVBaseSettings
+{
+	GENERATED_BODY()
+
+public:
+#if WITH_EDITOR
+	virtual FName GetDefaultNodeName() const override { return FName(TEXT("ProceduralVegetationSlope")); }
+	virtual FText GetDefaultNodeTitle() const override;
+	virtual FText GetNodeTooltipText() const override;
+	virtual FLinearColor GetNodeTitleColor() const override { return FLinearColor::Yellow; }
+
+	virtual TArray<EPVRenderType> GetDefaultRenderType() const override { return TArray{ EPVRenderType::PointData }; }
+#endif
+
+protected:
+	virtual FPCGDataTypeIdentifier GetInputPinTypeIdentifier() const override;
+	virtual FPCGDataTypeIdentifier GetOutputPinTypeIdentifier() const override;
+
+	virtual FPCGElementPtr CreateElement() const override;
+
+public:
+	UPROPERTY(EditAnywhere, Category = "Slope", meta = (ShowOnlyInnerProperties))
+	FPVSlopeParams SlopeParams;
+};
+
+class FPVSlopeElement : public FPVBaseElement
+{
+protected:
+	virtual bool ExecuteInternal(FPCGContext* InContext) const override;
+};
diff --git a/Engine/Plugins/Experimental/ProceduralVegetationEditor/Source/ProceduralVegetation/Public/Facades/PVFoliageFacade.h b/Engine/Plugins/Experimental/ProceduralVegetationEditor/Source/ProceduralVegetation/Public/Facades/PVFoliageFacade.h
--- a/Engine/Plugins/Experimental/ProceduralVegetationEditor/Source/ProceduralVegetation/Public/Facades/PVFoliageFacade.h
+++ b/Engine/Plugins/Experimental/ProceduralVegetationEditor/Source/ProceduralVegetation/Public/Facades/PVFoliageFacade.h
@@ -66,8 +66,16 @@ namespace PV::Facades
 
 		const FVector3f& GetPivotPoint(const int32 Index) const;
 
+		const FVector3f& GetUpVector(const int32 Index) const;
+
+		const FVector3f& GetNormalVector(const int32 Index) const;
+
 		void SetPivotPoint(const int32 Index, const FVector3f& Input);
 
+		void SetUpVector(const int32 Index, const FVector3f& Input);
+
+		void SetNormalVector(const int32 Index, const FVector3f& Input);
+
 		void SetScale(const int32 Index, const float& Input);
 
 		void SetFoliageName(const int32 Index, const FString& Input);
