[PVE] Mannequin offset fixes, padding for text in scale vis., Fixed exposure settings, GreyAmbientProfile by default in viewport

#virtualized

diff --git a/Engine/Plugins/Experimental/ProceduralVegetationEditor/Source/ProceduralVegetationEditor/Private/SPVEditorViewport.cpp b/Engine/Plugins/Experimental/ProceduralVegetationEditor/Source/ProceduralVegetationEditor/Private/SPVEditorViewport.cpp
--- a/Engine/Plugins/Experimental/ProceduralVegetationEditor/Source/ProceduralVegetationEditor/Private/SPVEditorViewport.cpp
+++ b/Engine/Plugins/Experimental/ProceduralVegetationEditor/Source/ProceduralVegetationEditor/Private/SPVEditorViewport.cpp
@@ -16,12 +16,30 @@
 
 #include "Visualizations/PVScaleVisualizationComponent.h"
 #include "DataTypes/PVData.h"
+
+#include "Editor/EditorPerProjectUserSettings.h"
+
 #include "Nodes/PVBaseSettings.h"
 
 #include "Widgets/Input/SNumericEntryBox.h"
 
 #define LOCTEXT_NAMESPACE "PVEditorViewportToolbarSections"
 
+SPVEditorViewport::SPVEditorViewport()
+{
+	if (AdvancedPreviewScene)
+	{
+		if (!GetPreviewProfileController()->SetActiveProfile(UDefaultEditorProfiles::GreyAmbientProfileName.ToString()))
+		{
+			AdvancedPreviewScene->SetProfileIndex(GetDefault<UEditorPerProjectUserSettings>()->AssetViewerProfileIndex);
+		}
+		if (!AdvancedPreviewScene->IsPostProcessingEnabled())
+		{
+			AdvancedPreviewScene->HandleTogglePostProcessing();
+		}
+	}
+}
+
 SPVEditorViewport::~SPVEditorViewport()
 {
 	const FName ToolbarName = "PVE.ViewportToolbar";
@@ -108,6 +126,14 @@ void SPVEditorViewport::BindCommands()
 		FIsActionChecked::CreateSP(this, &SPVEditorViewport::IsAutoFocusViewportChecked));
 }
 
+TSharedRef<FEditorViewportClient> SPVEditorViewport::MakeEditorViewportClient()
+{
+	TSharedRef<FEditorViewportClient> ViewportClient = SPCGEditorViewport::MakeEditorViewportClient();
+	ViewportClient->ExposureSettings.bFixed = true;
+	ViewportClient->ExposureSettings.FixedEV100 = -0.2f;
+	return ViewportClient;
+}
+
 void SPVEditorViewport::InitVisualizationScene()
 {
 	const FString MannequinPath = TEXT("/ProceduralVegetationEditor/Mannequin/Viewport_Mannequin_T_Pose.Viewport_Mannequin_T_Pose");
@@ -119,9 +145,7 @@ void SPVEditorViewport::InitVisualizationScene()
 
 		ManagedResources.Add(MannequinComponent);
 		AdvancedPreviewScene->AddComponent(MannequinComponent, FTransform::Identity);
-		MannequinComponent->SetRelativeLocation(
-			FVector((-FocusBounds.BoxExtent.X - MannequinBounds.BoxExtent.X) * 1.5f - GetMannequinOffset(), 0.0, 0.0)
-		);
+		SetMannequinOffset((-FocusBounds.BoxExtent.X - MannequinBounds.BoxExtent.X) * 1.5f, true);
 	}
 
 	ScaleVisualizationComponent = NewObject<UPVScaleVisualizationComponent>(GetTransientPackage(), NAME_None, RF_Transient);
@@ -169,10 +193,10 @@ TSharedRef<SWidget> SPVEditorViewport::CreateMannequinOffsetWidget() const
 		.AllowSpin(true)
 		.AllowWheel(true)
 		.Delta(5.0f)
-		.MinValue(-100.0f)
-		.MaxValue(500.0f)
-		.MinSliderValue(-100.0f)
-		.MaxSliderValue(500.0f)
+		.MinValue(-2000.0f)
+		.MaxValue(2000.0f)
+		.MinSliderValue(-2000.0f)
+		.MaxSliderValue(2000.0f)
 		.Value_Lambda([this]{ return GetMannequinOffset(); })
 		.OnValueChanged_Lambda([this](const float NewValue)
 			{
@@ -257,15 +281,11 @@ void SPVEditorViewport::SetMannequinOffset(const float NewValue, const bool bSav
 	if (bSaveConfig)
 	{
 		EditorSettings->SaveConfig();
-		return;
 	}
 	
 	if (MannequinComponent)
 	{
-		const FBoxSphereBounds MannequinBounds = MannequinComponent->CalcLocalBounds();
-		MannequinComponent->SetRelativeLocation(
-			FVector((-FocusBounds.BoxExtent.X - MannequinBounds.BoxExtent.X) * 1.5f - EditorSettings->MannequinOffset, 0.0, 0.0)
-		);
+		MannequinComponent->SetRelativeLocation(FVector(EditorSettings->MannequinOffset, 0.0, 0.0));
 		Client->Invalidate();
 	}
 }
diff --git a/Engine/Plugins/Experimental/ProceduralVegetationEditor/Source/ProceduralVegetationEditor/Private/SPVEditorViewport.h b/Engine/Plugins/Experimental/ProceduralVegetationEditor/Source/ProceduralVegetationEditor/Private/SPVEditorViewport.h
--- a/Engine/Plugins/Experimental/ProceduralVegetationEditor/Source/ProceduralVegetationEditor/Private/SPVEditorViewport.h
+++ b/Engine/Plugins/Experimental/ProceduralVegetationEditor/Source/ProceduralVegetationEditor/Private/SPVEditorViewport.h
@@ -42,6 +42,7 @@ public:
 	SLATE_BEGIN_ARGS(SPVEditorViewport) {}
 	SLATE_END_ARGS()
 
+	SPVEditorViewport();
 	virtual ~SPVEditorViewport() override;
 
 	void Construct(const FArguments& InArgs);
@@ -56,6 +57,8 @@ protected:
 	virtual void BindCommands() override;
 	//~ End SEditorViewport Interface
 
+	virtual TSharedRef<FEditorViewportClient> MakeEditorViewportClient() override;
+
 private:
 	void InitVisualizationScene();
 	void ResetVisualizationScene();
diff --git a/Engine/Plugins/Experimental/ProceduralVegetationEditor/Source/ProceduralVegetationEditor/Private/Visualizations/PVDataVisualization.cpp b/Engine/Plugins/Experimental/ProceduralVegetationEditor/Source/ProceduralVegetationEditor/Private/Visualizations/PVDataVisualization.cpp
--- a/Engine/Plugins/Experimental/ProceduralVegetationEditor/Source/ProceduralVegetationEditor/Private/Visualizations/PVDataVisualization.cpp
+++ b/Engine/Plugins/Experimental/ProceduralVegetationEditor/Source/ProceduralVegetationEditor/Private/Visualizations/PVDataVisualization.cpp
@@ -66,6 +66,12 @@ FPCGSetupSceneFunc FPVDataVisualization::GetViewportSetupFunc(const UPCGSettings
 			UE_LOG(LogProceduralVegetationDataVisualization, Error, TEXT("Failed to setup data viewport, the data was lost or invalid."));
 			return;
 		}
+
+		InOutParams.Scene->SetFloorVisibility(true);
+		InOutParams.Scene->SetFloorOffset(1.0f);
+
+		// Reapply view mode settings as for some reason SetFloorVisibility changes EngineFlags
+		InOutParams.EditorViewportClient->SetViewMode(InOutParams.EditorViewportClient->GetViewMode());
 			
 		InOutParams.EditorViewportClient->SetRealtime(true);
 		if (!InOutParams.Scene->IsGridEnabled())
@@ -81,9 +87,6 @@ FPCGSetupSceneFunc FPVDataVisualization::GetViewportSetupFunc(const UPCGSettings
 			InOutParams.Scene->HandleTogglePostProcessing();
 		}
 
-		InOutParams.Scene->SetFloorVisibility(true);
-		InOutParams.Scene->SetFloorOffset(1.0f);
-			
 		const FManagedArrayCollection& Collection = WeakData->GetCollection();
 
 		FBoxSphereBounds Bounds(ForceInit);
@@ -202,7 +205,7 @@ void FPVDataVisualization::BonesRenderer(const FManagedArrayCollection& InCollec
 	InOutParams.Scene->AddComponent(BoneComponent, FTransform::Identity);
 
 	FString InTextToDraw = FString::Format(TEXT("{0} Bones"),{BoneComponent->GetBoneCount()});
-	const FVector3f& Pos = (FVector3f::ForwardVector * 100);
+	const FVector3f& Pos = FVector3f::ForwardVector * 100;
 
 	DrawDebugString(BoneComponent->GetWorld(), FVector::Zero(), InTextToDraw);
 	
@@ -221,106 +224,118 @@ void FPVDataVisualization::BonesRenderer(const FManagedArrayCollection& InCollec
 	OutBounds = OutBounds + BoneComponent->CalcLocalBounds();
 }
 
+TPair<int32, int32> GetMeshTrisAndVerts(const UObject* FoliageObject)
+{
+	int32 NumTriangles = 0;
+	int32 NumVertices = 0;
+	if (const UStaticMesh* StaticFoliageMesh = Cast<UStaticMesh>(FoliageObject))
+	{
+		NumTriangles = StaticFoliageMesh->GetNumTriangles(0);
+		NumVertices = StaticFoliageMesh->GetNumVertices(0);
+	}
+	else if (const USkeletalMesh* SkeletalFoliageMesh = Cast<USkeletalMesh>(FoliageObject))
+	{
+		FSkeletalMeshRenderData* SkelMeshRenderData = SkeletalFoliageMesh->GetResourceForRendering();
+		if (SkelMeshRenderData && SkelMeshRenderData->LODRenderData.Num() > 0)
+		{
+			NumTriangles = SkelMeshRenderData->LODRenderData[0].GetTotalFaces();
+			NumVertices = SkelMeshRenderData->LODRenderData[0].GetNumVertices();
+		}
+	}
+	return {NumTriangles, NumVertices};
+}
+
 void FPVDataVisualization::FoliageGridRenderer(const FManagedArrayCollection& InCollection, FPCGSceneSetupParams& InOutParams,
 	FBoxSphereBounds& OutBounds)
 {
-	PV::Facades::FFoliageFacade FoliageFacade(InCollection);
+	const PV::Facades::FFoliageFacade FoliageFacade(InCollection);
 
-	const int FoliageNum = FoliageFacade.NumFoliageNames();
+	const TArray<FString> FoliageNames = FoliageFacade.GetFoliageNames();
 
-	FBox MaxFoliageBounds;
-	float MaxFoliageVolume = std::numeric_limits<double>::lowest();
+	FBox MaxFoliageBounds(ForceInit);
 	int ValidFoliageEntries = 0;
-	
-	for (int32 FoliageNameIndex = 0; FoliageNameIndex < FoliageNum; ++FoliageNameIndex)
-	{
-		FString FoliageName = FoliageFacade.GetFoliageName(FoliageNameIndex);
 
-		if (UStaticMesh* StaticFoliageMesh = LoadObject<UStaticMesh>(nullptr, *FoliageName))
+	static const auto GetMeshBounds = [](auto* Mesh)-> FBox
+		{
+			const FBox BoundingBox = Mesh
+				? Mesh->GetBounds().GetBox()
+				: FBox(ForceInit);
+			return BoundingBox;
+		};
+
+	const auto AddComponentToScene = [&](USceneComponent* ComponentRef, const FTransform& Transform = FTransform::Identity, const bool UpdateBounds = true)
 		{
-			FBox BoundingBox = StaticFoliageMesh->GetBoundingBox();
+			InOutParams.ManagedResources.Add(ComponentRef);
+			InOutParams.Scene->AddComponent(ComponentRef, Transform);
 
-			if (BoundingBox.IsValid)
+			if (UpdateBounds)
 			{
-				if (BoundingBox.GetVolume() > MaxFoliageVolume)
-				{
-					MaxFoliageVolume = BoundingBox.GetVolume();
-					MaxFoliageBounds = BoundingBox;
-				}
+				ComponentRef->UpdateBounds();
+				OutBounds = OutBounds + ComponentRef->CalcBounds(Transform);
 			}
+		};
 
-			ValidFoliageEntries++;
-		}
-		else if (USkeletalMesh* SkeletalFoliageMesh = LoadObject<USkeletalMesh>(nullptr, *FoliageName))
+	for (const FString& FoliageName : FoliageNames)
+	{
+		if (UObject* FoliageMesh = LoadObject<UObject>(nullptr, *FoliageName))
 		{
-			FBox BoundingBox = SkeletalFoliageMesh->GetBounds().GetBox();
-
-			if (BoundingBox.IsValid)
+			FBox Bounds = GetMeshBounds(Cast<UStaticMesh>(FoliageMesh));
+			if (!Bounds.IsValid)
 			{
-				if (BoundingBox.GetVolume() > MaxFoliageVolume)
-				{
-					MaxFoliageVolume = BoundingBox.GetVolume();
-					MaxFoliageBounds = BoundingBox;
-				}
+				Bounds = GetMeshBounds(Cast<USkeletalMesh>(FoliageMesh));
 			}
-
+			MaxFoliageBounds = Bounds.IsValid && Bounds.GetVolume() > MaxFoliageBounds.GetVolume() ? Bounds : MaxFoliageBounds;
 			ValidFoliageEntries++;
 		}
 	}
-
-	FName FoliageName = NAME_None;
-	const FVector FoliageEntryBounds = MaxFoliageBounds.GetSize();
-	const float FoliageHorizontalPadding = 100;
-	const float FoliageVericalPadding = 50;
-	const float TextItemPadding = 1.1f;
-	const float BackgroundHorizontalPadding = 0.1f;
-	const float BackgroundVerticalPadding = 0.05f;
 	
-	FVector FoliageHorizontalSize = (FoliageEntryBounds + FoliageHorizontalPadding) * FVector::RightVector ;
-	FVector FoliageVerticalSize = (FoliageEntryBounds + FoliageVericalPadding) * FVector::UpVector;
-	FVector TotalHorizontalSize = ValidFoliageEntries/2.0f * FoliageHorizontalSize;
-
-	FVector FoliagePosOffset = FVector::ZeroVector - TotalHorizontalSize;
-	FVector FoliagePos = FoliagePosOffset;
-
 	UStaticMesh* TextBackgroundMesh = LoadObject<UStaticMesh>(nullptr,TEXT("/Engine/BasicShapes/Plane"));
 	UMaterial* BackgroundMat = LoadObject<UMaterial>(nullptr, TEXT("/ProceduralVegetationEditor/Materials/TextBackgroundMat.TextBackgroundMat"));
-
-	check(BackgroundMat && TextBackgroundMesh);
+	UMaterial* TextMat = LoadObject<UMaterial>(nullptr, TEXT("/ProceduralVegetationEditor/Materials/TextMaterial.TextMaterial"));
 	
-	for (int32 FoliageNameIndex = 0; FoliageNameIndex < FoliageNum; ++FoliageNameIndex)
-	{
-		FString FoliagePath = FoliageFacade.GetFoliageName(FoliageNameIndex);
+	const float FoliageHorizontalPadding = MaxFoliageBounds.GetSize().GetMax() + 100.0f;
+	const float FoliageVerticalPadding = MaxFoliageBounds.GetSize().GetMax() + 100.0f;
 
-		UStaticMesh* StaticFoliageMesh = LoadObject<UStaticMesh>(nullptr, *FoliagePath);
-		USkeletalMesh* SkeletalFoliageMesh = LoadObject<USkeletalMesh>(nullptr, *FoliagePath);
+	FVector FoliagePos = FVector(FoliageHorizontalPadding * FoliageNames.Num() / -2.0f, 0.0f, 0.0f);
 
-		const bool bHasValidMesh = StaticFoliageMesh || SkeletalFoliageMesh; 
-		
-		if (bHasValidMesh)
+	for (const FString& FoliagePath : FoliageNames)
+	{
+		UObject* FoliageObject = LoadObject<UObject>(nullptr, FoliagePath);
+		if (UStaticMesh* StaticFoliageMesh = Cast<UStaticMesh>(FoliageObject))
 		{
-			FoliagePos += FoliageHorizontalSize;
-			FoliageName = StaticFoliageMesh != nullptr ? StaticFoliageMesh->GetFName() : SkeletalFoliageMesh->GetFName();
-			
-			FText TriCount = FText::AsNumber(PV::Utilities::GetMeshTriangles(FoliagePath));
-			FString InTextToDraw = FString::Format(TEXT("Name: {0}\nTriangles:{1}"),{FoliageName.ToString(), TriCount.ToString()});
+			UStaticMeshComponent* MeshComponent = NewObject<UStaticMeshComponent>(GetTransientPackage(), NAME_None, RF_Transient);
+			MeshComponent->SetStaticMesh(StaticFoliageMesh);
+			AddComponentToScene(MeshComponent, FTransform(FoliagePos));
+		}
+		else if (USkeletalMesh* SkeletalFoliageMesh = Cast<USkeletalMesh>(FoliageObject))
+		{
+			USkeletalMeshComponent* MeshComponent = NewObject<USkeletalMeshComponent>(GetTransientPackage(), NAME_None, RF_Transient);
+			MeshComponent->SetSkeletalMesh(SkeletalFoliageMesh);
+			AddComponentToScene(MeshComponent, FTransform(FoliagePos));
+		}
 
+		if (FoliageObject)
+		{
+			const auto [Tris, Verts] = GetMeshTrisAndVerts(FoliageObject);
+			const FString InTextToDraw = FString::Printf(TEXT("Name: %s\nTris: %d\nVerts: %d"), *FoliageObject->GetName(), Tris, Verts);
+			
 			TObjectPtr<UTextRenderComponent> TextRenderComponent = NewObject<UTextRenderComponent>(GetTransientPackage(), NAME_None, RF_Transient);
+			TextRenderComponent->SetMaterial(0, TextMat);
 			TextRenderComponent->SetText(FText::FromString(InTextToDraw));
 			TextRenderComponent->SetTextRenderColor(FColor::Black);
 			TextRenderComponent->SetWorldSize(10);
+			TextRenderComponent->SetVerticalAlignment(EVRTA_TextTop);
 			TextRenderComponent->SetGenerateOverlapEvents(false);
-
-			FVector TextPos = FoliagePos;
-			TextPos += FoliageVerticalSize * TextItemPadding;
-			TextPos += FoliageHorizontalSize/2.0f - FoliageHorizontalSize/4.0f;
-
-			FTransform TextTransform = FTransform(FVector(TextPos));
 			
-			InOutParams.ManagedResources.Add(TextRenderComponent);
-			InOutParams.Scene->AddComponent(TextRenderComponent, TextTransform);
+			FBoxSphereBounds TextBounds = TextRenderComponent->CalcLocalBounds();
+			const FVector TextPos = FoliagePos + TextBounds.BoxExtent.Y * FVector::BackwardVector + FoliageVerticalPadding * FVector::UpVector;
+			AddComponentToScene(
+				TextRenderComponent,
+				FTransform(FRotator(0, 90, 0), TextPos),
+				false
+			);
 
-			OutBounds = OutBounds + TextRenderComponent->CalcBounds(TextTransform);
+			TextBounds = TextBounds.TransformBy(TextRenderComponent->GetComponentTransform());
 
 			if (TextBackgroundMesh && BackgroundMat)
 			{
@@ -328,50 +343,19 @@ void FPVDataVisualization::FoliageGridRenderer(const FManagedArrayCollection& In
 				TextBackgroundMeshComponent->SetStaticMesh(TextBackgroundMesh);
 				TextBackgroundMeshComponent->SetMaterial(0, BackgroundMat);
 
-				auto TextBackgroundBounds = TextBackgroundMeshComponent->CalcBounds(FTransform::Identity);
 
-				float TextBackgroundScale = FoliageHorizontalSize.Y/TextBackgroundBounds.BoxExtent.Y;
-				TextBackgroundScale /= 2;
-
-				FVector BackgroundPos = TextPos + FVector::BackwardVector * 0.1;
-				BackgroundPos -= FoliageHorizontalSize/2.0f;
-				BackgroundPos += FoliageHorizontalSize * BackgroundHorizontalPadding;
-				BackgroundPos += FoliageVerticalSize * BackgroundVerticalPadding;
-			
-				FTransform TextBackgroundTransform = FTransform(BackgroundPos);
-				FRotator TextRotator = FRotator(270,0,0);
-				TextBackgroundTransform.SetRotation(TextRotator.Quaternion());
-				TextBackgroundTransform.SetScale3D(FVector(0.4, TextBackgroundScale, 1));
-			
-				InOutParams.ManagedResources.Add(TextBackgroundMeshComponent);
-				InOutParams.Scene->AddComponent(
+				AddComponentToScene(
 					TextBackgroundMeshComponent,
-					FTransform(TextBackgroundTransform)
+					FTransform(
+						FRotator(0, 0, 90),
+						TextBounds.Origin + FVector::LeftVector * 0.01f,
+						FVector(TextBounds.BoxExtent.X / 40.0, TextBounds.BoxExtent.Z / 40.0, 1)
+					),
+					false
 				);
 			}
-		}
-		
-		if(StaticFoliageMesh)
-		{
-			TObjectPtr<UStaticMeshComponent> MeshComponent = NewObject<UStaticMeshComponent>(GetTransientPackage(), NAME_None, RF_Transient);
-			MeshComponent->SetStaticMesh(StaticFoliageMesh);
-			MeshComponent->UpdateBounds();
-
-			InOutParams.ManagedResources.Add(MeshComponent);
-			InOutParams.Scene->AddComponent(MeshComponent, FTransform(FoliagePos));
-
-			OutBounds = OutBounds + MeshComponent->CalcBounds(FTransform(FoliagePos));
-		}
-		else if(SkeletalFoliageMesh)
-		{	
-			TObjectPtr<USkeletalMeshComponent> MeshComponent = NewObject<USkeletalMeshComponent>(GetTransientPackage(), NAME_None, RF_Transient);
-			MeshComponent->SetSkeletalMesh(SkeletalFoliageMesh);
-			MeshComponent->UpdateBounds();
-
-			InOutParams.ManagedResources.Add(MeshComponent);
-			InOutParams.Scene->AddComponent(MeshComponent, FTransform(FoliagePos));
-
-			OutBounds = OutBounds + MeshComponent->CalcBounds(FTransform(FoliagePos));
+			
+			FoliagePos.X += FoliageHorizontalPadding;
 		}
 	}
 }
diff --git a/Engine/Plugins/Experimental/ProceduralVegetationEditor/Source/ProceduralVegetationEditor/Private/Visualizations/PVScaleVisualizationComponent.cpp b/Engine/Plugins/Experimental/ProceduralVegetationEditor/Source/ProceduralVegetationEditor/Private/Visualizations/PVScaleVisualizationComponent.cpp
--- a/Engine/Plugins/Experimental/ProceduralVegetationEditor/Source/ProceduralVegetationEditor/Private/Visualizations/PVScaleVisualizationComponent.cpp
+++ b/Engine/Plugins/Experimental/ProceduralVegetationEditor/Source/ProceduralVegetationEditor/Private/Visualizations/PVScaleVisualizationComponent.cpp
@@ -70,7 +70,8 @@ void UPVScaleVisualizationComponent::UpdateScaleVisualizations()
 	Flush();
 
 	const auto WorldToMetersScale = GetWorld()->GetWorldSettings()->WorldToMeters;
-	const double TextSize = (Bounds.BoxExtent.GetAbsMax() * 20.0f) / WorldToMetersScale;
+	const double TextSize = (Bounds.BoxExtent.X + Bounds.BoxExtent.Y + Bounds.BoxExtent.Z) * 20.0f / 3.0f / WorldToMetersScale;
+	
 	{
 		FVector StartPosition = Bounds.Origin;
 		StartPosition.X -= Bounds.BoxExtent.X;
@@ -88,7 +89,7 @@ void UPVScaleVisualizationComponent::UpdateScaleVisualizations()
 		TextRenderComponent->SetWorldLocation(EndPosition);
 		TextRenderComponent->SetWorldRotation(FRotator(0, 90, 0));
 		TextRenderComponent->SetTextRenderColor(FLinearColor::Red.ToFColorSRGB());
-		TextRenderComponent->SetHorizontalAlignment(EHTA_Right);
+		TextRenderComponent->SetHorizontalAlignment(EHTA_Left);
 		TextRenderComponent->SetVerticalAlignment(EVRTA_TextBottom);
 		TextRenderComponent->SetWorldSize(TextSize);
 		TextRenderComponent->SetText(FText::FromString(FString::Printf(TEXT("%.2fm"), Length)));
@@ -111,7 +112,7 @@ void UPVScaleVisualizationComponent::UpdateScaleVisualizations()
 		const TObjectPtr<UTextRenderComponent> TextRenderComponent = GetTextRenderComponent(1);
 		TextRenderComponent->SetWorldLocation(EndPosition);
 		TextRenderComponent->SetTextRenderColor(FColor::Green);
-		TextRenderComponent->SetHorizontalAlignment(EHTA_Left);
+		TextRenderComponent->SetHorizontalAlignment(EHTA_Right);
 		TextRenderComponent->SetVerticalAlignment(EVRTA_TextBottom);
 		TextRenderComponent->SetWorldSize(TextSize);
 		TextRenderComponent->SetText(FText::FromString(FString::Printf(TEXT("%.2fm"), Length)));
