Address BugHawk issues in PVE Slope Node

diff --git a/Engine/Plugins/Experimental/ProceduralVegetationEditor/Source/ProceduralVegetation/Private/Implementations/PVSlope.cpp b/Engine/Plugins/Experimental/ProceduralVegetationEditor/Source/ProceduralVegetation/Private/Implementations/PVSlope.cpp
--- a/Engine/Plugins/Experimental/ProceduralVegetationEditor/Source/ProceduralVegetation/Private/Implementations/PVSlope.cpp
+++ b/Engine/Plugins/Experimental/ProceduralVegetationEditor/Source/ProceduralVegetation/Private/Implementations/PVSlope.cpp
@@ -19,42 +19,18 @@ namespace PVSlope
 		return YawRotation * LeanRotation;
 	}
 
-	TArray<int32> GetBranchChildIndices(int32 BranchIndex, const PV::Facades::FBranchFacade& BranchFacade)
-	{
-		const TArray<int32>& CurrentBranchChildren = BranchFacade.GetChildren(BranchIndex);
-		const TManagedArray<int32>& BranchNumbers = BranchFacade.GetBranchNumbers();
-
-		TArray<int32> ChildrenIndices;
-		ChildrenIndices.Reserve(CurrentBranchChildren.Num());
-
-		for (const int32 Child : CurrentBranchChildren)
-		{
-			const int32 ChildIndex = BranchNumbers.Find(Child);
-			const TArray<int32>& Parents = BranchFacade.GetParents(ChildIndex);
-
-			if (ChildIndex != INDEX_NONE && !Parents.IsEmpty() && Parents[Parents.Num() - 1] == BranchNumbers[BranchIndex])
-			{
-				ChildrenIndices.Add(ChildIndex);
-			}
-		}
-
-		return ChildrenIndices;
-	}
-
 	TArray<int32> GetBranchSegmentChildren(
 		const PV::Facades::FBranchFacade& BranchFacade, 
 		const PV::Facades::FPointFacade& PointFacade,
 		int32 BranchIndex, 
-		int32 BranchPointIndex)
+		int32 BranchPointIndex,
+		int32 PreviousBranchPointIndex)
 	{
-		const TArray<int32>& BranchPoints = BranchFacade.GetPoints(BranchIndex);
 		const TArray<int32>& BranchChildren = BranchFacade.GetChildren(BranchIndex);
 		const TManagedArray<int32>& BranchNumbers = BranchFacade.GetBranchNumbers();
 
-		check(BranchPoints.Find(BranchPointIndex) != INDEX_NONE);
-
 		const float BranchPointLengthFromRoot = PointFacade.GetLengthFromRoot(BranchPointIndex);
-		const float PreviousBranchPointLengthFromRoot = PointFacade.GetLengthFromRoot(BranchPointIndex - 1);
+		const float PreviousBranchPointLengthFromRoot = PointFacade.GetLengthFromRoot(PreviousBranchPointIndex);
 
 		TArray<int32> ChildrenIndices;
 		ChildrenIndices.Reserve(BranchChildren.Num());
@@ -67,13 +43,19 @@ namespace PVSlope
 				continue;
 			}
 
-			const TArray<int32>& Parents = BranchFacade.GetParents(ChildIndex);
-			if (Parents.IsEmpty() || Parents[Parents.Num() - 1] != BranchNumbers[BranchIndex])
+			const int32 ParentIndex = BranchFacade.GetParentIndex(ChildIndex);
+			if (ParentIndex != BranchIndex)
+			{
+				continue;
+			}
+
+			const TArray<int32>& ChildBranchPoints = BranchFacade.GetPoints(ChildIndex);
+			if (ChildBranchPoints.Num() == 0)
 			{
 				continue;
 			}
 
-			const int32 FirstChildPointIndex = BranchFacade.GetPoints(ChildIndex)[0];
+			const int32 FirstChildPointIndex = ChildBranchPoints[0];
 			const float ChildLengthFromRoot = PointFacade.GetLengthFromRoot(FirstChildPointIndex);
 			if (ChildLengthFromRoot > BranchPointLengthFromRoot || ChildLengthFromRoot <= PreviousBranchPointLengthFromRoot)
 			{
@@ -90,12 +72,13 @@ namespace PVSlope
 		PV::Facades::FFoliageFacade& FoliageFacade,
 		const PV::Facades::FPointFacade& PointFacade,
 		int32 BranchIndex,
-		int32 BranchPointIndex)
+		int32 BranchPointIndex,
+		int32 PreviousBranchPointIndex)
 	{
 		TArray<int32> OutFoliageEntryIds;
 
 		const float BranchPointLengthFromRoot = PointFacade.GetLengthFromRoot(BranchPointIndex);
-		const float PreviousBranchPointLengthFromRoot = PointFacade.GetLengthFromRoot(BranchPointIndex - 1);
+		const float PreviousBranchPointLengthFromRoot = PointFacade.GetLengthFromRoot(PreviousBranchPointIndex);
 
 		TArray<int32> FoliageEntryIds = FoliageFacade.GetFoliageEntryIdsForBranch(BranchIndex);
 		OutFoliageEntryIds.Reserve(FoliageEntryIds.Num());
@@ -126,6 +109,7 @@ namespace PVSlope
 		for (int32 Index = 0; Index < CurrentBranchPoints.Num(); ++Index)
 		{
 			const int32 PointIndex = CurrentBranchPoints[Index];
+			const int32 PrevPointIndex = Index > 0 ? CurrentBranchPoints[Index - 1] : INDEX_NONE;
 			const FVector3f& PointPosition = PointPositions[PointIndex];
 
 			const float PointDistSquared = (PointPosition - BranchStartPosition).SizeSquared();
@@ -134,7 +118,7 @@ namespace PVSlope
 				OutLengthSquared = PointDistSquared;
 			}
 
-			const TArray<int32> BranchSegmentChildren = GetBranchSegmentChildren(BranchFacade, PointFacade, BranchIndex, PointIndex);
+			const TArray<int32> BranchSegmentChildren = GetBranchSegmentChildren(BranchFacade, PointFacade, BranchIndex, PointIndex, PrevPointIndex);
 			for (const int32 ChildIndex : BranchSegmentChildren)
 			{
 				CalcBranchLengthSquared_Recursive(BranchFacade, PointFacade, ChildIndex, BranchStartPosition, OutLengthSquared);
@@ -184,6 +168,7 @@ namespace PVSlope
 		for (int32 Index = 0; Index < CurrentBranchPoints.Num(); ++Index)
 		{
 			const int32 PointIndex = CurrentBranchPoints[Index];
+			const int32 PrevPointIndex = Index > 0 ? CurrentBranchPoints[Index - 1] : INDEX_NONE;
 			FVector3f& Position = PointPositions[PointIndex];
 
 			const float TreeAngleAlpha = 1.0 - (TreeHeight > 0 ? ((Position - TrunkPosition).Size() / TreeHeight) : 0);
@@ -202,7 +187,7 @@ namespace PVSlope
 				CurrentBudDirections[i] = SlopeRotation.RotateVector(CurrentBudDirections[i]);
 			}
 
-			const TArray<int32> BranchSegmentFoliageEntryIds = GetBranchSegmentFoliageEntryIds(FoliageFacade, PointFacade, BranchIndex, PointIndex);
+			const TArray<int32> BranchSegmentFoliageEntryIds = GetBranchSegmentFoliageEntryIds(FoliageFacade, PointFacade, BranchIndex, PointIndex, PrevPointIndex);
 			for (int32 FoliageEntryId : BranchSegmentFoliageEntryIds)
 			{
 				const FVector3f FoliageEntryPivotPoint = FoliageFacade.GetPivotPoint(FoliageEntryId);
@@ -224,7 +209,7 @@ namespace PVSlope
 			Position = NewPosition;
 			ParentPosition = Position;
 
-			const TArray<int32> BranchSegmentChildren = GetBranchSegmentChildren(BranchFacade, PointFacade, BranchIndex, PointIndex);
+			const TArray<int32> BranchSegmentChildren = GetBranchSegmentChildren(BranchFacade, PointFacade, BranchIndex, PointIndex, PrevPointIndex);
 			for (const int32 ChildIndex : BranchSegmentChildren)
 			{
 				ApplySlope_Recursive(
diff --git a/Engine/Plugins/Experimental/ProceduralVegetationEditor/Source/ProceduralVegetation/Private/Implementations/PVSlope.h b/Engine/Plugins/Experimental/ProceduralVegetationEditor/Source/ProceduralVegetation/Private/Implementations/PVSlope.h
--- a/Engine/Plugins/Experimental/ProceduralVegetationEditor/Source/ProceduralVegetation/Private/Implementations/PVSlope.h
+++ b/Engine/Plugins/Experimental/ProceduralVegetationEditor/Source/ProceduralVegetation/Private/Implementations/PVSlope.h
@@ -38,7 +38,7 @@ struct FPVSlopeParams
 
 	/** The pivot point of the tunk. */
 	UPROPERTY(EditAnywhere, Category = "Slope")
-	EPVSlopeTrunkPivotPoint TrunkPivotPoint;
+	EPVSlopeTrunkPivotPoint TrunkPivotPoint = EPVSlopeTrunkPivotPoint::Origin;
 };
 
 struct FPVSlope
