[PVE] Fix issue where Mesher segment reduction settings are too sensitive to PScale and curvature
Also fix potential bug in Foliage distributor node

diff --git a/Engine/Plugins/Experimental/ProceduralVegetationEditor/Source/ProceduralVegetation/Private/Implementations/PVFoliage.cpp b/Engine/Plugins/Experimental/ProceduralVegetationEditor/Source/ProceduralVegetation/Private/Implementations/PVFoliage.cpp
--- a/Engine/Plugins/Experimental/ProceduralVegetationEditor/Source/ProceduralVegetation/Private/Implementations/PVFoliage.cpp
+++ b/Engine/Plugins/Experimental/ProceduralVegetationEditor/Source/ProceduralVegetation/Private/Implementations/PVFoliage.cpp
@@ -102,6 +102,12 @@ void FPVFoliage::DistributeFoliage(FManagedArrayCollection& OutCollection, const
 			TArray<FVector3f> AttachmentPointsNormalVectors;
 			TArray<float> AttachmentPointsScales;
 			TArray<int32> BranchPoints = BranchFacadeSource.GetPoints(BranchIndex);
+
+			if (BranchPoints.Num() < 2)
+			{
+				continue;
+			}
+			
 			const float FirstPointLengthFromRoot = PointFacadeSource.GetLengthFromRoot(BranchPoints[0]);
 			const float LastPointLengthFromRoot = PointFacadeSource.GetLengthFromRoot(BranchPoints.Last());
 			const float BranchLength = LastPointLengthFromRoot - FirstPointLengthFromRoot;
@@ -119,7 +125,7 @@ void FPVFoliage::DistributeFoliage(FManagedArrayCollection& OutCollection, const
 				}
 
 				float InstanceSpacingRampValue = DistributionSettings.InstanceSpacingRamp
-					? DistributionSettings.InstanceSpacingRamp->GetRichCurveConst()->Eval(LookUp)
+					? DistributionSettings.InstanceSpacingRamp->GetRichCurveConst()->Eval(1.0f - LookUp)
 					: Increment;
 				// The increment is nudged slightly before blending with the weighted ramp value, the number is chosen based on observation
 				float AdjustedIncrement = Increment + (InstanceSpacingRampValue * (Increment * 10.0f));
@@ -199,7 +205,7 @@ void FPVFoliage::DistributeFoliage(FManagedArrayCollection& OutCollection, const
 						const float BranchGradient = FMath::Clamp(
 							FMath::Lerp(CurrentPointBranchGradient, NextPointBranchGradient, DistanceAlpha), 0.0f, 1.0f);
 						float BranchGradientRampValue = ScaleSettings.ScaleRamp
-							? ScaleSettings.ScaleRamp->GetRichCurveConst()->Eval(BranchGradient)
+							? ScaleSettings.ScaleRamp->GetRichCurveConst()->Eval(1.0f - BranchGradient)
 							: 1.0f;
 
 						const float CurrentPointScale = PointFacadeSource.GetPointScale(CurrentPointIndex);
@@ -229,7 +235,7 @@ void FPVFoliage::DistributeFoliage(FManagedArrayCollection& OutCollection, const
 						int32 BudsToAdd = 1;
 
 						// Compute Up Vector
-						// Last point
+						// Last point (order of the array is last to first)
 						if (Idx == 0)
 						{
 							AttachmentPointsUpVectors.Add(BudApicalDirectionVector);
@@ -256,7 +262,7 @@ void FPVFoliage::DistributeFoliage(FManagedArrayCollection& OutCollection, const
 								{
 									float AxilAngleRampValue = VectorSettings.AxilAngleRamp
 										? FMath::Clamp(
-											VectorSettings.AxilAngleRamp->GetRichCurveConst()->Eval(BranchGradient),
+											VectorSettings.AxilAngleRamp->GetRichCurveConst()->Eval(1.0f - BranchGradient),
 											0.0f,
 											1.0f)
 										: 1.0f;
diff --git a/Engine/Plugins/Experimental/ProceduralVegetationEditor/Source/ProceduralVegetation/Private/Implementations/PVMeshBuilder.cpp b/Engine/Plugins/Experimental/ProceduralVegetationEditor/Source/ProceduralVegetation/Private/Implementations/PVMeshBuilder.cpp
--- a/Engine/Plugins/Experimental/ProceduralVegetationEditor/Source/ProceduralVegetation/Private/Implementations/PVMeshBuilder.cpp
+++ b/Engine/Plugins/Experimental/ProceduralVegetationEditor/Source/ProceduralVegetation/Private/Implementations/PVMeshBuilder.cpp
@@ -407,7 +407,7 @@ TSet<int32> FPVMeshBuilder::ComputePointGradients(const PV::Facades::FPointFacad
 		float DeltaModifier = 1.0f;
 		if (!HardPoints.Contains(PointIndex))
 		{
-			DeltaModifier = MeshBuilderParams.SegmentRetentionImpact * MaxPointScale * DeltaModifierGradient;
+			DeltaModifier = MeshBuilderParams.SegmentRetentionImpact * DeltaModifierGradient;
 		}
 		OutDeltaModifiers.Add(DeltaModifier);
 	}
@@ -415,6 +415,57 @@ TSet<int32> FPVMeshBuilder::ComputePointGradients(const PV::Facades::FPointFacad
 	return PointsToRemove;
 }
 
+float FPVMeshBuilder::GetMaxDeltaBetweenHardPoints(const PV::Facades::FBranchFacade& BranchFacade, const PV::Facades::FPointFacade& PointFacade,
+                                                   const TSet<int32>& HardPoints)
+{
+	float MaxDeltaBetweenHardPoints = 0.0f;
+
+	for (int32 BranchIndex = 0; BranchIndex < BranchFacade.GetElementCount(); ++BranchIndex)
+	{
+		const TArray<int32>& BranchPoints = BranchFacade.GetPoints(BranchIndex);
+
+		if (BranchPoints.Num() < 2)
+		{
+			continue;
+		}
+
+		TArray<int32> IndicesOfHardPointsForBranch;
+		IndicesOfHardPointsForBranch.Add(0);
+
+		for (int32 Idx = 1; Idx < BranchPoints.Num() - 1; ++Idx)
+		{
+			if (HardPoints.Contains(BranchPoints[Idx]))
+			{
+				IndicesOfHardPointsForBranch.Add(Idx);
+			}
+		}
+
+		IndicesOfHardPointsForBranch.Add(BranchPoints.Num() - 1);
+
+		for (int32 Idx = 0; Idx < IndicesOfHardPointsForBranch.Num() - 1; ++Idx)
+		{
+			const int32 Pt0Index = BranchPoints[IndicesOfHardPointsForBranch[Idx]];
+			const int32 Pt2Index = BranchPoints[IndicesOfHardPointsForBranch[Idx + 1]];
+			const FVector Pt0Position = FVector(PointFacade.GetPosition(Pt0Index));
+			const FVector Pt2Position = FVector(PointFacade.GetPosition(Pt2Index));
+
+			for (int32 k = IndicesOfHardPointsForBranch[Idx]; k < IndicesOfHardPointsForBranch[Idx + 1]; ++k)
+			{
+				const int32 Pt1Index = BranchPoints[k];
+				const FVector Pt1Position = FVector(PointFacade.GetPosition(Pt1Index));
+
+				if (const float Delta = FMath::PointDistToLine(Pt1Position, Pt2Position - Pt0Position, Pt0Position);
+					Delta > MaxDeltaBetweenHardPoints)
+				{
+					MaxDeltaBetweenHardPoints = Delta;
+				}
+			}
+		}
+	}
+
+	return MaxDeltaBetweenHardPoints;
+}
+
 void FPVMeshBuilder::PerformPathSimplification(const PV::Facades::FBranchFacade& BranchFacade,
                                                const PV::Facades::FPointFacade& PointFacade, const FPVMeshBuilderParams& MeshBuilderParams,
                                                const float MaxPointScale, const TSet<int32>& HardPoints, const TArray<float>& DeltaModifiers,
@@ -422,7 +473,8 @@ void FPVMeshBuilder::PerformPathSimplification(const PV::Facades::FBranchFacade&
 {
 	const int32 Iterations = MeshBuilderParams.Accuracy;
 	// Epsilon is the maximum allowed curvature error
-	const float Epsilon = MeshBuilderParams.SegmentReduction / static_cast<float>(Iterations);
+	const float Epsilon = FMath::Pow(MeshBuilderParams.SegmentReduction, 2) / static_cast<float>(Iterations);
+	const float MaxDeltaBetweenHardPoints = GetMaxDeltaBetweenHardPoints(BranchFacade, PointFacade, HardPoints);
 
 	for (int32 BranchIndex = 0; BranchIndex < BranchFacade.GetElementCount(); ++BranchIndex)
 	{
@@ -434,7 +486,7 @@ void FPVMeshBuilder::PerformPathSimplification(const PV::Facades::FBranchFacade&
 
 		for (int32 j = 0; j < Iterations; ++j)
 		{
-			const float IterationEpsilon = Epsilon * (j + 1) * MaxPointScale;
+			const float IterationEpsilon = Epsilon * (j + 1);
 			for (int32 i = 1; i < BranchPoints.Num() - 1; ++i)
 			{
 				const int32 Pt0Index = BranchPoints[i - 1],
@@ -444,10 +496,14 @@ void FPVMeshBuilder::PerformPathSimplification(const PV::Facades::FBranchFacade&
 				              Pt1Position = FVector(PointFacade.GetPosition(Pt1Index)),
 				              Pt2Position = FVector(PointFacade.GetPosition(Pt2Index));
 				const float DeltaModifier = DeltaModifiers[Pt1Index];
-				const float Delta =
-					FMath::PointDistToLine(Pt1Position, Pt2Position - Pt0Position, Pt0Position) + DeltaModifier;
+				const float Delta = FMath::PointDistToLine(Pt1Position, Pt2Position - Pt0Position, Pt0Position);
+
+				const float NormalizedDelta = FMath::GetMappedRangeValueClamped(
+					FVector2f(0.0f, MaxDeltaBetweenHardPoints),
+					FVector2f(0.0f, 1.0f),
+					Delta);
 
-				if (Delta < IterationEpsilon && !HardPoints.Contains(Pt1Index))
+				if (NormalizedDelta + DeltaModifier < IterationEpsilon && !HardPoints.Contains(Pt1Index))
 				{
 					InOutPointsToRemove.Add(Pt1Index);
 					BranchPoints.Remove(Pt1Index);
diff --git a/Engine/Plugins/Experimental/ProceduralVegetationEditor/Source/ProceduralVegetation/Private/Implementations/PVMeshBuilder.h b/Engine/Plugins/Experimental/ProceduralVegetationEditor/Source/ProceduralVegetation/Private/Implementations/PVMeshBuilder.h
--- a/Engine/Plugins/Experimental/ProceduralVegetationEditor/Source/ProceduralVegetation/Private/Implementations/PVMeshBuilder.h
+++ b/Engine/Plugins/Experimental/ProceduralVegetationEditor/Source/ProceduralVegetation/Private/Implementations/PVMeshBuilder.h
@@ -46,7 +46,7 @@ struct FPVMeshBuilderParams
 	float PointRemoval = 0.0f;
 
 	UPROPERTY(EditAnywhere, Category="Mesh Settings", DisplayName="Segment Reduction",
-		meta=(ClampMin=0.0f, ClampMax=1.0f, UIMin=0.0f, UIMax=0.1f, Tooltip=
+		meta=(ClampMin=0.0f, ClampMax=1.0f, UIMin=0.0f, UIMax=1.0f, Tooltip=
 			"Controls how aggressively segments are reduced.\n\nSets the strength of segment merging/decimation. Higher values reduce edge count for lighter meshes; lower values preserve detail."
 		))
 	float SegmentReduction = 0.0f;
@@ -82,7 +82,7 @@ struct FPVMeshBuilderParams
 	float ShortestSegmentLength = 0.0f;
 
 	UPROPERTY(EditAnywhere, Category="Mesh Settings|Advanced", DisplayName="Segment Retention Impact",
-		meta=(ClampMin=0.0f, ClampMax=5.0f, UIMin=0.0f, UIMax=2.0f, Tooltip=
+		meta=(ClampMin=0.0f, ClampMax=5.0f, UIMin=0.0f, UIMax=1.0f, Tooltip=
 			"Control the impact of reduced segments based on other criteria\n\nSegment retention impact allows you to reintroduce points that have been reduced based on user defined criteria\'s. See; Hull Retention, Ground Retention and Trunk Retention."
 		))
 	float SegmentRetentionImpact = 1.0f;
@@ -313,6 +313,9 @@ private:
 	                                         const TSet<int32>& HardPoints, const float MaxPointScale, TArray<float>& OutMeshDivisionsGradients,
 	                                         TArray<float>& OutDeltaModifiers);
 
+	static float GetMaxDeltaBetweenHardPoints(const PV::Facades::FBranchFacade& BranchFacade, const PV::Facades::FPointFacade& PointFacade,
+	                                          const TSet<int32>& HardPoints);
+
 	static void PerformPathSimplification(const PV::Facades::FBranchFacade& BranchFacade,
 	                                      const PV::Facades::FPointFacade& PointFacade, const FPVMeshBuilderParams& MeshBuilderParams,
 	                                      const float MaxPointScale, const TSet<int32>& HardPoints, const TArray<float>& DeltaModifiers,
