[PVE] Physics Asset Added for Skeletal mesh export

diff --git a/Engine/Plugins/Experimental/ProceduralVegetationEditor/Source/ProceduralVegetation/Private/Facades/PVBoneFacade.cpp b/Engine/Plugins/Experimental/ProceduralVegetationEditor/Source/ProceduralVegetation/Private/Facades/PVBoneFacade.cpp
--- a/Engine/Plugins/Experimental/ProceduralVegetationEditor/Source/ProceduralVegetation/Private/Facades/PVBoneFacade.cpp
+++ b/Engine/Plugins/Experimental/ProceduralVegetationEditor/Source/ProceduralVegetation/Private/Facades/PVBoneFacade.cpp
@@ -247,7 +247,7 @@ namespace PV::Facades
 		SetBranchSimulationGroupIndex(PointFacade, BranchFacade);
 	}
 	
-	bool FBoneFacade::CreateBoneData(TArray<FBoneNode>& Bones,const float ReductionStrength)
+	bool FBoneFacade::CreateBoneData(TArray<FBoneNode>& Bones,const float ReductionStrength,bool bSetWindData /*= true*/)
 	{
 		FPointFacade PointFacade(VertexPointIds.GetConstCollection());
 		FBranchFacade BranchFacade(VertexPointIds.GetConstCollection());
@@ -257,8 +257,11 @@ namespace PV::Facades
 		{
 			return false;	
 		}
-		
-		SetWindSimulationData(*VertexPointIds.GetCollection());
+
+		if (bSetWindData)
+		{
+			SetWindSimulationData(*VertexPointIds.GetCollection());
+		}
 		
 		const TManagedArray<FVector3f>& Positions = PointFacade.GetPositions();
 
diff --git a/Engine/Plugins/Experimental/ProceduralVegetationEditor/Source/ProceduralVegetation/Public/Facades/PVBoneFacade.h b/Engine/Plugins/Experimental/ProceduralVegetationEditor/Source/ProceduralVegetation/Public/Facades/PVBoneFacade.h
--- a/Engine/Plugins/Experimental/ProceduralVegetationEditor/Source/ProceduralVegetation/Public/Facades/PVBoneFacade.h
+++ b/Engine/Plugins/Experimental/ProceduralVegetationEditor/Source/ProceduralVegetation/Public/Facades/PVBoneFacade.h
@@ -52,7 +52,7 @@ namespace PV::Facades
 
 		void SetWindSimulationData(FManagedArrayCollection& Collection);
 		
-		bool CreateBoneData(TArray<FBoneNode>& Bones,const float ReductionStrength);
+		bool CreateBoneData(TArray<FBoneNode>& Bones,const float ReductionStrength,bool bSetWindData = true);
 
 		static FBoneNode* FindClosestBone(const FManagedArrayCollection& Collection, TArray<FBoneNode>& Bones, const int32 PointID);
 
diff --git a/Engine/Plugins/Experimental/ProceduralVegetationEditor/Source/ProceduralVegetation/Public/PVExportParams.h b/Engine/Plugins/Experimental/ProceduralVegetationEditor/Source/ProceduralVegetation/Public/PVExportParams.h
--- a/Engine/Plugins/Experimental/ProceduralVegetationEditor/Source/ProceduralVegetation/Public/PVExportParams.h
+++ b/Engine/Plugins/Experimental/ProceduralVegetationEditor/Source/ProceduralVegetation/Public/PVExportParams.h
@@ -27,6 +27,19 @@ enum class EPVAssetReplacePolicy : uint8
 	Ignore,
 };
 
+UENUM()
+enum class EPVCollisionGeneration : uint8
+{
+	/** Do not create collision for mesh */
+	None,
+
+	/** Creates collision for trunk only */
+	TrunkOnly,
+
+	/** Creates collision for all branch generations */
+	AllGenerations,
+};
+
 USTRUCT()
 struct PROCEDURALVEGETATION_API FPVExportParams
 {
@@ -56,6 +69,9 @@ struct PROCEDURALVEGETATION_API FPVExportParams
 	UPROPERTY(EditAnywhere, Category = "Mesh | Nanite", Meta = (DisplayName = "Nanite Shape Preservation Method", Tooltip="Shape preservation setting to be used in nanite settings."))
 	ENaniteShapePreservation NaniteShapePreservation = ENaniteShapePreservation::Voxelize;
 
+	UPROPERTY(EditAnywhere, Category = "Mesh | Collision", Meta = (EditCondition = "ExportMeshType == EPVExportMeshType::SkeletalMesh", Tooltip="Creates the physic asset for skeletal mesh."))
+	EPVCollisionGeneration CollisionGeneration;
+
 	UPROPERTY(EditAnywhere, Category = "Mesh | Dynamic Wind", Meta = (EditCondition = "ExportMeshType == EPVExportMeshType::SkeletalMesh", Tooltip="Wind Settings contains the wind simulation group data, Default presets are available for Saplings and Trees."))
 	TObjectPtr<UPVWindSettings> WindSettings;
 
@@ -90,6 +106,21 @@ struct PROCEDURALVEGETATION_API FPVExportParams
 		return FPaths::Combine(ContentBrowserFolder.Path, GetOutputSkeletonName());
 	}
 
+	FString GetOutputPhysicsAssetName() const
+	{
+		return FString::Printf(TEXT("%s_Physics"), *MeshName.ToString());
+	}
+
+	FString GetOutputPhysicsAssetPackagePath() const
+	{
+		return FPaths::Combine(ContentBrowserFolder.Path, GetOutputPhysicsAssetName());
+	}
+	
+	bool IsCollisionEnable() const
+	{
+		return CollisionGeneration != EPVCollisionGeneration::None;
+	}
+
 	UClass* GetMeshClass() const;
 
 	bool Validate(FString& OutError) const;
diff --git a/Engine/Plugins/Experimental/ProceduralVegetationEditor/Source/ProceduralVegetationEditor/Private/Helpers/PVExportHelper.cpp b/Engine/Plugins/Experimental/ProceduralVegetationEditor/Source/ProceduralVegetationEditor/Private/Helpers/PVExportHelper.cpp
--- a/Engine/Plugins/Experimental/ProceduralVegetationEditor/Source/ProceduralVegetationEditor/Private/Helpers/PVExportHelper.cpp
+++ b/Engine/Plugins/Experimental/ProceduralVegetationEditor/Source/ProceduralVegetationEditor/Private/Helpers/PVExportHelper.cpp
@@ -4,6 +4,7 @@
 
 #include "DynamicWindImportData.h"
 #include "DynamicWindSkeletalData.h"
+#include "EditorSupportDelegates.h"
 #include "EngineAnalytics.h"
 #include "IAssetTools.h"
 #include "MeshDescriptionToDynamicMesh.h"
@@ -35,6 +36,7 @@
 #include "Engine/StaticMesh.h"
 
 #include "Facades/PVBranchFacade.h"
+#include "Facades/PVPlantFacade.h"
 #include "Facades/PVPointFacade.h"
 
 #include "GeometryCollection/GeometryCollection.h"
@@ -48,6 +50,10 @@
 #include "Helpers/PVUtilities.h"
 
 #include "Misc/Paths.h"
+#include "PhysicsEngine/ConstraintUtils.h"
+#include "PhysicsEngine/PhysicsAsset.h"
+#include "PhysicsEngine/PhysicsConstraintTemplate.h"
+#include "PhysicsEngine/SkeletalBodySetup.h"
 
 #include "UObject/Package.h"
 
@@ -977,10 +983,6 @@ namespace PV::Export::Internal
 				return false;
 			}
 
-			Collection.CopyTo(&ExportedCollection);
-
-			DynamicMesh = CollectionToDynamicMesh(ExportedCollection);
-
 			if (!UpdateStage(EExportMeshStage::InitializingExportMesh, 0.5f))
 			{
 				return false;
@@ -1050,7 +1052,7 @@ namespace PV::Export::Internal
 			{
 				return false;
 			}
-
+			
 			ExportMesh->Build(true);
 			FStaticMeshCompilingManager::Get().FinishAllCompilation();
 		}
@@ -1077,8 +1079,10 @@ namespace PV::Export::Internal
 		const FManagedArrayCollection& Collection,
 		ENaniteShapePreservation InShapePreservation,
 		TObjectPtr<const UPVWindSettings> InWindSettings,
+		TObjectPtr<UPhysicsAsset> InPhysicsAsset,
 		bool bBuildNaniteAssemblies,
 		bool bShouldExportFoliage,
+		EPVCollisionGeneration CollisionGeneration,
 		TArray<UPackage*>& OutModifiedPackages,
 		const FStatusReportCallback& StatusReportCallback
 	)
@@ -1312,6 +1316,11 @@ namespace PV::Export::Internal
 			}
 		}
 
+		if (InPhysicsAsset)
+		{
+			SetupPhysicAsset(Collection, CollisionGeneration, InPhysicsAsset, BoneNodes);
+		}
+
 		{
 			TRACE_CPUPROFILER_EVENT_SCOPE(DynamicWind);
 
@@ -1494,6 +1503,86 @@ namespace PV::Export::Internal
 				
 		}, EDynamicMeshChangeType::GeneralEdit, EDynamicMeshAttributeChangeFlags::Unknown, false);
 	}
+
+	void SetupPhysicAsset(const FManagedArrayCollection& Collection, const EPVCollisionGeneration& CollisionGeneration,UPhysicsAsset* PhysicsAsset, TArray<Facades::FBoneNode> InBones)
+	{
+		if (!PhysicsAsset)
+		{
+			UE_LOG(LogProceduralVegetationEditor, Warning, TEXT("No physics asset specified to create physics bodies"));
+			return;
+		}
+
+		TArray<TObjectPtr<USkeletalBodySetup>> SkeletalBodySetups;
+		
+		TArray<Facades::FBoneNode> ReducedBones;
+		Facades::FBoneFacade BoneFacade = Facades::FBoneFacade(Collection);
+		Facades::FPointFacade PointFacade = Facades::FPointFacade(Collection);
+		Facades::FPlantFacade PlantFacade = Facades::FPlantFacade(Collection);
+		
+		BoneFacade.CreateBoneData(ReducedBones, 1.0, false);
+
+		for (const Facades::FBoneNode&  Node : ReducedBones)
+		{
+			USkeletalBodySetup* BodySetup = NewObject<USkeletalBodySetup>(PhysicsAsset);
+				
+			auto MappedBone = InBones.FindByPredicate([Node](const Facades::FBoneNode& IterNode){ return Node.PointIndex == IterNode.PointIndex; });
+
+			if (!MappedBone)
+			{
+				UE_LOG(LogProceduralVegetationEditor, Log, TEXT("MappedBone not found for BodySetup"));
+				continue;
+			}
+
+			if (!InBones.IsValidIndex(MappedBone->ParentBoneIndex))
+			{
+				UE_LOG(LogProceduralVegetationEditor, Log, TEXT("MapedBone not found for BodySetup"));
+				continue;
+			}
+
+			bool bTrunk = PlantFacade.IsTrunkIndex(Node.BranchIndex);
+
+			if (bTrunk || CollisionGeneration == EPVCollisionGeneration::AllGenerations)
+			{
+				BodySetup->BoneName = MappedBone->BoneName;
+				
+				FKSphylElem Element;
+				
+				FVector BasePosition = Node.AbsolutePosition;
+				float Radius = PointFacade.GetPointScale(Node.PointIndex);
+				if (ReducedBones.IsValidIndex(Node.ParentBoneIndex))
+				{
+					BasePosition = ReducedBones[Node.ParentBoneIndex].AbsolutePosition;
+					Radius = PointFacade.GetPointScale(Node.ParentBoneIndex);
+				}
+				
+				float Length = FVector::Dist(Node.AbsolutePosition ,  BasePosition);
+				Element.Length = Length;
+				Element.Radius = Radius * 1.25f;
+				
+				FVector Direction = (BasePosition - Node.AbsolutePosition);
+				Direction.Normalize();
+
+				FVector Axis = FVector::UpVector;
+				FQuat Quat = FQuat::FindBetween(Axis, Direction);
+
+				Element.Center = Direction * Length * 0.5;
+				Element.Rotation = Quat.Rotator();
+
+				if (Element.Length > 0)
+				{
+					BodySetup->AggGeom.SphylElems.Add(Element);
+					BodySetup->CollisionTraceFlag = CTF_UseSimpleAsComplex;
+					BodySetup->CreatePhysicsMeshes();
+					SkeletalBodySetups.Add(BodySetup);
+				}
+			}
+		}
+
+		PhysicsAsset->SkeletalBodySetups = SkeletalBodySetups;
+		PhysicsAsset->UpdateBodySetupIndexMap();
+		PhysicsAsset->UpdateBoundsBodiesArray();
+		FEditorSupportDelegates::RedrawAllViewports.Broadcast();
+	}
 }
 
 namespace PV::Export
@@ -1626,6 +1715,30 @@ namespace PV::Export
 			
 			ExportMesh->SetSkeleton(ExportSkeleton);
 
+			TObjectPtr<UPhysicsAsset> PhysicsAsset;
+			FString PhysicsAssetName = ExportParams.GetOutputPhysicsAssetName();
+			FString PhysicsPackageName = ExportParams.GetOutputPhysicsAssetPackagePath();
+
+			if (ExportParams.IsCollisionEnable())
+			{
+				Internal::EGetValidAssetPathResult PhysicsAssetPathResult = Internal::GetValidAssetPath(ExportParams.ReplacePolicy, UPhysicsAsset::StaticClass(), PhysicsPackageName, PhysicsAssetName);
+				if (PhysicsAssetPathResult != Internal::EGetValidAssetPathResult::Success)
+				{
+					return ExitWithAssetPathResult(PhysicsAssetPathResult);
+				}
+				
+				TObjectPtr<UPackage> PhysicsPackage = Internal::CreateAssetPackage(PhysicsPackageName);
+				
+				if (TObjectPtr<UPhysicsAsset> ExistingPhysicAsset = Cast<UPhysicsAsset>(PhysicsPackage->FindAssetInPackage()))
+				{
+					ExistingPhysicAsset->SetExternalPackage(nullptr);
+					ExistingPhysicAsset->MarkAsGarbage();
+				}
+				
+				PhysicsAsset = NewObject<UPhysicsAsset>(PhysicsPackage, *PhysicsAssetName, RF_Standalone | RF_Public);
+				ModifiedPackages.Add(PhysicsAsset->GetOutermost());
+			}
+
 			ModifiedPackages.Add(ExportMesh->GetOutermost());
 			ModifiedPackages.Add(ExportSkeleton->GetOutermost());
 
@@ -1634,8 +1747,10 @@ namespace PV::Export
 				Collection,
 				ExportParams.NaniteShapePreservation,
 				ExportParams.WindSettings,
+				PhysicsAsset,
 				bCreateNaniteFoliage,
 				bShouldExportFoliage,
+				ExportParams.CollisionGeneration,
 				ModifiedPackages,
 				StatusReportCallback
 			))
@@ -1646,6 +1761,12 @@ namespace PV::Export
 			OutCreatedAssets.Add(MeshPackageName);
 			OutCreatedAssets.Add(SkeletonPackageName);
 
+			if(PhysicsAsset)
+			{
+				PhysicsAsset->PreviewSkeletalMesh = ExportMesh;
+				ExportMesh->SetPhysicsAsset(PhysicsAsset);
+				Internal::AttachProceduralVegetationLink(PhysicsAsset, InProceduralVegetation);
+			}
 			Internal::AttachProceduralVegetationLink(ExportMesh, InProceduralVegetation);
 			Internal::AttachProceduralVegetationLink(ExportSkeleton, InProceduralVegetation);
 		}
diff --git a/Engine/Plugins/Experimental/ProceduralVegetationEditor/Source/ProceduralVegetationEditor/Private/Helpers/PVExportHelper.h b/Engine/Plugins/Experimental/ProceduralVegetationEditor/Source/ProceduralVegetationEditor/Private/Helpers/PVExportHelper.h
--- a/Engine/Plugins/Experimental/ProceduralVegetationEditor/Source/ProceduralVegetationEditor/Private/Helpers/PVExportHelper.h
+++ b/Engine/Plugins/Experimental/ProceduralVegetationEditor/Source/ProceduralVegetationEditor/Private/Helpers/PVExportHelper.h
@@ -5,6 +5,7 @@
 #include "Facades/PVBoneFacade.h"
 #include "Engine/EngineTypes.h"
 
+enum class EPVCollisionGeneration : uint8;
 class UInstancedStaticMeshComponent;
 class UStaticMesh;
 class USkeletalMesh;
@@ -89,8 +90,10 @@ namespace PV::Export::Internal
 		const FManagedArrayCollection& Collection,
 		ENaniteShapePreservation InShapePreservation,
 		TObjectPtr<const UPVWindSettings> InWindSettings,
+		TObjectPtr<UPhysicsAsset> InPhysicsAsset,
 		bool bBuildNaniteAssemblies,
 		bool bShouldExportFoliage,
+		EPVCollisionGeneration CollisionGeneration,
 		TArray<UPackage*>& OutModifiedPackages,
 		const FStatusReportCallback& StatusReportCallback
 	);
@@ -102,6 +105,8 @@ namespace PV::Export::Internal
 	void RemoveUnwantedSkinWeights(TObjectPtr<UDynamicMesh> DynamicMesh, const FName ProfileName, const TArray<int32>& VertexBoneIDs,
 		int32 FoliageVertexStart, const TArray<Facades::FBoneNode>& BoneNodes);
 
+	void SetupPhysicAsset(const FManagedArrayCollection& Collection,const EPVCollisionGeneration& CollisionGeneration, UPhysicsAsset* PhysicsAsset, TArray<Facades::FBoneNode> InBones);
+
 	void ConvertToDefaultSkeletalMesh(USkeletalMesh* SkeletalMesh,
 		FDynamicMesh3& Mesh,
 		const TArray<TObjectPtr<UMaterialInterface>>& Materials
diff --git a/Engine/Plugins/Experimental/ProceduralVegetationEditor/Source/ProceduralVegetationEditor/Private/PVEditor.cpp b/Engine/Plugins/Experimental/ProceduralVegetationEditor/Source/ProceduralVegetationEditor/Private/PVEditor.cpp
--- a/Engine/Plugins/Experimental/ProceduralVegetationEditor/Source/ProceduralVegetationEditor/Private/PVEditor.cpp
+++ b/Engine/Plugins/Experimental/ProceduralVegetationEditor/Source/ProceduralVegetationEditor/Private/PVEditor.cpp
@@ -476,6 +476,11 @@ FReply FPVEditor::ExportInternal()
 			if (ExportParams.ExportMeshType == EPVExportMeshType::SkeletalMesh)
 			{
 				ExportPaths.Add(ExportParams.GetOutputSkeletonPackagePath());
+
+				if (ExportParams.IsCollisionEnable())
+				{
+					ExportPaths.Add(ExportParams.GetOutputPhysicsAssetPackagePath());
+				}
 			}
 		}
 
