[PVE] Overlay stats in PVE viewport

diff --git a/Engine/Plugins/Experimental/ProceduralVegetationEditor/Source/ProceduralVegetationEditor/Private/PVEditor.cpp b/Engine/Plugins/Experimental/ProceduralVegetationEditor/Source/ProceduralVegetationEditor/Private/PVEditor.cpp
--- a/Engine/Plugins/Experimental/ProceduralVegetationEditor/Source/ProceduralVegetationEditor/Private/PVEditor.cpp
+++ b/Engine/Plugins/Experimental/ProceduralVegetationEditor/Source/ProceduralVegetationEditor/Private/PVEditor.cpp
@@ -3,6 +3,9 @@
 #include "PVEditor.h"
 
 #include "EngineAnalytics.h"
+#include "IMessageLogListing.h"
+#include "MessageLogModule.h"
+#include "PackagesDialog.h"
 #include "PCGDefaultExecutionSource.h"
 #include "PCGEditorCommands.h"
 #include "ProceduralVegetation.h"
@@ -13,24 +16,23 @@
 #include "SPVEditorViewport.h"
 #include "ToolMenus.h"
 #include "ToolMenuSection.h"
-#include "PackagesDialog.h"
-#include "ProceduralVegetationEditorModule.h"
 
 #include "Dataflow/DataflowCollectionSpreadSheetWidget.h"
 
 #include "DataTypes/PVData.h"
 #include "DataTypes/PVFoliageMeshData.h"
-#include "Helpers/PVAnalyticsHelper.h"
 
+#include "Facades/PVAttributesNames.h"
+
+#include "GeometryCollection/GeometryCollection.h"
+
+#include "Helpers/PVAnalyticsHelper.h"
 #include "Helpers/PVExportHelper.h"
 #include "Helpers/PVUtilities.h"
 
 #include "Misc/MessageDialog.h"
 #include "Misc/UObjectToken.h"
 
-#include "MessageLogModule.h"
-#include "IMessageLogListing.h"
-
 #include "Nodes/PCGEditorGraphNodeBase.h"
 #include "Nodes/PVOutputSettings.h"
 
@@ -305,16 +307,21 @@ UPVBaseSettings* GetNodeSettings(UPCGEditorGraphNodeBase* InNode)
 	return nullptr;
 }
 
-void FPVEditor:: ChangeNodeInspection(UPCGEditorGraphNodeBase* InNode)
+void FPVEditor::ChangeNodeInspection(UPCGEditorGraphNodeBase* InNode)
 {
-	if (auto Settings = GetNodeSettings(InNode))
+	if (UPVBaseSettings* Settings = GetNodeSettings(InNode))
 	{
-		if (ViewportWidgets.Num() > 0)
+		if (ViewportWidgets.Num() == 0)
 		{
-			ViewportWidgets[0]->OnNodeInspectionChanged(Settings);
+			return;
 		}
 
+		ViewportWidgets[0]->OnNodeInspectionChanged(Settings);
+		
 		NodeBeingInspected = InNode;
+		CollectionBeingInspected = SelectedCollection;
+		
+		UpdateCollectionSpreadSheet(InNode->GetPCGNode());
 		
 		SetNodeInspected(InNode, true);	
 	}
@@ -340,15 +347,20 @@ void FPVEditor::OnSelectedNodesChanged(const TSet<UObject*>& InNewSelection)
 
 	if (UPCGNode* SelectedNode = GetSelectedNode() ? GetSelectedNode()->GetPCGNode() : nullptr)
 	{
-		UpdateCollectionSpreadSheet(SelectedNode);
+		SelectedCollection = MakeShared<FManagedArrayCollection>();
+		FName OutputName;
+		GetCollectionFromNode(SelectedNode, OutputName, *SelectedCollection);
 		
 		if (UPVBaseSettings* Settings = Cast<UPVBaseSettings>(SelectedNode->GetSettings()))
 		{
 			const UPVBaseSettings* InspectionSettings = GetNodeSettings(NodeBeingInspected);
-				
 			if (!InspectionSettings || !InspectionSettings->bLockInspection)
 			{
 				ChangeNodeInspection(EdNode);
+
+				TArray<FText> OverlayStats;
+				GatherStats(OverlayStats);
+				ViewportWidgets[0]->PopulateStatsOverlayText(OverlayStats);
 			}
 		}	
 	}
@@ -754,8 +766,8 @@ FReply FPVEditor::ExportInternal()
 
 void FPVEditor::OnLockNodeSelection()
 {
-	auto NodeBeingInspectedSettings = GetNodeSettings(NodeBeingInspected);
-	auto SelectedNodeSettings = GetNodeSettings(GetSelectedNode());
+	UPVBaseSettings* NodeBeingInspectedSettings = GetNodeSettings(NodeBeingInspected);
+	UPVBaseSettings* SelectedNodeSettings = GetNodeSettings(GetSelectedNode());
 
 	if (NodeBeingInspectedSettings)
 	{
@@ -770,6 +782,22 @@ void FPVEditor::OnLockNodeSelection()
 			ChangeNodeInspection(GetSelectedNode());
 		}
 	}
+	if (ViewportWidgets.IsEmpty())
+	{
+		return;
+	}
+
+	NodeBeingInspectedSettings = GetNodeSettings(NodeBeingInspected);
+
+	TArray<FText> OverlayStats;
+	GatherStats(OverlayStats);
+	ViewportWidgets[0]->PopulateStatsOverlayText(OverlayStats);
+	
+	ViewportWidgets[0]->SetOverlayText(
+		NodeBeingInspectedSettings->bLockInspection
+		? NodeBeingInspected->GetNodeTitle(ENodeTitleType::FullTitle)
+		: FText::GetEmpty()
+	);
 }
 
 TSharedRef<SPCGEditorViewport> FPVEditor::CreateViewportWidget()
@@ -874,27 +902,82 @@ TSharedRef<SDockTab> FPVEditor::SpawnTab_CollectionSpreadSheet(const FSpawnTabAr
 	return DockableTab;
 }
 
-void FPVEditor::UpdateCollectionSpreadSheet(const UPCGNode* InSelectedNode) const
+void FPVEditor::UpdateCollectionSpreadSheet(const UPCGNode* InSelectedNode)
 {
-	if (CollectionSpreadSheetTab.IsValid() && CollectionSpreadSheetTab.Pin()->GetVisibility().IsVisible() && CollectionSpreadSheetWidget)
+	const bool bIsCollectionSpreadSheetVisible = CollectionSpreadSheetTab.IsValid() &&
+		CollectionSpreadSheetTab.Pin()->GetVisibility().IsVisible() &&
+		CollectionSpreadSheetWidget;
+
+	if (InSelectedNode)
 	{
-		if (InSelectedNode)
+		if (bIsCollectionSpreadSheetVisible)
 		{
-			FManagedArrayCollection OutCollection;
-			FName OutputName;
-			GetCollectionFromNode(InSelectedNode, OutputName, OutCollection);
-			
-			TSharedPtr<const FManagedArrayCollection> CopiedCollectionPtr = MakeShared<const FManagedArrayCollection>(OutCollection);
 			CollectionSpreadSheetWidget->GetCollectionTable()->GetCollectionInfoMap().Empty();
-			CollectionSpreadSheetWidget->GetCollectionTable()->GetCollectionInfoMap().Add(OutputName.ToString(), { CopiedCollectionPtr });
+			CollectionSpreadSheetWidget->GetCollectionTable()->GetCollectionInfoMap().Add("Out", {CollectionBeingInspected});
 			CollectionSpreadSheetWidget->SetData(InSelectedNode->GetNodeTitle(EPCGNodeTitleType::FullTitle).ToString());
 		}
-		else
+	}
+	else if (bIsCollectionSpreadSheetVisible)
+	{
+		CollectionSpreadSheetWidget->SetData(FString());
+	}
+	
+	if (bIsCollectionSpreadSheetVisible)
+	{
+		CollectionSpreadSheetWidget->RefreshWidget();
+	}
+}
+
+void FPVEditor::GatherStats(TArray<FText>& OverlayStats)
+{
+	if (CollectionBeingInspected->HasGroup(PV::GroupNames::PointGroup))
+	{
+		const int32 NumPoints = CollectionBeingInspected->NumElements(PV::GroupNames::PointGroup);
+		OverlayStats.Add(FText::FromString("Points: " + FString::FromInt(NumPoints)));
+	}
+	if (CollectionBeingInspected->HasGroup(PV::GroupNames::BranchGroup))
+	{
+		const int32 NumBranches = CollectionBeingInspected->NumElements(PV::GroupNames::BranchGroup);
+		OverlayStats.Add(FText::FromString("Branches: " + FString::FromInt(NumBranches)));
+	}
+	if (CollectionBeingInspected->HasGroup(FGeometryCollection::VerticesGroup))
+	{
+		const int32 NumVertices = CollectionBeingInspected->NumElements(FGeometryCollection::VerticesGroup);
+		OverlayStats.Add(FText::FromString("Vertices: " + FString::FromInt(NumVertices)));
+	}
+	if (CollectionBeingInspected->HasGroup(FGeometryCollection::FacesGroup))
+	{
+		const int32 NumTriangles = CollectionBeingInspected->NumElements(FGeometryCollection::FacesGroup);
+		OverlayStats.Add(FText::FromString("Triangles: " + FString::FromInt(NumTriangles)));
+	}
+	if (CollectionBeingInspected->HasGroup(PV::GroupNames::BonesGroup))
+	{
+		const int32 NumBones = CollectionBeingInspected->NumElements(PV::GroupNames::BonesGroup);
+		OverlayStats.Add(FText::FromString("Bones: " + FString::FromInt(NumBones)));
+	}
+	if (CollectionBeingInspected->HasGroup(PV::GroupNames::FoliageNamesGroup))
+	{
+		const int32 NumUniqueFoliage = CollectionBeingInspected->NumElements(PV::GroupNames::FoliageNamesGroup);
+		OverlayStats.Add(FText::FromString("Unique Foliage: " + FString::FromInt(NumUniqueFoliage)));
+	}
+	if (CollectionBeingInspected->HasGroup(PV::GroupNames::FoliageGroup))
+	{
+		const int32 NumFoliageInstances = CollectionBeingInspected->NumElements(PV::GroupNames::FoliageGroup);
+		OverlayStats.Add(FText::FromString("Foliage Instances: " + FString::FromInt(NumFoliageInstances)));
+	}
+	if (CollectionBeingInspected->HasGroup(FGeometryCollection::GeometryGroup))
+	{
+		const TManagedArray<FBox>& BoundingBoxes = CollectionBeingInspected->GetAttribute<FBox>(
+			FGeometryCollection::BoundingBoxAttribute,
+			FGeometryCollection::GeometryGroup
+		);
+		FBox BoundingBox;
+		for (int32 i = 0; i < BoundingBoxes.Num(); ++i)
 		{
-			CollectionSpreadSheetWidget->SetData(FString());
+			BoundingBox += BoundingBoxes[i];
 		}
-
-		CollectionSpreadSheetWidget->RefreshWidget();	
+		const FIntVector3 Extent = FIntVector3(BoundingBox.GetExtent() * 2.0f);
+		OverlayStats.Add(FText::FromString(FString::Printf(TEXT("Approx Size (without Foliage): %dx%dx%d"), Extent.X, Extent.Y, Extent.Z)));
 	}
 }
 
diff --git a/Engine/Plugins/Experimental/ProceduralVegetationEditor/Source/ProceduralVegetationEditor/Private/PVEditor.h b/Engine/Plugins/Experimental/ProceduralVegetationEditor/Source/ProceduralVegetationEditor/Private/PVEditor.h
--- a/Engine/Plugins/Experimental/ProceduralVegetationEditor/Source/ProceduralVegetationEditor/Private/PVEditor.h
+++ b/Engine/Plugins/Experimental/ProceduralVegetationEditor/Source/ProceduralVegetationEditor/Private/PVEditor.h
@@ -55,6 +55,8 @@ private:
 	TObjectPtr<UProceduralVegetation> ProceduralVegetationBeingEdited = nullptr;
 
 	TArray<TObjectPtr<UPCGEditorGraphNodeBase>> SelectedNodes;
+	TSharedPtr<FManagedArrayCollection> CollectionBeingInspected = nullptr;
+	TSharedPtr<FManagedArrayCollection> SelectedCollection = nullptr;
 	TObjectPtr<UPCGEditorGraphNodeBase> NodeBeingInspected = nullptr;
 	
 	TArray<TSharedRef<SPVEditorViewport>> ViewportWidgets;
@@ -68,7 +70,9 @@ private:
 
 	TSharedRef<SDockTab> SpawnTab_CollectionSpreadSheet(const FSpawnTabArgs& SpawnTabArgs);
 	
-	void UpdateCollectionSpreadSheet(const UPCGNode* InSelectedNode) const;
+	void UpdateCollectionSpreadSheet(const UPCGNode* InSelectedNode);
+
+	void GatherStats(TArray<FText>& OverlayStats);
 	
 	void OnExport_Clicked();
 
diff --git a/Engine/Plugins/Experimental/ProceduralVegetationEditor/Source/ProceduralVegetationEditor/Private/SPVEditorViewport.cpp b/Engine/Plugins/Experimental/ProceduralVegetationEditor/Source/ProceduralVegetationEditor/Private/SPVEditorViewport.cpp
--- a/Engine/Plugins/Experimental/ProceduralVegetationEditor/Source/ProceduralVegetationEditor/Private/SPVEditorViewport.cpp
+++ b/Engine/Plugins/Experimental/ProceduralVegetationEditor/Source/ProceduralVegetationEditor/Private/SPVEditorViewport.cpp
@@ -1,32 +1,42 @@
 // Copyright Epic Games, Inc. All Rights Reserved.
 
 #include "SPVEditorViewport.h"
+
 #include "AdvancedPreviewScene.h"
 #include "AdvancedPreviewSceneMenus.h"
 #include "PCGEditorSettings.h"
 #include "PVEditorCommands.h"
 #include "PVEditorSettings.h"
-#include "ViewportToolbar/UnrealEdViewportToolbar.h"
 #include "ToolMenu.h"
 #include "ToolMenus.h"
 #include "ToolMenuSection.h"
-#include "Engine/StaticMesh.h"
+
+#include "Brushes/SlateColorBrush.h"
+
 #include "Components/StaticMeshComponent.h"
 #include "Components/TextRenderComponent.h"
 
-#include "Visualizations/PVScaleVisualizationComponent.h"
 #include "DataTypes/PVData.h"
 
 #include "Editor/EditorPerProjectUserSettings.h"
 
+#include "Engine/StaticMesh.h"
+
 #include "Nodes/PVBaseSettings.h"
 
+#include "ViewportToolbar/UnrealEdViewportToolbar.h"
+
+#include "Visualizations/PVScaleVisualizationComponent.h"
+
 #include "Widgets/Input/SNumericEntryBox.h"
+#include "Widgets/Text/SRichTextBlock.h"
 
 #define LOCTEXT_NAMESPACE "PVEditorViewportToolbarSections"
 
 SPVEditorViewport::SPVEditorViewport()
 {
+	PreviewNodeBackgroundBrush = MakeShared<FSlateColorBrush>(FColor{70, 100, 200});
+	
 	if (AdvancedPreviewScene)
 	{
 		if (!GetPreviewProfileController()->SetActiveProfile(UDefaultEditorProfiles::GreyAmbientProfileName.ToString()))
@@ -134,6 +144,54 @@ TSharedRef<FEditorViewportClient> SPVEditorViewport::MakeEditorViewportClient()
 	return ViewportClient;
 }
 
+void SPVEditorViewport::PopulateViewportOverlays(TSharedRef<SOverlay> Overlay)
+{
+	SPCGEditorViewport::PopulateViewportOverlays(Overlay);
+
+	Overlay
+		->AddSlot()
+		.VAlign(VAlign_Top)
+		.HAlign(HAlign_Fill)
+		[
+			SNew(SBorder)
+			.BorderImage(PreviewNodeBackgroundBrush.Get())
+			.IsEnabled_Lambda([this]
+				{
+					return bIsPreviewingLockedNode;
+				})
+			.Visibility_Lambda([this]
+				{
+					return bIsPreviewingLockedNode
+						? EVisibility::Visible
+						: EVisibility::Collapsed;
+				})
+			[
+				SAssignNew(OverlayText, SRichTextBlock)
+				.Justification(ETextJustify::Center)
+			]
+		];
+	
+	Overlay
+		->AddSlot()
+		.VAlign(VAlign_Top)
+		.HAlign(HAlign_Left)
+		.Padding(6.0f, 25.f)
+		[
+			SNew(SBorder)
+			.BorderImage(FAppStyle::Get().GetBrush("FloatingBorder"))
+			.Visibility_Lambda([this]
+				{
+					return StatsOverlayText && !StatsOverlayText->GetText().IsEmpty()
+						? EVisibility::Visible
+						: EVisibility::Collapsed;
+				})
+			.Padding(4.f)
+			[
+				SAssignNew(StatsOverlayText, SRichTextBlock)
+			]
+		];
+}
+
 void SPVEditorViewport::InitVisualizationScene()
 {
 	const FString MannequinPath = TEXT("/ProceduralVegetationEditor/Mannequin/Viewport_Mannequin_T_Pose.Viewport_Mannequin_T_Pose");
@@ -282,7 +340,7 @@ void SPVEditorViewport::SetMannequinOffset(const float NewValue, const bool bSav
 	{
 		EditorSettings->SaveConfig();
 	}
-	
+
 	if (MannequinComponent)
 	{
 		MannequinComponent->SetRelativeLocation(FVector(EditorSettings->MannequinOffset, 0.0, 0.0));
@@ -547,4 +605,26 @@ void SPVEditorViewport::OnNodeInspectionChanged(UPVBaseSettings* InSettings)
 	}
 }
 
+void SPVEditorViewport::SetOverlayText(const FText& CurrentlyLockedNodeName)
+{
+	bIsPreviewingLockedNode = !CurrentlyLockedNodeName.IsEmpty();
+	static FText NormalTextStyle = FText::FromString(TEXT("<TextBlock.ShadowedText>Previewing Output of Node: {0}</>"));
+	FTextBuilder FinalText;
+	FinalText.AppendLineFormat(NormalTextStyle, CurrentlyLockedNodeName);
+	OverlayText->SetText(FinalText.ToText());
+}
+
+void SPVEditorViewport::PopulateStatsOverlayText(const TArrayView<FText> TextItems)
+{
+	static FText NormalTextStyle = FText::FromString(TEXT("<TextBlock.ShadowedText>{0}</>"));
+	
+	FTextBuilder FinalText;
+	for (const FText& TextItem : TextItems)
+	{
+		FinalText.AppendLineFormat(NormalTextStyle, TextItem);
+	}
+	
+	StatsOverlayText->SetText(FinalText.ToText());
+}
+
 #undef LOCTEXT_NAMESPACE
diff --git a/Engine/Plugins/Experimental/ProceduralVegetationEditor/Source/ProceduralVegetationEditor/Private/SPVEditorViewport.h b/Engine/Plugins/Experimental/ProceduralVegetationEditor/Source/ProceduralVegetationEditor/Private/SPVEditorViewport.h
--- a/Engine/Plugins/Experimental/ProceduralVegetationEditor/Source/ProceduralVegetationEditor/Private/SPVEditorViewport.h
+++ b/Engine/Plugins/Experimental/ProceduralVegetationEditor/Source/ProceduralVegetationEditor/Private/SPVEditorViewport.h
@@ -5,6 +5,7 @@
 #include "CoreMinimal.h"
 
 #include "Widgets/AssetEditorViewport/SPCGEditorViewport.h"
+#include "Widgets/Text/SRichTextBlock.h"
 
 #include "SPVEditorViewport.generated.h"
 
@@ -39,9 +40,6 @@ public:
 class SPVEditorViewport : public SPCGEditorViewport
 {
 public:
-	SLATE_BEGIN_ARGS(SPVEditorViewport) {}
-	SLATE_END_ARGS()
-
 	SPVEditorViewport();
 	virtual ~SPVEditorViewport() override;
 
@@ -51,6 +49,9 @@ public:
 
 	void OnNodeInspectionChanged(UPVBaseSettings* InSettings);
 
+	void SetOverlayText(const FText& CurrentlyLockedNodeName = FText::GetEmpty());
+	void PopulateStatsOverlayText(const TArrayView<FText> TextItems);
+
 protected:
 	//~ Begin SEditorViewport Interface
 	virtual TSharedPtr<SWidget> BuildViewportToolbar() override;
@@ -58,6 +59,7 @@ protected:
 	//~ End SEditorViewport Interface
 
 	virtual TSharedRef<FEditorViewportClient> MakeEditorViewportClient() override;
+	virtual void PopulateViewportOverlays(TSharedRef<SOverlay> Overlay) override;
 
 private:
 	void InitVisualizationScene();
@@ -96,4 +98,10 @@ protected:
 private:
 	TObjectPtr<UStaticMeshComponent> MannequinComponent = nullptr;
 	TObjectPtr<UPVScaleVisualizationComponent> ScaleVisualizationComponent = nullptr;
+
+	bool bIsPreviewingLockedNode = false;
+	TSharedPtr<SRichTextBlock> OverlayText;
+	TSharedPtr<SRichTextBlock> StatsOverlayText;
+
+	TSharedPtr<FSlateBrush> PreviewNodeBackgroundBrush;
 };
