[PVE] Fixes
- Fixed warning when trying to Load mesh object
- Disabled phototropic bias
- Updated mesh profile fall off and scale behavior

diff --git a/Engine/Plugins/Experimental/ProceduralVegetationEditor/Source/ProceduralVegetation/Private/Helpers/PVUtilities.cpp b/Engine/Plugins/Experimental/ProceduralVegetationEditor/Source/ProceduralVegetation/Private/Helpers/PVUtilities.cpp
--- a/Engine/Plugins/Experimental/ProceduralVegetationEditor/Source/ProceduralVegetation/Private/Helpers/PVUtilities.cpp
+++ b/Engine/Plugins/Experimental/ProceduralVegetationEditor/Source/ProceduralVegetation/Private/Helpers/PVUtilities.cpp
@@ -70,9 +70,10 @@ namespace PV::Utilities
 			for (int32 FoliageNameIndex = 0; FoliageNameIndex < InFoliageFacade.NumFoliageNames(); ++FoliageNameIndex)
 			{
 				FString FoliageName = InFoliageFacade.GetFoliageName(FoliageNameIndex);
+				FSoftObjectPath MeshPath(FoliageName);
 
-				UStaticMesh* StaticFoliageMesh = LoadObject<UStaticMesh>(nullptr, *FoliageName);
-				USkeletalMesh* SkeletalFoliageMesh = LoadObject<USkeletalMesh>(nullptr, *FoliageName);
+				UStaticMesh* StaticFoliageMesh = PackageExists(MeshPath.GetLongPackageName(), UStaticMesh::StaticClass()) ? LoadObject<UStaticMesh>(nullptr, *FoliageName) : nullptr;
+				USkeletalMesh* SkeletalFoliageMesh = PackageExists(MeshPath.GetLongPackageName(), USkeletalMesh::StaticClass()) ? LoadObject<USkeletalMesh>(nullptr, *FoliageName) : nullptr;
 
 				if(StaticFoliageMesh)
 				{
@@ -218,6 +219,35 @@ namespace PV::Utilities
 		return false;
 	}
 
+	bool PackageExists(const FString& LongPackageName, UClass* AssetClass)
+	{
+		UPackage* Package = FindPackage(nullptr, *LongPackageName);
+		if (IsValid(Package))
+		{
+			UObject* ExistingAsset = Package->FindAssetInPackage();
+			return ExistingAsset && ExistingAsset->GetClass() == AssetClass;
+		}
+		else
+		{
+			TArray<FAssetData> OutAssets;
+			IAssetRegistry& AssetRegistry = FAssetRegistryModule::GetRegistry();
+			AssetRegistry.GetAssetsByPackageName(*LongPackageName, OutAssets, true);
+
+			if (OutAssets.Num() == 0)
+			{
+				return false;
+			}
+
+			if (OutAssets.Num() != 1
+				|| OutAssets[0].AssetClassPath == AssetClass->GetClassPathName())
+			{
+				return true;
+			}
+		}
+
+		return false;
+	}
+
 	int32 GetMeshTriangles(const FString InMeshPath)
 	{
 		UStaticMesh* StaticFoliageMesh = LoadObject<UStaticMesh>(nullptr, *InMeshPath);
diff --git a/Engine/Plugins/Experimental/ProceduralVegetationEditor/Source/ProceduralVegetation/Private/Implementations/PVGravity.h b/Engine/Plugins/Experimental/ProceduralVegetationEditor/Source/ProceduralVegetation/Private/Implementations/PVGravity.h
--- a/Engine/Plugins/Experimental/ProceduralVegetationEditor/Source/ProceduralVegetation/Private/Implementations/PVGravity.h
+++ b/Engine/Plugins/Experimental/ProceduralVegetationEditor/Source/ProceduralVegetation/Private/Implementations/PVGravity.h
@@ -36,7 +36,8 @@ struct FPVGravityParams
 	UPROPERTY(EditAnywhere, Category="Gravity", meta=(ClampMin=0.0f, ClampMax=1.0f, UIMin=0.0f, UIMax=1.0f, Tooltip="Preserves initial branch angle during bending.\n\nBias that resists over-bending by maintaining the original branching angle. Increase to keep initial structure while still responding to forces."))
 	float AngleCorrection = 0;
 
-	UPROPERTY(EditAnywhere, Category="Gravity", meta=(ClampMin=0.0f, ClampMax=1.0f, UIMin=0.0f, UIMax=1.0f, EditCondition = "Mode == EGravityMode::Phototropic", Tooltip="Bias toward light-facing growth.\n\nBias between optimal light direction and shadow avoidance.  0 Being Light optimal, 1 being shadow avoidance"))
+	// UPROPERTY(EditAnywhere, Category="Gravity", meta=(ClampMin=0.0f, ClampMax=1.0f, UIMin=0.0f, UIMax=1.0f, EditCondition = "Mode == EGravityMode::Phototropic", Tooltip="Bias toward light-facing growth.\n\nBias between optimal light direction and shadow avoidance.  0 Being Light optimal, 1 being shadow avoidance"))
+	UPROPERTY()
 	float PhototropicBias = 0;
 };
 
diff --git a/Engine/Plugins/Experimental/ProceduralVegetationEditor/Source/ProceduralVegetation/Private/Implementations/PVMeshBuilder.cpp b/Engine/Plugins/Experimental/ProceduralVegetationEditor/Source/ProceduralVegetation/Private/Implementations/PVMeshBuilder.cpp
--- a/Engine/Plugins/Experimental/ProceduralVegetationEditor/Source/ProceduralVegetation/Private/Implementations/PVMeshBuilder.cpp
+++ b/Engine/Plugins/Experimental/ProceduralVegetationEditor/Source/ProceduralVegetation/Private/Implementations/PVMeshBuilder.cpp
@@ -827,8 +827,8 @@ void FPVMeshBuilder::GenerateBranchMeshData(const bool bPrimitiveIsTrunk, const
 		PreviousAimVector = UpdatedAimVector;
 		PreviousUpVector = PointUp;
 
-		const float PlantGradient = PointFacade.GetPlantGradient(PointIndex);
-		const float ProfileFallOff = MeshBuilderParams.PlantProfileFallOff.GetRichCurveConst()->Eval(PlantGradient);
+		const float BranchGradient = PointFacade.GetBranchGradient(PointIndex);
+		const float ProfileFallOff = MeshBuilderParams.PlantProfileFallOff.GetRichCurveConst()->Eval(BranchGradient);
 		if (PointScale == 0.0f)
 		[[unlikely]]
 		{
@@ -850,8 +850,7 @@ void FPVMeshBuilder::GenerateBranchMeshData(const bool bPrimitiveIsTrunk, const
 				if (bShouldApplyProfile)
 				{
 					ProfileMultiplier = GetProfileMultiplier(ProfilePoints, Displacement_UV_U);
-					ProfileMultiplier = FMath::Lerp(1.0f, ProfileMultiplier, ProfileFallOff);
-					ProfileMultiplier = ProfileMultiplier * MeshBuilderParams.PlantProfileScale;
+					ProfileMultiplier = FMath::Lerp(1.0f, ProfileMultiplier * MeshBuilderParams.PlantProfileScale, ProfileFallOff);
 				}
 
 				const FVector3f Direction = FQuat4f(UpdatedAimVector, RotationDelta * Division).RotateVector(PointUp).GetSafeNormal();
diff --git a/Engine/Plugins/Experimental/ProceduralVegetationEditor/Source/ProceduralVegetation/Private/Implementations/PVMeshBuilder.h b/Engine/Plugins/Experimental/ProceduralVegetationEditor/Source/ProceduralVegetation/Private/Implementations/PVMeshBuilder.h
--- a/Engine/Plugins/Experimental/ProceduralVegetationEditor/Source/ProceduralVegetation/Private/Implementations/PVMeshBuilder.h
+++ b/Engine/Plugins/Experimental/ProceduralVegetationEditor/Source/ProceduralVegetation/Private/Implementations/PVMeshBuilder.h
@@ -151,7 +151,7 @@ struct FPVMeshBuilderParams
 	FPVFloatRamp PlantProfileFallOff;
 
 	UPROPERTY(EditAnywhere, Category="Plant Profile Settings", DisplayName="Profile Scale",
-		meta=(ClampMin=0.0f, ClampMax=3.0f, UIMin=0.0f, UIMax=3.0f, EditCondition="bIsPlantProfileDropdownEnabled", EditConditionHides,
+		meta=(ClampMin=0.0f, ClampMax=10.0f, UIMin=0.0f, UIMax=10.0f, EditCondition="bIsPlantProfileDropdownEnabled", EditConditionHides,
 			HideEditConditionToggle))
 	float PlantProfileScale = 1.0f;
 
diff --git a/Engine/Plugins/Experimental/ProceduralVegetationEditor/Source/ProceduralVegetation/Public/Helpers/PVUtilities.h b/Engine/Plugins/Experimental/ProceduralVegetationEditor/Source/ProceduralVegetation/Public/Helpers/PVUtilities.h
--- a/Engine/Plugins/Experimental/ProceduralVegetationEditor/Source/ProceduralVegetation/Public/Helpers/PVUtilities.h
+++ b/Engine/Plugins/Experimental/ProceduralVegetationEditor/Source/ProceduralVegetation/Public/Helpers/PVUtilities.h
@@ -23,6 +23,7 @@ namespace PV::Utilities
 
 	bool IsFileNameValid(FName FileName, FText& Reason);
 	bool PROCEDURALVEGETATION_API DoesConflictingPackageExist(const FString& PackageName, UClass* InClass);
+	bool PackageExists(const FString& LongPackageName, UClass* AssetClass);
 	int32 PROCEDURALVEGETATION_API GetMeshTriangles(const FString InMeshPath);
 
 	bool IsValidPVData(const FManagedArrayCollection& Collection);
