[PVE] Readded lines removed unintentionaly

diff --git a/Engine/Plugins/Experimental/ProceduralVegetationEditor/Source/ProceduralVegetationEditor/Private/Helpers/PVExportHelper.cpp b/Engine/Plugins/Experimental/ProceduralVegetationEditor/Source/ProceduralVegetationEditor/Private/Helpers/PVExportHelper.cpp
--- a/Engine/Plugins/Experimental/ProceduralVegetationEditor/Source/ProceduralVegetationEditor/Private/Helpers/PVExportHelper.cpp
+++ b/Engine/Plugins/Experimental/ProceduralVegetationEditor/Source/ProceduralVegetationEditor/Private/Helpers/PVExportHelper.cpp
@@ -983,6 +983,10 @@ namespace PV::Export::Internal
 				return false;
 			}
 
+			Collection.CopyTo(&ExportedCollection);
+
+			DynamicMesh = CollectionToDynamicMesh(ExportedCollection);
+
 			if (!UpdateStage(EExportMeshStage::InitializingExportMesh, 0.5f))
 			{
 				return false;
