[PVE] bug hawk fix

diff --git a/Engine/Plugins/Experimental/ProceduralVegetationEditor/Source/ProceduralVegetation/Private/Helpers/PVUtilities.cpp b/Engine/Plugins/Experimental/ProceduralVegetationEditor/Source/ProceduralVegetation/Private/Helpers/PVUtilities.cpp
--- a/Engine/Plugins/Experimental/ProceduralVegetationEditor/Source/ProceduralVegetation/Private/Helpers/PVUtilities.cpp
+++ b/Engine/Plugins/Experimental/ProceduralVegetationEditor/Source/ProceduralVegetation/Private/Helpers/PVUtilities.cpp
@@ -238,8 +238,7 @@ namespace PV::Utilities
 				return false;
 			}
 
-			if (OutAssets.Num() != 1
-				|| OutAssets[0].AssetClassPath == AssetClass->GetClassPathName())
+			if (OutAssets[0].AssetClassPath == AssetClass->GetClassPathName())
 			{
 				return true;
 			}
