[PVE] Update Tooltips for Mesher node. And update plugin description.

diff --git a/Engine/Plugins/Experimental/ProceduralVegetationEditor/ProceduralVegetationEditor.uplugin b/Engine/Plugins/Experimental/ProceduralVegetationEditor/ProceduralVegetationEditor.uplugin
--- a/Engine/Plugins/Experimental/ProceduralVegetationEditor/ProceduralVegetationEditor.uplugin
+++ b/Engine/Plugins/Experimental/ProceduralVegetationEditor/ProceduralVegetationEditor.uplugin
@@ -3,7 +3,7 @@
 	"Version": 1,
 	"VersionName": "1.0",
 	"FriendlyName": "Procedural Vegetation Editor",
-	"Description": "Procedural vegetation generation and editing tool for editor",
+	"Description": "Node Graph based Editor that allows users to create Nanite Foliage ready vegetation directly in the engine. Users can load Procedural Vegetation Presets that contain prebuilt data for a species, and customize/create variations using the node graph.",
 	"Category": "Other",
 	"CreatedBy": "",
 	"CreatedByURL": "",
diff --git a/Engine/Plugins/Experimental/ProceduralVegetationEditor/Source/ProceduralVegetation/Private/Implementations/PVMeshBuilder.h b/Engine/Plugins/Experimental/ProceduralVegetationEditor/Source/ProceduralVegetation/Private/Implementations/PVMeshBuilder.h
--- a/Engine/Plugins/Experimental/ProceduralVegetationEditor/Source/ProceduralVegetation/Private/Implementations/PVMeshBuilder.h
+++ b/Engine/Plugins/Experimental/ProceduralVegetationEditor/Source/ProceduralVegetation/Private/Implementations/PVMeshBuilder.h
@@ -41,13 +41,13 @@ struct FPVMeshBuilderParams
 
 	UPROPERTY(EditAnywhere, Category="Mesh Settings", DisplayName="Point Removal",
 		meta=(ClampMin=0.0f, ClampMax=0.1f, UIMin=0.0f, UIMax=0.1f, Tooltip=
-			"Removes small points to simplify the mesh.\n\nDefines the strength of point removal. Higher values eliminate more minor vertices or detail points to reduce complexity at the cost of fine detail."
+			"Removes small points to simplify the mesh.\n\nPoints are used to define the tree structure and used as basis for generating the mesh. This does not remove points but reduces the points considered for mesh generation. The impact is dependent on Hull, Main Trunk, Ground and Scale Retention settings."
 		))
 	float PointRemoval = 0.0f;
 
 	UPROPERTY(EditAnywhere, Category="Mesh Settings", DisplayName="Segment Reduction",
 		meta=(ClampMin=0.0f, ClampMax=1.0f, UIMin=0.0f, UIMax=1.0f, Tooltip=
-			"Controls how aggressively segments are reduced.\n\nSets the strength of segment merging/decimation. Higher values reduce edge count for lighter meshes; lower values preserve detail."
+			"Controls how aggressively segments are reduced.\n\nSegments are created based on the number of points required to represent the shape of the plant. The algorithm tries to simplify the polyline by reducing the number of points required to represent it, while maintaining its overall shape. The impact is dependent on Hull, Main Trunk, Ground and Scale Retention settings.\nHigher values reduce edge count for lighter meshes; lower values preserve detail."
 		))
 	float SegmentReduction = 0.0f;
 
@@ -71,7 +71,7 @@ struct FPVMeshBuilderParams
 
 	UPROPERTY(EditAnywhere, Category="Mesh Settings|Advanced", DisplayName="Longest Segment Allowed (meters)",
 		meta=(ClampMin=0.0f, ClampMax=50.0f, UIMin=0.0f, UIMax=10.0f, Tooltip=
-			"Maximum allowed segment length.\n\nCaps segment length after simplification to maintain even geometry and prevent stretched polygons. Specified in mesh units."
+			"Maximum allowed segment length.\n\nCaps segment length after simplification to maintain even geometry and prevent stretched polygons."
 		))
 	float LongestSegmentLength = 10.0f;
 
@@ -136,7 +136,8 @@ struct FPVMeshBuilderParams
 	FPVFloatRamp ScaleRetentionGradient;
 
 	UPROPERTY(EditAnywhere, Category="Plant Profile Settings", DisplayName="Profile",
-		meta=(EditCondition="bIsPlantProfileDropdownEnabled", EditConditionHides, HideEditConditionToggle, GetOptions="GetPlantProfileOptions"))
+		meta=(EditCondition="bIsPlantProfileDropdownEnabled", EditConditionHides, HideEditConditionToggle, GetOptions="GetPlantProfileOptions",
+			Tooltip="Select a profile (cross sectional shape) from predefined options\n\nNone applies a simple circular cross sectional shape"))
 	FString SelectedPlantProfile = "None";
 
 	UPROPERTY(Transient)
@@ -147,16 +148,19 @@ struct FPVMeshBuilderParams
 
 	UPROPERTY(EditAnywhere, Category="Plant Profile Settings", DisplayName="Profile Falloff",
 		meta=(XAxisMin=0.0f, XAxisMax=1.0f, YAxisMin=0.0f, YAxisMax=1.0f, EditCondition="bIsPlantProfileDropdownEnabled", EditConditionHides,
-			HideEditConditionToggle))
+			HideEditConditionToggle, Tooltip=
+			"Gradient controlling how the profile shape is applied across the length of the trunk and branches\n\nCan be used to create a base where the trunk widens and transitions into main roots (Root Flare)"
+		))
 	FPVFloatRamp PlantProfileFallOff;
 
 	UPROPERTY(EditAnywhere, Category="Plant Profile Settings", DisplayName="Profile Scale",
 		meta=(ClampMin=0.0f, ClampMax=10.0f, UIMin=0.0f, UIMax=10.0f, EditCondition="bIsPlantProfileDropdownEnabled", EditConditionHides,
-			HideEditConditionToggle))
+			HideEditConditionToggle, Tooltip="Scale of the profile applied\n\nCan be used to create a wider base in conjection with Profile Falloff"))
 	float PlantProfileScale = 1.0f;
 
 	UPROPERTY(EditAnywhere, Category="Plant Profile Settings",
-		meta=(EditCondition="bIsPlantProfileDropdownEnabled", EditConditionHides, HideEditConditionToggle))
+		meta=(EditCondition="bIsPlantProfileDropdownEnabled", EditConditionHides, HideEditConditionToggle, Tooltip=
+			"Should the profile be applied to branches or just the trunk"))
 	bool bApplyProfileToBranches = false;
 
 	UPROPERTY(EditAnywhere, Category = "Displacement Settings",
diff --git a/Engine/Plugins/Experimental/ProceduralVegetationEditor/Source/ProceduralVegetation/Private/Nodes/PVFoliageDistributorSettings.h b/Engine/Plugins/Experimental/ProceduralVegetationEditor/Source/ProceduralVegetation/Private/Nodes/PVFoliageDistributorSettings.h
--- a/Engine/Plugins/Experimental/ProceduralVegetationEditor/Source/ProceduralVegetation/Private/Nodes/PVFoliageDistributorSettings.h
+++ b/Engine/Plugins/Experimental/ProceduralVegetationEditor/Source/ProceduralVegetation/Private/Nodes/PVFoliageDistributorSettings.h
@@ -55,7 +55,7 @@ private:
 	UPROPERTY(EditAnywhere, Category="Distribution Settings", meta=(EditCondition="OverrideDistribution", ClampMin=0.0f, ClampMax=1.0f, UIMin=0.0f, UIMax=1.0f, Tooltip="Strength of the spacing ramp’s influence.\n\nBlends the ramp with the base spacing. Lower values favor the base spacing (subtle effect); higher values apply the ramp more strongly and can amplify spacing differences."))
 	float InstanceSpacingRampEffect = 0.0f;
 
-	UPROPERTY(EditAnywhere, Category="Distribution Settings", meta=(EditCondition="OverrideDistribution", ClampMin=-1, Tooltip="Caps the number of instances per branch\n\nlimit on how many foliage instances can be placed per branch to avoid overcrowding."))
+	UPROPERTY(EditAnywhere, Category="Distribution Settings", meta=(EditCondition="OverrideDistribution", ClampMin=-1, ClampMax=1000, Tooltip="Caps the number of instances per branch\n\nlimit on how many foliage instances can be placed per branch to avoid overcrowding."))
 	int32 MaxPerBranch = -1;
 
 	UPROPERTY(EditAnywhere, Category="Scale Settings", meta=(EditCondition="OverrideDistribution", ClampMin=0.0f, UIMin=0.0f, UIMax=5.0f, Tooltip="Default size multiplier for placed foliage\n\nSets the base uniform scale applied to each instance before randomness and ramps are considered. Increase to enlarge all instances proportionally; decrease for smaller foliage. Final size still respects Min/Max limits."))
