[PVE] Remove bounds from stats.

diff --git a/Engine/Plugins/Experimental/ProceduralVegetationEditor/Source/ProceduralVegetationEditor/Private/PVEditor.cpp b/Engine/Plugins/Experimental/ProceduralVegetationEditor/Source/ProceduralVegetationEditor/Private/PVEditor.cpp
--- a/Engine/Plugins/Experimental/ProceduralVegetationEditor/Source/ProceduralVegetationEditor/Private/PVEditor.cpp
+++ b/Engine/Plugins/Experimental/ProceduralVegetationEditor/Source/ProceduralVegetationEditor/Private/PVEditor.cpp
@@ -965,20 +965,6 @@ void FPVEditor::GatherStats(TArray<FText>& OverlayStats)
 		const int32 NumFoliageInstances = CollectionBeingInspected->NumElements(PV::GroupNames::FoliageGroup);
 		OverlayStats.Add(FText::FromString("Foliage Instances: " + FString::FromInt(NumFoliageInstances)));
 	}
-	if (CollectionBeingInspected->HasGroup(FGeometryCollection::GeometryGroup))
-	{
-		const TManagedArray<FBox>& BoundingBoxes = CollectionBeingInspected->GetAttribute<FBox>(
-			FGeometryCollection::BoundingBoxAttribute,
-			FGeometryCollection::GeometryGroup
-		);
-		FBox BoundingBox;
-		for (int32 i = 0; i < BoundingBoxes.Num(); ++i)
-		{
-			BoundingBox += BoundingBoxes[i];
-		}
-		const FIntVector3 Extent = FIntVector3(BoundingBox.GetExtent() * 2.0f);
-		OverlayStats.Add(FText::FromString(FString::Printf(TEXT("Approx Size (without Foliage): %dx%dx%d"), Extent.X, Extent.Y, Extent.Z)));
-	}
 }
 
 #undef LOCTEXT_NAMESPACE
