[PVE] Collisions for Static mesh export

diff --git a/Engine/Plugins/Experimental/ProceduralVegetationEditor/Source/ProceduralVegetation/Public/PVExportParams.h b/Engine/Plugins/Experimental/ProceduralVegetationEditor/Source/ProceduralVegetation/Public/PVExportParams.h
--- a/Engine/Plugins/Experimental/ProceduralVegetationEditor/Source/ProceduralVegetation/Public/PVExportParams.h
+++ b/Engine/Plugins/Experimental/ProceduralVegetationEditor/Source/ProceduralVegetation/Public/PVExportParams.h
@@ -69,10 +69,13 @@ struct PROCEDURALVEGETATION_API FPVExportParams
 	UPROPERTY(EditAnywhere, Category = "Mesh | Nanite", Meta = (DisplayName = "Nanite Shape Preservation Method", Tooltip="Shape preservation setting to be used in nanite settings."))
 	ENaniteShapePreservation NaniteShapePreservation = ENaniteShapePreservation::Voxelize;
 
-	UPROPERTY(EditAnywhere, Category = "Mesh | Collision", Meta = (EditCondition = "ExportMeshType == EPVExportMeshType::SkeletalMesh", Tooltip="Creates the physic asset for skeletal mesh."))
-	EPVCollisionGeneration CollisionGeneration;
+	UPROPERTY(EditAnywhere, Category = "Mesh | Collision", Meta = (EditCondition = "ExportMeshType == EPVExportMeshType::StaticMesh", EditConditionHides, Tooltip="Creates the physic asset for static mesh."))
+	bool bCollision;
 
-	UPROPERTY(EditAnywhere, Category = "Mesh | Dynamic Wind", Meta = (EditCondition = "ExportMeshType == EPVExportMeshType::SkeletalMesh", Tooltip="Wind Settings contains the wind simulation group data, Default presets are available for Saplings and Trees."))
+	UPROPERTY(EditAnywhere, Category = "Mesh | Collision", Meta = (EditCondition = "ExportMeshType == EPVExportMeshType::SkeletalMesh", EditConditionHides, Tooltip="Creates the physic asset for skeletal mesh."))
+	EPVCollisionGeneration CollisionGeneration;
+	
+	UPROPERTY(EditAnywhere, Category = "Mesh | Dynamic Wind", Meta = (EditCondition = "ExportMeshType == EPVExportMeshType::SkeletalMesh", EditConditionHides, Tooltip="Wind Settings contains the wind simulation group data, Default presets are available for Saplings and Trees."))
 	TObjectPtr<UPVWindSettings> WindSettings;
 
 	void Initialize(const FString& InAssetPath, const FString& InName, const FString& InDefaultWindSettingsPath)
diff --git a/Engine/Plugins/Experimental/ProceduralVegetationEditor/Source/ProceduralVegetationEditor/Private/Helpers/PVExportHelper.cpp b/Engine/Plugins/Experimental/ProceduralVegetationEditor/Source/ProceduralVegetationEditor/Private/Helpers/PVExportHelper.cpp
--- a/Engine/Plugins/Experimental/ProceduralVegetationEditor/Source/ProceduralVegetationEditor/Private/Helpers/PVExportHelper.cpp
+++ b/Engine/Plugins/Experimental/ProceduralVegetationEditor/Source/ProceduralVegetationEditor/Private/Helpers/PVExportHelper.cpp
@@ -955,12 +955,47 @@ namespace PV::Export::Internal
 		);
 	}
 
+	void AddCollisionToStaticMesh(const FManagedArrayCollection& Collection, TObjectPtr<UStaticMesh> ExportMesh)
+	{
+		UBodySetup* BodySetup = ExportMesh->GetBodySetup();
+		if(!BodySetup)
+		{
+			BodySetup = NewObject<UBodySetup>(ExportMesh);
+		}
+		else
+		{
+			BodySetup->RemoveSimpleCollision();
+		}
+			
+		FKConvexElem ConvexElem;
+
+		FGeometryCollection GeoCollection;
+		Collection.CopyTo(&GeoCollection);
+		
+		TArray<FVector> Vertices;
+		for (const auto&Vertex : GeoCollection.Vertex)
+		{
+			Vertices.Add(FVector(Vertex));
+		}
+		
+		ConvexElem.VertexData = Vertices;
+		ConvexElem.UpdateElemBox();
+		BodySetup->AggGeom.ConvexElems.Add(ConvexElem);
+
+		BodySetup->bHasCookedCollisionData = true;
+		BodySetup->InvalidatePhysicsData();
+		BodySetup->CreatePhysicsMeshes();
+
+		ExportMesh->SetBodySetup(BodySetup);
+	}
+
 	bool ExportCollectionToStaticMesh(
 		TObjectPtr<UStaticMesh> ExportMesh,
 		const FManagedArrayCollection& Collection,
 		ENaniteShapePreservation InShapePreservation,
 		bool bBuildNaniteAssemblies,
 		bool bShouldExportFoliage,
+		bool bCollision,
 		TArray<UPackage*>& OutModifiedPackages,
 		const FStatusReportCallback& StatusReportCallback
 	)
@@ -987,6 +1022,11 @@ namespace PV::Export::Internal
 
 			DynamicMesh = CollectionToDynamicMesh(ExportedCollection);
 
+			if (bCollision)
+			{
+				AddCollisionToStaticMesh(ExportedCollection, ExportMesh);
+			}
+
 			if (!UpdateStage(EExportMeshStage::InitializingExportMesh, 0.5f))
 			{
 				return false;
@@ -1677,6 +1717,7 @@ namespace PV::Export
 				ExportParams.NaniteShapePreservation,
 				bCreateNaniteFoliage,
 				bShouldExportFoliage,
+				ExportParams.bCollision,
 				ModifiedPackages,
 				StatusReportCallback
 			))
diff --git a/Engine/Plugins/Experimental/ProceduralVegetationEditor/Source/ProceduralVegetationEditor/Private/Helpers/PVExportHelper.h b/Engine/Plugins/Experimental/ProceduralVegetationEditor/Source/ProceduralVegetationEditor/Private/Helpers/PVExportHelper.h
--- a/Engine/Plugins/Experimental/ProceduralVegetationEditor/Source/ProceduralVegetationEditor/Private/Helpers/PVExportHelper.h
+++ b/Engine/Plugins/Experimental/ProceduralVegetationEditor/Source/ProceduralVegetationEditor/Private/Helpers/PVExportHelper.h
@@ -81,6 +81,7 @@ namespace PV::Export::Internal
 		ENaniteShapePreservation InShapePreservation,
 		bool bBuildNaniteAssemblies,
 		bool bShouldExportFoliage,
+		bool bCollision,
 		TArray<UPackage*>& OutModifiedPackages,
 		const FStatusReportCallback& StatusReportCallback
 	);
