[PVE] Tooltip updated

diff --git a/Engine/Plugins/Experimental/ProceduralVegetationEditor/Source/ProceduralVegetation/Public/PVExportParams.h b/Engine/Plugins/Experimental/ProceduralVegetationEditor/Source/ProceduralVegetation/Public/PVExportParams.h
--- a/Engine/Plugins/Experimental/ProceduralVegetationEditor/Source/ProceduralVegetation/Public/PVExportParams.h
+++ b/Engine/Plugins/Experimental/ProceduralVegetationEditor/Source/ProceduralVegetation/Public/PVExportParams.h
@@ -69,10 +69,10 @@ struct PROCEDURALVEGETATION_API FPVExportParams
 	UPROPERTY(EditAnywhere, Category = "Mesh | Nanite", Meta = (DisplayName = "Nanite Shape Preservation Method", Tooltip="Shape preservation setting to be used in nanite settings."))
 	ENaniteShapePreservation NaniteShapePreservation = ENaniteShapePreservation::Voxelize;
 
-	UPROPERTY(EditAnywhere, Category = "Mesh | Collision", Meta = (EditCondition = "ExportMeshType == EPVExportMeshType::StaticMesh", EditConditionHides, Tooltip="Creates the physic asset for static mesh."))
+	UPROPERTY(EditAnywhere, Category = "Mesh | Collision", Meta = (EditCondition = "ExportMeshType == EPVExportMeshType::StaticMesh", EditConditionHides, Tooltip="Creates the collision for static mesh."))
 	bool bCollision;
 
-	UPROPERTY(EditAnywhere, Category = "Mesh | Collision", Meta = (EditCondition = "ExportMeshType == EPVExportMeshType::SkeletalMesh", EditConditionHides, Tooltip="Creates the physic asset for skeletal mesh."))
+	UPROPERTY(EditAnywhere, Category = "Mesh | Collision", Meta = (EditCondition = "ExportMeshType == EPVExportMeshType::SkeletalMesh", EditConditionHides, Tooltip="Creates the physics asset for skeletal mesh."))
 	EPVCollisionGeneration CollisionGeneration;
 	
 	UPROPERTY(EditAnywhere, Category = "Mesh | Dynamic Wind", Meta = (EditCondition = "ExportMeshType == EPVExportMeshType::SkeletalMesh", EditConditionHides, Tooltip="Wind Settings contains the wind simulation group data, Default presets are available for Saplings and Trees."))
