[PVE] Fixed compilation issue on old tool chains

diff --git a/Engine/Plugins/Experimental/ProceduralVegetationEditor/Source/ProceduralVegetationEditor/Private/Visualizations/PVDataVisualization.cpp b/Engine/Plugins/Experimental/ProceduralVegetationEditor/Source/ProceduralVegetationEditor/Private/Visualizations/PVDataVisualization.cpp
--- a/Engine/Plugins/Experimental/ProceduralVegetationEditor/Source/ProceduralVegetationEditor/Private/Visualizations/PVDataVisualization.cpp
+++ b/Engine/Plugins/Experimental/ProceduralVegetationEditor/Source/ProceduralVegetationEditor/Private/Visualizations/PVDataVisualization.cpp
@@ -263,7 +263,7 @@ void FPVDataVisualization::FoliageGridRenderer(const FManagedArrayCollection& In
 			return BoundingBox;
 		};
 
-	const auto AddComponentToScene = [&](USceneComponent* ComponentRef, const FTransform& Transform = FTransform::Identity, const bool UpdateBounds = true)
+	const auto AddComponentToScene = [&](USceneComponent* ComponentRef, const FTransform& Transform, const bool UpdateBounds = true)
 		{
 			InOutParams.ManagedResources.Add(ComponentRef);
 			InOutParams.Scene->AddComponent(ComponentRef, Transform);
