[PVE] Issue : Material gets removed when generation offset is tweaked
Reason : With the large offset the min value exceeds the number or available materials
Fix : Making sure the min value never exceeds the max value, Also removed the Clamp offset method as it was confusing of the users only Material Distribution method should be enough to i.e clamp or refit the available materials.
Disabling the auto focus viewport by default.

diff --git a/Engine/Plugins/Experimental/ProceduralVegetationEditor/Source/ProceduralVegetation/Private/Implementations/PVMaterialSettings.cpp b/Engine/Plugins/Experimental/ProceduralVegetationEditor/Source/ProceduralVegetation/Private/Implementations/PVMaterialSettings.cpp
--- a/Engine/Plugins/Experimental/ProceduralVegetationEditor/Source/ProceduralVegetation/Private/Implementations/PVMaterialSettings.cpp
+++ b/Engine/Plugins/Experimental/ProceduralVegetationEditor/Source/ProceduralVegetation/Private/Implementations/PVMaterialSettings.cpp
@@ -77,7 +77,7 @@ void FPVMaterialSettings::SetBranchUVMaterial(FManagedArrayCollection& Collectio
 		MaxAge = FMath::Max(MaxAge, BudDevelopment[2]);
 		MaxPScale = FMath::Max(MaxPScale, PScale);
 	}
-
+	
 	for (int32 BranchIndex = 0; BranchIndex < BranchFacade.GetElementCount(); BranchIndex++)
 	{
 		const auto& BranchPoints = BranchFacade.GetPoints(BranchIndex);
@@ -103,19 +103,16 @@ void FPVMaterialSettings::SetBranchUVMaterial(FManagedArrayCollection& Collectio
 			BranchMaterialVal = GetMaterialValue(PScale, MinPScale, MaxPScale, MaxGeneration, MinGeneration);
 		}
 
-		if (DistributionMethod == EMaterialDistributionMethod::Fit)
-		{
-			MaxGeneration = MaterialSetups.Num() + GenerationOffset;
-		}
-
-		if (OffsetMethod == EGenerationOffsetMethod::Clamped)
+		const int MinValue = FMath::Min( MinGeneration + GenerationOffset , MaterialSetups.Num());
+		
+		if (DistributionMethod == EMaterialDistributionMethod::Repeat)
 		{
-			BranchMaterialVal = FMath::Clamp(BranchMaterialVal + GenerationOffset,MinGeneration + GenerationOffset,MaterialSetups.Num());
+			BranchMaterialVal = FMath::Clamp(BranchMaterialVal + GenerationOffset,MinValue,MaterialSetups.Num());
 		}
-		else if (OffsetMethod == EGenerationOffsetMethod::Refit)
+		else if (DistributionMethod == EMaterialDistributionMethod::Fit)
 		{
 			FVector2f OldRange = FVector2f(MinGeneration, MaxGeneration);
-			FVector2f NewRange = FVector2f(MinGeneration + GenerationOffset, MaterialSetups.Num());
+			FVector2f NewRange = FVector2f(MinValue, MaterialSetups.Num());
 			BranchMaterialVal = FMath::GetMappedRangeValueClamped(OldRange, NewRange, BranchMaterialVal);
 		}
 
diff --git a/Engine/Plugins/Experimental/ProceduralVegetationEditor/Source/ProceduralVegetation/Private/Implementations/PVMaterialSettings.h b/Engine/Plugins/Experimental/ProceduralVegetationEditor/Source/ProceduralVegetation/Private/Implementations/PVMaterialSettings.h
--- a/Engine/Plugins/Experimental/ProceduralVegetationEditor/Source/ProceduralVegetation/Private/Implementations/PVMaterialSettings.h
+++ b/Engine/Plugins/Experimental/ProceduralVegetationEditor/Source/ProceduralVegetation/Private/Implementations/PVMaterialSettings.h
@@ -81,9 +81,6 @@ struct FPVMaterialSettings
 	UPROPERTY(EditAnywhere, Category="Materials", meta=(ClampMin = "0", Tooltip="Shifts which branch generation the material targets.\n\nAdjusts the generation index used for material assignment. For example, an offset of 1 targets branches one level above the default (trunk → gen0, first branches → gen1, etc.)."))
 	int32 GenerationOffset = 0;
 
-	UPROPERTY(EditAnywhere, Category="Materials", meta=(Tooltip="How generation offset is interpreted.\n\nDefines handling of GenerationOffset: Clamped keeps values within valid generations; Refit rescales offsets across available generations for smoother distribution."))
-	EGenerationOffsetMethod OffsetMethod = EGenerationOffsetMethod::Clamped;
-
 	UPROPERTY(EditAnywhere, Category="Materials", meta=(Tooltip="How materials are selected based on number of materials available vs generations\n\nIf number of generations exceed the number of materials, Fit will fit all generations between available materials. Repeat will apply the last material entry to all generations exceeding number of materials"))
 	EMaterialDistributionMethod DistributionMethod = EMaterialDistributionMethod::Repeat;
 
diff --git a/Engine/Plugins/Experimental/ProceduralVegetationEditor/Source/ProceduralVegetation/Private/Implementations/PVMeshBuilder.cpp b/Engine/Plugins/Experimental/ProceduralVegetationEditor/Source/ProceduralVegetation/Private/Implementations/PVMeshBuilder.cpp
--- a/Engine/Plugins/Experimental/ProceduralVegetationEditor/Source/ProceduralVegetation/Private/Implementations/PVMeshBuilder.cpp
+++ b/Engine/Plugins/Experimental/ProceduralVegetationEditor/Source/ProceduralVegetation/Private/Implementations/PVMeshBuilder.cpp
@@ -33,7 +33,7 @@ void MergeMaterials(const TArray<FTrunkGenerationMaterialSetup>& InMaterialSetup
 
 		if (!MaterialPathIDMap.Contains(MaterialPathName))
 		{
-			MaterialPathIDMap.Add(MaterialPathName);
+			MaterialPathIDMap.Add(MaterialPathName , MeshData.MaterialID);
 		}
 		else
 		{
diff --git a/Engine/Plugins/Experimental/ProceduralVegetationEditor/Source/ProceduralVegetationEditor/Public/PVEditorSettings.h b/Engine/Plugins/Experimental/ProceduralVegetationEditor/Source/ProceduralVegetationEditor/Public/PVEditorSettings.h
--- a/Engine/Plugins/Experimental/ProceduralVegetationEditor/Source/ProceduralVegetationEditor/Public/PVEditorSettings.h
+++ b/Engine/Plugins/Experimental/ProceduralVegetationEditor/Source/ProceduralVegetationEditor/Public/PVEditorSettings.h
@@ -40,7 +40,7 @@ public:
 	bool bShowScaleVisualization = false;
 
 	UPROPERTY(EditAnywhere, config, Category = Export)
-	bool bAutoFocusViewport = true;
+	bool bAutoFocusViewport = false;
 
 	UPROPERTY(EditAnywhere, config, Category = Settings)
 	EPVExportType ExportType = EPVExportType::Selection;
