[PVE] Ensure fixed size for overridable foliage entries in the Foliage Palette node

diff --git a/Engine/Plugins/Experimental/ProceduralVegetationEditor/Source/ProceduralVegetation/Private/Nodes/PVFoliagePaletteSettings.h b/Engine/Plugins/Experimental/ProceduralVegetationEditor/Source/ProceduralVegetation/Private/Nodes/PVFoliagePaletteSettings.h
--- a/Engine/Plugins/Experimental/ProceduralVegetationEditor/Source/ProceduralVegetation/Private/Nodes/PVFoliagePaletteSettings.h
+++ b/Engine/Plugins/Experimental/ProceduralVegetationEditor/Source/ProceduralVegetation/Private/Nodes/PVFoliagePaletteSettings.h
@@ -34,7 +34,9 @@ public:
 	UPROPERTY(EditAnywhere, Category="Meshes", meta=(PinHiddenByDefault, InlineEditConditionToggle))
 	bool bOverrideFoliage = false;
 
-	UPROPERTY(EditAnywhere, Category = "Meshes", Meta = (EditCondition="bOverrideFoliage", AllowedClasses = "/Script/Engine.StaticMesh,/Script/Engine.SkeletalMesh", Tooltip="List of Foliage meshes used for this plant\n\nList of all foliage meshes applied based on procedural vegetation preset."))
+	UPROPERTY(EditAnywhere, Category = "Meshes", EditFixedSize,
+		Meta = (NoResetToDefault, EditCondition="bOverrideFoliage", AllowedClasses = "/Script/Engine.StaticMesh,/Script/Engine.SkeletalMesh", Tooltip=
+			"List of Foliage meshes used for this plant\n\nList of all foliage meshes applied based on procedural vegetation preset."))
 	TArray<TSoftObjectPtr<UObject>> FoliageMeshes;
 
 };
