[PVE] Fixed issue: Exports Fails Consistently on First Attempt when renamed, Succeeds on Second

Fix/How: When changing the mesh name or path in output settings and clicking on the export button, existing inspection data is cleared, and rendering gets scheduled. Graph execution happens from the tick function. The export popup is a blocking popup, so the tick is not fired during that time (so no graph execution). So, no node data is available when the export button is pressed. To fix that, I have disabled the graph execution for these properties because they don't affect graph state.

diff --git a/Engine/Plugins/Experimental/ProceduralVegetationEditor/Source/ProceduralVegetation/Private/Nodes/PVOutputSettings.cpp b/Engine/Plugins/Experimental/ProceduralVegetationEditor/Source/ProceduralVegetation/Private/Nodes/PVOutputSettings.cpp
--- a/Engine/Plugins/Experimental/ProceduralVegetationEditor/Source/ProceduralVegetation/Private/Nodes/PVOutputSettings.cpp
+++ b/Engine/Plugins/Experimental/ProceduralVegetationEditor/Source/ProceduralVegetation/Private/Nodes/PVOutputSettings.cpp
@@ -83,6 +83,24 @@ void UPVOutputSettings::PostEditChangeProperty(FPropertyChangedEvent& PropertyCh
 		}
 	}
 }
+
+EPCGChangeType UPVOutputSettings::GetChangeTypeForProperty(FPropertyChangedEvent& PropertyChangedEvent) const
+{
+	const FName PropertyName = PropertyChangedEvent.GetPropertyName();
+	const FName MemberPropertyName = PropertyChangedEvent.GetMemberPropertyName();
+
+	if (MemberPropertyName == GET_MEMBER_NAME_CHECKED(UPVOutputSettings, ExportSettings))
+	{
+		if (PropertyName == GET_MEMBER_NAME_CHECKED(FDirectoryPath, Path)
+			||PropertyName == GET_MEMBER_NAME_CHECKED(FPVExportParams, MeshName))
+		{
+			return EPCGChangeType::Cosmetic;
+		}
+	}
+
+	return Super::GetChangeTypeForProperty(PropertyChangedEvent);
+}
+
 #endif
 
 bool FPVOutputElement::ExecuteInternal(FPCGContext* InContext) const
diff --git a/Engine/Plugins/Experimental/ProceduralVegetationEditor/Source/ProceduralVegetation/Public/Nodes/PVOutputSettings.h b/Engine/Plugins/Experimental/ProceduralVegetationEditor/Source/ProceduralVegetation/Public/Nodes/PVOutputSettings.h
--- a/Engine/Plugins/Experimental/ProceduralVegetationEditor/Source/ProceduralVegetation/Public/Nodes/PVOutputSettings.h
+++ b/Engine/Plugins/Experimental/ProceduralVegetationEditor/Source/ProceduralVegetation/Public/Nodes/PVOutputSettings.h
@@ -31,6 +31,7 @@ protected:
 
 #if WITH_EDITOR
 	virtual void PostEditChangeProperty(FPropertyChangedEvent& PropertyChangedEvent) override;
+	virtual EPCGChangeType GetChangeTypeForProperty(FPropertyChangedEvent& PropertyChangedEvent) const override;
 #endif
 	
 public:
