[PVE] Issue 1: Replacing the exisitng skeletal mesh that is adde to scene using the ISKM causes a crashes.
Reason : For replacing the existing skeletal mesh we create a new mesh object with the same name as existing mesh. Upon build Skeltal mesh looks for components that uses the mesh and try to clean up the nanite rendeing data for the components that have this mesh assigned and dont find the valid data and crash.
Fix : Instead of creating a new Mesh object when replacing we clean up the the existing mesh and resue that.

Issue 2: In exported mesh branches disappear.
Reason : The scale node was scaling the length from root attribute incorrectly that cuases the wrong bone influences and unassigned influences for nanite assembly parts hence the parts get culled when added to level.
Fix : The scale node should just multiply the existing Length from root attribute.

diff --git a/Engine/Plugins/Experimental/ProceduralVegetationEditor/Source/ProceduralVegetation/Private/Implementations/PVScale.cpp b/Engine/Plugins/Experimental/ProceduralVegetationEditor/Source/ProceduralVegetation/Private/Implementations/PVScale.cpp
--- a/Engine/Plugins/Experimental/ProceduralVegetationEditor/Source/ProceduralVegetation/Private/Implementations/PVScale.cpp
+++ b/Engine/Plugins/Experimental/ProceduralVegetationEditor/Source/ProceduralVegetation/Private/Implementations/PVScale.cpp
@@ -41,7 +41,6 @@ void FPVScale::ApplyScale(const float InManualScale, FManagedArrayCollection& Ou
 
 		const int32 RootPointIndex = TrunkPoints[0];
 		const FVector3f& PivotPos = PointFacade.GetPosition(RootPointIndex);
-		const float RootPointLFR = PointFacade.GetLengthFromRoot(RootPointIndex);
 
 		for (const int32 BranchIndex : BranchIndices)
 		{
@@ -50,15 +49,15 @@ void FPVScale::ApplyScale(const float InManualScale, FManagedArrayCollection& Ou
 			{
 				FVector3f PointPos = PointFacade.GetPosition(PointIndex);
 				float LFS = PointFacade.GetLengthFromSeed(PointIndex);
+				float LFR = PointFacade.GetLengthFromRoot(PointIndex);
 
 				PointPos -= PivotPos;
 				PointPos *= InManualScale;
 				PointPos += PivotPos;
 
 				LFS *= InManualScale;
-
-				float LFR = LFS + RootPointLFR;
-
+				LFR *= InManualScale;
+				
 				PointFacade.ModifyPositions()[PointIndex] = PointPos;
 				PointFacade.ModifyLengthFromSeeds()[PointIndex] = LFS;
 				PointFacade.ModifyLengthFromRoots()[PointIndex] = LFR;
diff --git a/Engine/Plugins/Experimental/ProceduralVegetationEditor/Source/ProceduralVegetation/Public/PVExportParams.h b/Engine/Plugins/Experimental/ProceduralVegetationEditor/Source/ProceduralVegetation/Public/PVExportParams.h
--- a/Engine/Plugins/Experimental/ProceduralVegetationEditor/Source/ProceduralVegetation/Public/PVExportParams.h
+++ b/Engine/Plugins/Experimental/ProceduralVegetationEditor/Source/ProceduralVegetation/Public/PVExportParams.h
@@ -70,10 +70,10 @@ struct PROCEDURALVEGETATION_API FPVExportParams
 	ENaniteShapePreservation NaniteShapePreservation = ENaniteShapePreservation::Voxelize;
 
 	UPROPERTY(EditAnywhere, Category = "Mesh | Collision", Meta = (EditCondition = "ExportMeshType == EPVExportMeshType::StaticMesh", EditConditionHides, Tooltip="Creates the collision for static mesh."))
-	bool bCollision;
+	bool bCollision = false;
 
 	UPROPERTY(EditAnywhere, Category = "Mesh | Collision", Meta = (EditCondition = "ExportMeshType == EPVExportMeshType::SkeletalMesh", EditConditionHides, Tooltip="Creates the physics asset for skeletal mesh."))
-	EPVCollisionGeneration CollisionGeneration;
+	EPVCollisionGeneration CollisionGeneration = EPVCollisionGeneration::None;
 	
 	UPROPERTY(EditAnywhere, Category = "Mesh | Dynamic Wind", Meta = (EditCondition = "ExportMeshType == EPVExportMeshType::SkeletalMesh", EditConditionHides, Tooltip="Wind Settings contains the wind simulation group data, Default presets are available for Saplings and Trees."))
 	TObjectPtr<UPVWindSettings> WindSettings;
diff --git a/Engine/Plugins/Experimental/ProceduralVegetationEditor/Source/ProceduralVegetationEditor/Private/Helpers/PVExportHelper.cpp b/Engine/Plugins/Experimental/ProceduralVegetationEditor/Source/ProceduralVegetationEditor/Private/Helpers/PVExportHelper.cpp
--- a/Engine/Plugins/Experimental/ProceduralVegetationEditor/Source/ProceduralVegetationEditor/Private/Helpers/PVExportHelper.cpp
+++ b/Engine/Plugins/Experimental/ProceduralVegetationEditor/Source/ProceduralVegetationEditor/Private/Helpers/PVExportHelper.cpp
@@ -1743,10 +1743,24 @@ namespace PV::Export
 			TObjectPtr<UPackage> MeshPackage = Internal::CreateAssetPackage(MeshPackageName);
 			TObjectPtr<UPackage> SkeletonPackage = Internal::CreateAssetPackage(SkeletonPackageName);
 
-			if (TObjectPtr<USkeletalMesh> ExistingMeshAsset = Cast<USkeletalMesh>(MeshPackage->FindAssetInPackage()))
+			TObjectPtr<USkeletalMesh> ExistingMeshAsset = Cast<USkeletalMesh>(MeshPackage->FindAssetInPackage());
+
+			TObjectPtr<USkeletalMesh> ExportMesh = nullptr;
+			
+			if (ExistingMeshAsset)
 			{
-				ExistingMeshAsset->SetExternalPackage(nullptr);
-				ExistingMeshAsset->MarkAsGarbage();
+				ExportMesh = ExistingMeshAsset;
+				ExistingMeshAsset->ClearMeshDescriptionAndBulkData(0);
+				FNaniteAssemblyData& NaniteAssemblyData = ExistingMeshAsset->NaniteSettings.NaniteAssemblyData;
+				if (NaniteAssemblyData.IsValid())
+				{
+					NaniteAssemblyData.Nodes.Empty();
+					NaniteAssemblyData.Parts.Empty();
+				}
+			}
+			else
+			{
+				ExportMesh = NewObject<USkeletalMesh>(MeshPackage, *MeshName, RF_Standalone | RF_Public);
 			}
 			
 			if (TObjectPtr<USkeleton> ExistingSkeletonAsset = Cast<USkeleton>(SkeletonPackage->FindAssetInPackage()))
@@ -1755,7 +1769,6 @@ namespace PV::Export
 				ExistingSkeletonAsset->MarkAsGarbage();
 			}
 
-			const TObjectPtr<USkeletalMesh> ExportMesh = NewObject<USkeletalMesh>(MeshPackage, *MeshName, RF_Standalone | RF_Public);
 			const TObjectPtr<USkeleton> ExportSkeleton = NewObject<USkeleton>(SkeletonPackage, *SkeletonName, RF_Standalone | RF_Public);
 			
 			ExportMesh->SetSkeleton(ExportSkeleton);
