[PVE] Fixed issue: Remove branches - Generation mode is not working with Hazal tree
Misc
- Disabled viewport background
- Clamped threshold value to 1.

diff --git a/Engine/Plugins/Experimental/ProceduralVegetationEditor/Source/ProceduralVegetation/Private/Implementations/PVRemoveBranches.cpp b/Engine/Plugins/Experimental/ProceduralVegetationEditor/Source/ProceduralVegetation/Private/Implementations/PVRemoveBranches.cpp
--- a/Engine/Plugins/Experimental/ProceduralVegetationEditor/Source/ProceduralVegetation/Private/Implementations/PVRemoveBranches.cpp
+++ b/Engine/Plugins/Experimental/ProceduralVegetationEditor/Source/ProceduralVegetation/Private/Implementations/PVRemoveBranches.cpp
@@ -38,8 +38,7 @@ namespace PV::Facades
 				continue;
 
 			const int32 RootPointIndex = TrunkPoints[0];
-			const int32 RootHierarchyGeneration =  InBranchFacade.GetHierarchyGenerationNumber(RootPointIndex);
-			
+			const int32 RootHierarchyGeneration =  InBranchFacade.GetHierarchyGenerationNumber(TrunkIndex);
 			const TArray<int>& RootBudDevelopment = InPointFacade.GetBudDevelopment(RootPointIndex);
 
 			float MaxValue = std::numeric_limits<float>::lowest();
diff --git a/Engine/Plugins/Experimental/ProceduralVegetationEditor/Source/ProceduralVegetation/Private/Nodes/PVRemoveBranchesSettings.h b/Engine/Plugins/Experimental/ProceduralVegetationEditor/Source/ProceduralVegetation/Private/Nodes/PVRemoveBranchesSettings.h
--- a/Engine/Plugins/Experimental/ProceduralVegetationEditor/Source/ProceduralVegetation/Private/Nodes/PVRemoveBranchesSettings.h
+++ b/Engine/Plugins/Experimental/ProceduralVegetationEditor/Source/ProceduralVegetation/Private/Nodes/PVRemoveBranchesSettings.h
@@ -31,7 +31,7 @@ public:
 	UPROPERTY(EditAnywhere, Category = Settings, meta=(Tooltip="Sets the rule used to remove branches.\n\nChooses the basis for pruning (e.g., generation, length from root, age, or radius). Determines which branches are candidates for removal."))
 	ERemoveBranchesBasis BranchRemoveBasis = ERemoveBranchesBasis::Length;
 
-	UPROPERTY(EditAnywhere, Category = Settings, meta=(ClampMin = 0f, ClampMax = 10.0f, UIMin = 0f, UIMax = 1.0f, Tooltip="Removal threshold for the chosen basis.\n\nBranches exceeding or falling below this threshold (per the selected basis) are removed. Use lower/higher values to prune sparsely or aggressively."))
+	UPROPERTY(EditAnywhere, Category = Settings, meta=(ClampMin = 0f, ClampMax = 1.0f, UIMin = 0f, UIMax = 1.0f, Tooltip="Removal threshold for the chosen basis.\n\nBranches exceeding or falling below this threshold (per the selected basis) are removed. Use lower/higher values to prune sparsely or aggressively."))
 	float Threshold = 0;
 };
 
diff --git a/Engine/Plugins/Experimental/ProceduralVegetationEditor/Source/ProceduralVegetationEditor/Private/Visualizations/PVDataVisualization.cpp b/Engine/Plugins/Experimental/ProceduralVegetationEditor/Source/ProceduralVegetationEditor/Private/Visualizations/PVDataVisualization.cpp
--- a/Engine/Plugins/Experimental/ProceduralVegetationEditor/Source/ProceduralVegetationEditor/Private/Visualizations/PVDataVisualization.cpp
+++ b/Engine/Plugins/Experimental/ProceduralVegetationEditor/Source/ProceduralVegetationEditor/Private/Visualizations/PVDataVisualization.cpp
@@ -78,7 +78,7 @@ FPCGSetupSceneFunc FPVDataVisualization::GetViewportSetupFunc(const UPCGSettings
 		{
 			InOutParams.Scene->HandleToggleGrid();
 		}
-		if (!InOutParams.Scene->IsEnvironmentEnabled())
+		if (InOutParams.Scene->IsEnvironmentEnabled())
 		{
 			InOutParams.Scene->HandleToggleEnvironment();
 		}
