[PVE] Fix crash on selecting multiple same type nodes, having a ramp parameter

diff --git a/Engine/Plugins/Experimental/ProceduralVegetationEditor/Source/ProceduralVegetationEditor/Private/Customizations/PVFloatRampCustomization.cpp b/Engine/Plugins/Experimental/ProceduralVegetationEditor/Source/ProceduralVegetationEditor/Private/Customizations/PVFloatRampCustomization.cpp
--- a/Engine/Plugins/Experimental/ProceduralVegetationEditor/Source/ProceduralVegetationEditor/Private/Customizations/PVFloatRampCustomization.cpp
+++ b/Engine/Plugins/Experimental/ProceduralVegetationEditor/Source/ProceduralVegetationEditor/Private/Customizations/PVFloatRampCustomization.cpp
@@ -276,7 +276,6 @@ void FPVFloatRampCustomization::CustomizeHeader(
 		}
 		CurveEditor->RegisterToPropertyChangedEvent(RootHandle);
 
-
 		HeaderRow
 			.NameContent()
 			[
@@ -295,6 +294,21 @@ void FPVFloatRampCustomization::CustomizeHeader(
 				.TextStyle(FAppStyle::Get(), "DetailsView.CategoryTextStyle")
 			];
 	}
+	else
+	{
+		HeaderRow
+			.NameContent()
+			[
+				InPropertyHandle->CreatePropertyNameWidget()
+			]
+			.ValueContent()
+			[
+				SNew(STextBlock)
+					.Text(StructPtrs.Num() == 0 ? LOCTEXT("PVE_NoCurves", "No Curves - Unable to modify") : LOCTEXT("PVE_MultipleCurves", "Multiple Curves - Unable to modify"))
+                	.Font(FAppStyle::GetFontStyle("PropertyWindow.NormalFont"))
+                	.TextStyle(FAppStyle::Get(), "DetailsView.CategoryTextStyle")
+			];
+	}
 }
 
 void FPVFloatRampCustomization::CustomizeChildren(
@@ -308,7 +322,11 @@ void FPVFloatRampCustomization::CustomizeChildren(
 		PropertyHandle = InPropertyHandle;
 	}
 
-	check(FloatRampOwner);
+	if (!FloatRampOwner)
+	{
+		return;
+	}
+	
 	check(CurveEditor);
 
 	ChildBuilder
