Bug fixes
- Fixed Export crashes when setting the foliage palette entry to none
- Issue/Fix: This was happening because we were trying to load something on an empty foliage path. To fix that, we skip empty paths.
- Fixed issue where empty foliage slots are not ignored
- Issue/Fix: Empty slots were still added in the distribution, which caused some foliage spawning points to have nothing. To fix that, we are now using only valid foliage entries in the distribution.

diff --git a/Engine/Plugins/Experimental/ProceduralVegetationEditor/Source/ProceduralVegetation/Private/Facades/PVFoliageFacade.cpp b/Engine/Plugins/Experimental/ProceduralVegetationEditor/Source/ProceduralVegetation/Private/Facades/PVFoliageFacade.cpp
--- a/Engine/Plugins/Experimental/ProceduralVegetationEditor/Source/ProceduralVegetation/Private/Facades/PVFoliageFacade.cpp
+++ b/Engine/Plugins/Experimental/ProceduralVegetationEditor/Source/ProceduralVegetation/Private/Facades/PVFoliageFacade.cpp
@@ -246,16 +246,13 @@ namespace PV::Facades
 
 	TArray<FString> FFoliageFacade::GetFoliageNames() const
 	{
-		TArray<FString> FoliageNames;
 		if (NamesAttribute.IsValid())
 		{
-			for (int32 i = 0; i < NamesAttribute.Num(); i++)
-			{
-				FoliageNames.AddUnique(NamesAttribute[i]);
-			}
+			return NamesAttribute.Get().GetConstArray();
 		}
 
-		return FoliageNames;
+		const static TArray<FString> DefaultFoliageNames;
+		return DefaultFoliageNames;
 	}
 
 	const FVector3f& FFoliageFacade::GetPivotPoint(const int32 Index) const
diff --git a/Engine/Plugins/Experimental/ProceduralVegetationEditor/Source/ProceduralVegetation/Private/Implementations/PVFoliage.cpp b/Engine/Plugins/Experimental/ProceduralVegetationEditor/Source/ProceduralVegetation/Private/Implementations/PVFoliage.cpp
--- a/Engine/Plugins/Experimental/ProceduralVegetationEditor/Source/ProceduralVegetation/Private/Implementations/PVFoliage.cpp
+++ b/Engine/Plugins/Experimental/ProceduralVegetationEditor/Source/ProceduralVegetation/Private/Implementations/PVFoliage.cpp
@@ -92,6 +92,25 @@ void FPVFoliage::DistributeFoliage(FManagedArrayCollection& OutCollection, const
 
 		// Compute attachment points in terms of normalized distances
 		OutCollection.EmptyGroup(PV::GroupNames::FoliageGroup);
+ 
+		// Collect non-empty mesh name indices
+		const TArray<FString>& FoliageMeshNames = FoliageFacadeOutput.GetFoliageNames();
+		TArray<int32> NonEmptyFoliageMeshNameIndices;
+		NonEmptyFoliageMeshNameIndices.Reserve(FoliageMeshNames.Num());
+ 
+		for (int32 i = 0; i < FoliageMeshNames.Num(); ++i)
+		{
+			if (!FoliageMeshNames[i].IsEmpty())
+			{
+				NonEmptyFoliageMeshNameIndices.Add(i);
+			}
+		}
+ 
+		if (NonEmptyFoliageMeshNameIndices.Num() == 0)
+		{
+			return;
+		}
+		
 		for (int32 BranchIndex = 0; BranchIndex < NumBranches; ++BranchIndex)
 		{
 			TArray<float> NormalizedAttachmentPoints;
@@ -320,13 +339,13 @@ void FPVFoliage::DistributeFoliage(FManagedArrayCollection& OutCollection, const
 			// Set foliage entries
 			TManagedArrayAccessor<FVector3f> PivotPointsAttribute(OutCollection, PV::AttributeNames::FoliagePivotPoint,
 				PV::GroupNames::FoliageGroup);
-			const int32 NumFoliageMeshes = FoliageFacadeOutput.NumFoliageNames();
+			const int32 NumFoliageMeshes = NonEmptyFoliageMeshNameIndices.Num();
 			int32 CurrentFoliageMeshIndex = BranchIndex % NumFoliageMeshes;
 			for (int32 Idx = 0; Idx < AttachmentPointsLengthsFromRoot.Num(); ++Idx)
 			{
 				const int32 NewIndex = PivotPointsAttribute.AddElements(1);
 				FoliageFacadeOutput.SetFoliageEntry(NewIndex, {
-					.NameId = CurrentFoliageMeshIndex,
+					.NameId = NonEmptyFoliageMeshNameIndices[CurrentFoliageMeshIndex],
 					.BranchId = BranchIndex,
 					.PivotPoint = AttachmentPointsPositions[Idx],
 					.UpVector = AttachmentPointsUpVectors[Idx],
diff --git a/Engine/Plugins/Experimental/ProceduralVegetationEditor/Source/ProceduralVegetationEditor/Private/Helpers/PVExportHelper.cpp b/Engine/Plugins/Experimental/ProceduralVegetationEditor/Source/ProceduralVegetationEditor/Private/Helpers/PVExportHelper.cpp
--- a/Engine/Plugins/Experimental/ProceduralVegetationEditor/Source/ProceduralVegetationEditor/Private/Helpers/PVExportHelper.cpp
+++ b/Engine/Plugins/Experimental/ProceduralVegetationEditor/Source/ProceduralVegetationEditor/Private/Helpers/PVExportHelper.cpp
@@ -654,6 +654,11 @@ namespace PV::Export::Internal
 		
 		for (const FString& FoliageInstanceObjectPath : UsedFoliage)
 		{
+			if (FoliageInstanceObjectPath.IsEmpty())
+			{
+				continue;
+			}
+			
 			TObjectPtr<UStaticMesh> PartStaticMesh = LoadObject<UStaticMesh>(nullptr, *FoliageInstanceObjectPath);
 			if (PartStaticMesh)
 			{
@@ -753,6 +758,11 @@ namespace PV::Export::Internal
 
 		for (const FString& FoliageInstanceObjectPath : UsedFoliage)
 		{
+			if (FoliageInstanceObjectPath.IsEmpty())
+			{
+				continue;
+			}
+			
 			TObjectPtr<USkeletalMesh> PartSkeletalMesh = LoadObject<USkeletalMesh>(nullptr, *FoliageInstanceObjectPath);
 			if (PartSkeletalMesh)
 			{
