Fixed foliage face orientation
- Issue/Fix: When foliage is spawned near the trunk, It is not up aligned. The expectation is that foliage should spawn with leafs pointing up.

diff --git a/Engine/Plugins/Experimental/ProceduralVegetationEditor/Source/ProceduralVegetation/Private/Facades/PVFoliageFacade.cpp b/Engine/Plugins/Experimental/ProceduralVegetationEditor/Source/ProceduralVegetation/Private/Facades/PVFoliageFacade.cpp
--- a/Engine/Plugins/Experimental/ProceduralVegetationEditor/Source/ProceduralVegetation/Private/Facades/PVFoliageFacade.cpp
+++ b/Engine/Plugins/Experimental/ProceduralVegetationEditor/Source/ProceduralVegetation/Private/Facades/PVFoliageFacade.cpp
@@ -400,20 +400,24 @@ namespace PV::Facades
 	FTransform FFoliageFacade::GetFoliageTransform(int32 Id) const
 	{
 		const FFoliageEntryData Data = GetFoliageEntry(Id);
-
-		const FQuat UpVectorRotationQuat = FQuat::FindBetweenNormals(FVector::UpVector, FVector(Data.UpVector).GetSafeNormal());
-		const FVector RotatedVectorFace = UpVectorRotationQuat.RotateVector(FVector::RightVector);
-		float Dot = FVector::DotProduct(RotatedVectorFace, FVector(Data.NormalVector));
-		float Angle = FMath::Acos(FMath::Clamp(Dot, -1.0f, 1.0f));
-		const FQuat FaceRotationQuat = FQuat(FVector(Data.UpVector).GetSafeNormal(), Angle);
-
-		const FTransform ScaleTransform(FQuat::Identity, FVector::ZeroVector, FVector(Data.Scale));
-		const FTransform FaceRotationTransform(FaceRotationQuat);
-		const FTransform UpVectorRotationTransform(UpVectorRotationQuat);
-		const FTransform TranslationTransform(FRotator::ZeroRotator, FVector(Data.PivotPoint));
-		const FTransform FinalTransform = ScaleTransform * UpVectorRotationTransform * FaceRotationTransform * TranslationTransform;
-
-		return FinalTransform;
+ 
+		const FVector UpVectorN = FVector(Data.UpVector).GetSafeNormal();
+		if (UpVectorN.IsNearlyZero())
+		{
+			return FTransform(FQuat::Identity, FVector(Data.PivotPoint), FVector(Data.Scale));
+		}
+ 
+		const FVector NormalVectorN = FVector(Data.NormalVector).GetSafeNormal();
+ 
+		FVector OrthogonalizedNormalVectorN = (NormalVectorN - FVector::DotProduct(NormalVectorN, UpVectorN) * UpVectorN).GetSafeNormal();
+		if (OrthogonalizedNormalVectorN.IsNearlyZero())
+		{
+			FVector TempVector;
+			UpVectorN.FindBestAxisVectors(TempVector, OrthogonalizedNormalVectorN);
+		}
+		
+		const FQuat RotationQuat = FQuat(FRotationMatrix::MakeFromYZ(OrthogonalizedNormalVectorN, UpVectorN));
+		return FTransform(RotationQuat, FVector(Data.PivotPoint), FVector(Data.Scale));
 	}
 
 	const TManagedArray<FVector3f>& FFoliageFacade::GetPivotPositions() const
