[PVE] Fix the missing bones for some parts of the mesh.

Reason if the last points NjordPixelID is a whole number the point was not considred as last point. Based on the Last point we end the segment for which we can reduce the bones, due to this check the segment was not getting added.

Fix the last point should just be the Last point in the branch.

diff --git a/Engine/Plugins/Experimental/ProceduralVegetationEditor/Source/ProceduralVegetation/Private/Facades/PVBoneFacade.cpp b/Engine/Plugins/Experimental/ProceduralVegetationEditor/Source/ProceduralVegetation/Private/Facades/PVBoneFacade.cpp
--- a/Engine/Plugins/Experimental/ProceduralVegetationEditor/Source/ProceduralVegetation/Private/Facades/PVBoneFacade.cpp
+++ b/Engine/Plugins/Experimental/ProceduralVegetationEditor/Source/ProceduralVegetation/Private/Facades/PVBoneFacade.cpp
@@ -364,8 +364,7 @@ namespace PV::Facades
 						}
 					}
 
-					bool IsLastPoint = !IsConnectingPoint && !(FMath::Frac(NjordPixelID) <= 0.0001f) && Idx == BranchPoints.Num() - 1;
-					//If this point exist in the BoneNodes set it as parent Bone
+					bool IsLastPoint = Idx == BranchPoints.Num() - 1;
 				
 					//UE_LOG(LogProceduralVegetation, Warning, TEXT("Iterating Branch Points NjordPixelId, Id{%f} PointIndex{%i} PointFacade.GetBudNumber(BranchPoints[PointIndex]){%i} BranchID{%i} ParentBoneIndex{%i} bFound{%i} FMath::Frac(NjordPixelID){%f}"),
 					//	NjordPixelID, PointIdx, PointFacade.GetBudNumber(PointIdx), BranchIndex, ParentBoneIndex, FoundBoneNode ? 1 : 0 , FMath::Frac(NjordPixelID));
