[PVE] Added UV1 and UV2 data also added Generation data in UV0.X

UV0.X contains the UV0.X + Generation-1
UV1. X Contains Length from 0-1 U mapping
UV2.X contains the length from root * 0.01
UV2.Y contains Normalized Point Scale

diff --git a/Engine/Plugins/Experimental/ProceduralVegetationEditor/Source/ProceduralVegetation/Private/Facades/PVTreeFacade.cpp b/Engine/Plugins/Experimental/ProceduralVegetationEditor/Source/ProceduralVegetation/Private/Facades/PVTreeFacade.cpp
--- a/Engine/Plugins/Experimental/ProceduralVegetationEditor/Source/ProceduralVegetation/Private/Facades/PVTreeFacade.cpp
+++ b/Engine/Plugins/Experimental/ProceduralVegetationEditor/Source/ProceduralVegetation/Private/Facades/PVTreeFacade.cpp
@@ -133,6 +133,21 @@ namespace PV::Facades
 			OutCollection, PointsToRemove, BranchesToRemove, FoliageInstancesToRemove);
 	}
 
+	int32 FTreeFacade::GetGenerationNumberFromSourceBud(const FManagedArrayCollection& Collection, int32 BranchIndex)
+	{
+		const FBranchFacade BranchFacade(Collection);
+		const FPointFacade PointFacade(Collection);
+		
+		const auto& BranchPoints = BranchFacade.GetPoints(BranchIndex);
+		check(BranchPoints.Num() > 0);
+		const int32 BranchRootPoint = BranchPoints[0];
+
+		TArray<int32> BudDevelopment = PointFacade.GetBudDevelopment(BranchRootPoint);
+		check(BudDevelopment.Num() > 0);
+	
+		return BudDevelopment[0];
+	}
+
 	void FTreeFacade::GatherChildBranches(const FBranchFacade& InBranchFacade, const int InParentIndex, TArray<int32>& OutBranchesToRemove)
 	{
 		const TArray<int32>& BranchChildren = InBranchFacade.GetChildren(InParentIndex);
diff --git a/Engine/Plugins/Experimental/ProceduralVegetationEditor/Source/ProceduralVegetation/Private/Facades/PVTreeFacade.h b/Engine/Plugins/Experimental/ProceduralVegetationEditor/Source/ProceduralVegetation/Private/Facades/PVTreeFacade.h
--- a/Engine/Plugins/Experimental/ProceduralVegetationEditor/Source/ProceduralVegetation/Private/Facades/PVTreeFacade.h
+++ b/Engine/Plugins/Experimental/ProceduralVegetationEditor/Source/ProceduralVegetation/Private/Facades/PVTreeFacade.h
@@ -30,7 +30,7 @@ namespace PV::Facades
 
 		static void RemoveBranches(const FBranchFacade& InBranchFacade, TArray<int>& InBranchesToRemove, FManagedArrayCollection& OutCollection);
 		
-
+		static int32 GetGenerationNumberFromSourceBud(const FManagedArrayCollection& Collection, int32 BranchIndex);
 	private:
 		static void GatherChildBranches(const FBranchFacade& InBranchFacade, const int InParentIndex, TArray<int32>& OutBranchesToRemove);
 	};
diff --git a/Engine/Plugins/Experimental/ProceduralVegetationEditor/Source/ProceduralVegetation/Private/Implementations/PVMaterialSettings.cpp b/Engine/Plugins/Experimental/ProceduralVegetationEditor/Source/ProceduralVegetation/Private/Implementations/PVMaterialSettings.cpp
--- a/Engine/Plugins/Experimental/ProceduralVegetationEditor/Source/ProceduralVegetation/Private/Implementations/PVMaterialSettings.cpp
+++ b/Engine/Plugins/Experimental/ProceduralVegetationEditor/Source/ProceduralVegetation/Private/Implementations/PVMaterialSettings.cpp
@@ -4,6 +4,7 @@
 #include "Algo/MaxElement.h"
 #include "Facades/PVBranchFacade.h"
 #include "Facades/PVPointFacade.h"
+#include "Facades/PVTreeFacade.h"
 
 struct FBranchUVData
 {
@@ -69,15 +70,20 @@ void FPVMaterialSettings::SetBranchUVMaterial(FManagedArrayCollection& Collectio
 		float PScale = PointFacade.GetPointScale(PointIndex);
 
 		check(BudDevelopment.Num() > 2);
-		MinGeneration = FMath::Min(MinGeneration, BudDevelopment[0]);
 		MinAge = FMath::Min(MinAge, BudDevelopment[2]);
 		MinPScale = FMath::Min(MinPScale, PScale);
 
-		MaxGeneration = FMath::Max(MaxGeneration, BudDevelopment[0]);
 		MaxAge = FMath::Max(MaxAge, BudDevelopment[2]);
 		MaxPScale = FMath::Max(MaxPScale, PScale);
 	}
 	
+	for (int32 BranchIndex = 0; BranchIndex < BranchFacade.GetElementCount(); BranchIndex++)
+	{
+		const int32 GenerationNumber = PV::Facades::FTreeFacade::GetGenerationNumberFromSourceBud(Collection, BranchIndex);
+		MinGeneration = FMath::Min(MinGeneration, GenerationNumber);
+		MaxGeneration = FMath::Max(MaxGeneration, GenerationNumber);
+	}
+
 	for (int32 BranchIndex = 0; BranchIndex < BranchFacade.GetElementCount(); BranchIndex++)
 	{
 		const auto& BranchPoints = BranchFacade.GetPoints(BranchIndex);
@@ -86,13 +92,15 @@ void FPVMaterialSettings::SetBranchUVMaterial(FManagedArrayCollection& Collectio
 
 		float BranchMaterialVal = 0;
 
+		const int32 GenerationNumber = PV::Facades::FTreeFacade::GetGenerationNumberFromSourceBud(Collection, BranchIndex);
+	
 		TArray<int> BudDevelopment = PointFacade.GetBudDevelopment(BranchRootPoint);
 		float PScale = PointFacade.GetPointScale(BranchRootPoint);
-		check(BudDevelopment.Num() >= 2);
+		check(BudDevelopment.Num() > 2);
 		
 		if (MaterialMode == EUVMaterialMode::Generation)
 		{
-			BranchMaterialVal = GetMaterialValue(BudDevelopment[0], MinGeneration, MaxGeneration, MinGeneration, MaxGeneration);
+			BranchMaterialVal = GetMaterialValue(GenerationNumber, MinGeneration, MaxGeneration, MinGeneration, MaxGeneration);
 		}
 		else if (MaterialMode == EUVMaterialMode::Age)
 		{
diff --git a/Engine/Plugins/Experimental/ProceduralVegetationEditor/Source/ProceduralVegetation/Private/Implementations/PVMeshBuilder.cpp b/Engine/Plugins/Experimental/ProceduralVegetationEditor/Source/ProceduralVegetation/Private/Implementations/PVMeshBuilder.cpp
--- a/Engine/Plugins/Experimental/ProceduralVegetationEditor/Source/ProceduralVegetation/Private/Implementations/PVMeshBuilder.cpp
+++ b/Engine/Plugins/Experimental/ProceduralVegetationEditor/Source/ProceduralVegetation/Private/Implementations/PVMeshBuilder.cpp
@@ -6,13 +6,13 @@
 
 #include "ProceduralVegetationModule.h"
 #include "Algo/MaxElement.h"
-#include "Elements/PCGPrintElement.h"
 #include "Facades/PVBoneFacade.h"
 #include "Facades/PVBranchFacade.h"
 #include "Facades/PVBudVectorsFacade.h"
 #include "Facades/PVFoliageFacade.h"
 #include "Facades/PVPlantFacade.h"
 #include "Facades/PVPointFacade.h"
+#include "Facades/PVTreeFacade.h"
 
 void MergeMaterials(const TArray<FTrunkGenerationMaterialSetup>& InMaterialSetups, TArray<FLocalDynamicMeshData>& MeshDatas)
 {
@@ -79,11 +79,13 @@ void FPVMeshBuilder::GenerateGeometryCollection(const FManagedArrayCollection& I
 	OutGeometryCollection.AddElements(1, OutGeometryCollection.GeometryGroup);
 	OutGeometryCollection.AddElements(TriangleCount, OutGeometryCollection.FacesGroup);
 	OutGeometryCollection.AddElements(VertexCount, OutGeometryCollection.VerticesGroup);
-	OutGeometryCollection.SetNumUVLayers(1);
+	OutGeometryCollection.SetNumUVLayers(3);
 
 	TManagedArray<FVector3f>& Vertices = OutGeometryCollection.Vertex;
 	TManagedArray<FVector3f>& Normals = OutGeometryCollection.Normal;
 	TManagedArray<FVector2f>& UVs = *OutGeometryCollection.FindUVLayer(0);
+	TManagedArray<FVector2f>& UVs1 = *OutGeometryCollection.FindUVLayer(1);
+	TManagedArray<FVector2f>& UVs2 = *OutGeometryCollection.FindUVLayer(2);
 	TManagedArray<FIntVector>& Indices = OutGeometryCollection.Indices;
 	TManagedArray<int32>& MaterialIDAttribute = OutGeometryCollection.MaterialID;
 
@@ -164,6 +166,8 @@ void FPVMeshBuilder::GenerateGeometryCollection(const FManagedArrayCollection& I
 			Vertices[VertexIndex] = static_cast<FVector3f>(Vertex.Position);
 			Normals[VertexIndex] = static_cast<FVector3f>(Vertex.Normal);
 			UVs[VertexIndex] = Vertex.UV;
+			UVs1[VertexIndex] = Vertex.UV1;
+			UVs2[VertexIndex] = Vertex.UV2;
 			PointIDs[VertexIndex] = Vertex.NjordPixelId;
 			++VertexIndex;
 		}
@@ -744,7 +748,7 @@ void FPVMeshBuilder::GenerateBranchMeshData(const bool bPrimitiveIsTrunk, const
 	const TManagedArray<FVector3f>& PointPositions = PointFacade.GetPositions();
 	const TManagedArray<TArray<FVector3f>>& PointBudDirections = BudVectorFacade.GetBudDirections();
 	const TManagedArray<float>& PointScales = PointFacade.GetPointScales();
-
+	
 	const float MaxPointScale = *Algo::MaxElement(PointScales);
 	const float MaxScaleRatio = 1.0f / (MaxPointScale * UE_TWO_PI);
 	int32 PolyGroupIndex = 0;
@@ -762,11 +766,18 @@ void FPVMeshBuilder::GenerateBranchMeshData(const bool bPrimitiveIsTrunk, const
 
 	float Displacement_UV_V = 0.0f;
 	FVector3f PreviousPosition = PointPositions[PrimitivePoints[0]];
+
+	auto SafeDivide = [](const double InNumerator, const double InDenominator)
+	{
+		return FMath::IsNearlyZero(InDenominator) ? InNumerator : (InNumerator / InDenominator);
+	};
 	
 	for (int32 Index = 0; Index < PrimitivePoints.Num(); ++Index)
 	{
 		const int32& PointIndex = PrimitivePoints[Index];
 		const float NjordPixelId = PointIndex;
+		
+		const float LengthFromRoot = PointFacade.GetLengthFromRoot(PointIndex);
 
 		float TextureCoordV = PointFacade.GetTextureCoordV(PointIndex);
 		float TextureCoordUOffset = PointFacade.GetTextureCoordUOffset(PointIndex);
@@ -827,12 +838,16 @@ void FPVMeshBuilder::GenerateBranchMeshData(const bool bPrimitiveIsTrunk, const
 		PreviousAimVector = UpdatedAimVector;
 		PreviousUpVector = PointUp;
 
+		FVector2f UV1 = FVector2f(0, 1.0 - Displacement_UV_V);
+		const FVector2f UV2 = FVector2f(LengthFromRoot * 0.01, SafeDivide(PointScale, MaxPointScale));
+		const float UVOGenerationOffset = GenerationNumber - 1;
+
 		const float BranchGradient = PointFacade.GetBranchGradient(PointIndex);
 		const float ProfileFallOff = MeshBuilderParams.PlantProfileFallOff.GetRichCurveConst()->Eval(BranchGradient);
 		if (PointScale == 0.0f)
 		[[unlikely]]
 		{
-			CurrentIndices.Add(OutLocalMeshData.Vertices.Emplace(PointPosition, UpdatedAimVector, FVector2f(0.0f, TextureCoordV), NjordPixelId));
+			CurrentIndices.Add(OutLocalMeshData.Vertices.Emplace(PointPosition, UpdatedAimVector, FVector2f(UVOGenerationOffset, TextureCoordV), UV1, UV2, NjordPixelId));
 		}
 		else
 		[[likely]]
@@ -843,9 +858,11 @@ void FPVMeshBuilder::GenerateBranchMeshData(const bool bPrimitiveIsTrunk, const
 			{
 				float UV_U = (1 - (static_cast<float>(Division) / PointDivisions)) + TextureCoordUOffset;
 				UV_U = FMath::GetMappedRangeValueUnclamped(OldRange, NewRange, FMath::Wrap(UV_U, OldRange.X, OldRange.Y));
+				UV_U += UVOGenerationOffset;
 
 				float Displacement_UV_U = (static_cast<float>(Division) / PointDivisions);
-					
+				UV1.X = 1 - Displacement_UV_U;
+				
 				float ProfileMultiplier = 1.0f;
 				if (bShouldApplyProfile)
 				{
@@ -881,15 +898,18 @@ void FPVMeshBuilder::GenerateBranchMeshData(const bool bPrimitiveIsTrunk, const
 
 				const FVector3f Normal = FVector3f((RadialPoint - PointPosition).GetUnsafeNormal());
 
-				CurrentIndices.Add(OutLocalMeshData.Vertices.Emplace(RadialPoint, Normal, FVector2f(UV_U, TextureCoordV), NjordPixelId));
+				CurrentIndices.Add(OutLocalMeshData.Vertices.Emplace(RadialPoint, Normal, FVector2f(UV_U, TextureCoordV), UV1, UV2, NjordPixelId));
 			}
 		}
 
+		FVector2f UV0 = FVector2f(FMath::GetMappedRangeValueUnclamped(OldRange, NewRange, FMath::Wrap(TextureCoordUOffset, OldRange.X, OldRange.Y)),
+				TextureCoordV);
+		UV0.X += UVOGenerationOffset;
+
+		UV1.X = 0;
+		
 		const FLocalDynamicMeshData::FVertex FirstVertex = OutLocalMeshData.Vertices[CurrentIndices[0]];
-		CurrentIndices.Add(OutLocalMeshData.Vertices.Emplace(FirstVertex.Position, FirstVertex.Normal,
-			FVector2f(FMath::GetMappedRangeValueUnclamped(OldRange, NewRange, FMath::Wrap(TextureCoordUOffset, OldRange.X, OldRange.Y)),
-				TextureCoordV),
-			NjordPixelId));
+		CurrentIndices.Add(OutLocalMeshData.Vertices.Emplace(FirstVertex.Position, FirstVertex.Normal,UV0, UV1, UV2 ,NjordPixelId));
 
 		PreviousPosition = PointPosition;
 		
@@ -1050,7 +1070,7 @@ void FPVMeshBuilder::GetLocalDynamicMeshData(FManagedArrayCollection& Collection
 						});
 
 				const bool PrimitiveIsTrunk = PlantFacade.IsTrunkIndex(PrimitiveIndex);
-				const int32 GenerationNumber = BranchFacade.GetHierarchyGenerationNumber(PrimitiveIndex);
+				const int32 GenerationNumber = PV::Facades::FTreeFacade::GetGenerationNumberFromSourceBud(Collection, PrimitiveIndex);
 
 				MeshesData[PrimitiveIndex].MaterialID = BranchFacade.GetBranchUVMaterial(PrimitiveIndex);
 
diff --git a/Engine/Plugins/Experimental/ProceduralVegetationEditor/Source/ProceduralVegetation/Private/Implementations/PVMeshBuilder.h b/Engine/Plugins/Experimental/ProceduralVegetationEditor/Source/ProceduralVegetation/Private/Implementations/PVMeshBuilder.h
--- a/Engine/Plugins/Experimental/ProceduralVegetationEditor/Source/ProceduralVegetation/Private/Implementations/PVMeshBuilder.h
+++ b/Engine/Plugins/Experimental/ProceduralVegetationEditor/Source/ProceduralVegetation/Private/Implementations/PVMeshBuilder.h
@@ -226,6 +226,8 @@ struct FLocalDynamicMeshData
 		FVector3f Position;
 		FVector3f Normal;
 		FVector2f UV;
+		FVector2f UV1;
+		FVector2f UV2;
 		float NjordPixelId;
 	};
 
