[PVE] Foliage condition system

- Added attributes in foliage palette
- Implemented the condition system in the foliage distributor
- Implemented reading foliage attributes and conditions from JSON

Attributes are added with each foliage palette entry. These are 0-1 values that user can set with foliage entry. For each, if you want a specific foliage to only appear on tip, then you just set the tip to true and enable the tip condition in the foliage distributor.

Main logic in files:
FPVFoliage

UI changes in:
PVFoliagePaletteSettings
PVFoliageDistributorSettings

Files added to avoid changes in existing public headers:
PVFoliageAttributeNames
PVFoliageJSONHelper

diff --git a/Engine/Plugins/Experimental/ProceduralVegetationEditor/Source/ProceduralVegetation/Private/Facades/PVBudVectorsFacade.h b/Engine/Plugins/Experimental/ProceduralVegetationEditor/Source/ProceduralVegetation/Private/Facades/PVBudVectorsFacade.h
--- a/Engine/Plugins/Experimental/ProceduralVegetationEditor/Source/ProceduralVegetation/Private/Facades/PVBudVectorsFacade.h
+++ b/Engine/Plugins/Experimental/ProceduralVegetationEditor/Source/ProceduralVegetation/Private/Facades/PVBudVectorsFacade.h
@@ -7,9 +7,10 @@
 
 namespace PV::Facades
 {
+	constexpr int BudDirectionsApical = 0;
 	constexpr int BudDirectionsLightOptimalIndex = 2;
 	constexpr int BudDirectionsLightSubOptimal = 3;
-
+	constexpr int BudLightDetectedAvailableIndex = 0;
 	/**
 	 * FBudVectorsFacade is used to access and manipulate the BudVectors data from Point Group with in the ProceduralVegetation's FManagedArrayCollection
 	 * BudVectors represent different growth related vectors in plant data and are associated with each point. 
diff --git a/Engine/Plugins/Experimental/ProceduralVegetationEditor/Source/ProceduralVegetation/Private/Facades/PVFoliageAttributesNames.h b/Engine/Plugins/Experimental/ProceduralVegetationEditor/Source/ProceduralVegetation/Private/Facades/PVFoliageAttributesNames.h
new file mode 100644
--- /dev/null
+++ b/Engine/Plugins/Experimental/ProceduralVegetationEditor/Source/ProceduralVegetation/Private/Facades/PVFoliageAttributesNames.h
@@ -0,0 +1,42 @@
+// Copyright Epic Games, Inc. All Rights Reserved.
+#pragma once
+
+#include "CoreMinimal.h"
+
+namespace PV::GroupNames
+{
+	inline static const FName FoliageConditionGroup = FName(TEXT("FoliageConditions"));
+}
+
+namespace PV::FoliageAttributeNames
+{
+
+	//Foliage Attributes
+	inline static const FName FoliageNameID = FName(TEXT("FoliageNameID"));
+	inline static const FName FoliageBranchID = FName(TEXT("FoliageBranchID"));
+	inline static const FName FoliagePivotPoint = FName(TEXT("FoliagePivotPoint"));
+	inline static const FName FoliageUPVector = FName(TEXT("FoliageUPVector"));
+	inline static const FName FoliageNormalVector = FName(TEXT("FoliageNormalVector"));
+	inline static const FName FoliageScale = FName(TEXT("FoliageScale"));
+	inline static const FName FoliageLengthFromRoot = FName(TEXT("LengthFromRoot"));
+	inline static const FName FoliageConditionLight = FName(TEXT("ConditionLight"));
+	inline static const FName FoliageConditionTip = FName(TEXT("ConditionTip"));
+	inline static const FName FoliageConditionScale = FName(TEXT("ConditionScale"));
+	inline static const FName FoliageConditionUpAlignment = FName(TEXT("ConditionUpAlignment"));
+	inline static const FName FoliageConditionHealth = FName(TEXT("ConditionHealth"));
+	inline static const FName FoliageParentBoneID = FName(TEXT("FoliageParentBoneID"));
+
+	//FoliageInfo Attributes
+	inline static const FName FoliageName = FName(TEXT("Name"));
+	inline static const FName FoliageAttributeLight = FName(TEXT("Light"));
+	inline static const FName FoliageAttributeScale = FName(TEXT("Scale"));
+	inline static const FName FoliageAttributeTip = FName(TEXT("Tip"));
+	inline static const FName FoliageAttributeUpAlignment = FName(TEXT("UpAlignment"));
+	inline static const FName FoliageAttributeHealth = FName(TEXT("Health"));
+
+	// Condition
+	inline static const FName ConditionName = FName(TEXT("Name"));
+	inline static const FName ConditionWeight = FName(TEXT("Weight"));
+	inline static const FName ConditionOffset = FName(TEXT("Offset"));
+}
+
diff --git a/Engine/Plugins/Experimental/ProceduralVegetationEditor/Source/ProceduralVegetation/Private/Facades/PVFoliageConditionFacade.cpp b/Engine/Plugins/Experimental/ProceduralVegetationEditor/Source/ProceduralVegetation/Private/Facades/PVFoliageConditionFacade.cpp
new file mode 100644
--- /dev/null
+++ b/Engine/Plugins/Experimental/ProceduralVegetationEditor/Source/ProceduralVegetation/Private/Facades/PVFoliageConditionFacade.cpp
@@ -0,0 +1,111 @@
+﻿// Copyright Epic Games, Inc. All Rights Reserved.
+
+#include "Facades/PVFoliageConditionFacade.h"
+#include "PVFoliageAttributesNames.h"
+
+namespace PV::Facades
+{
+	FFoliageConditionFacade::FFoliageConditionFacade(FManagedArrayCollection& InCollection)
+		: NameAttribute(InCollection, FoliageAttributeNames::ConditionName, GroupNames::FoliageConditionGroup)
+		, WeightAttribute(InCollection, FoliageAttributeNames::ConditionWeight, GroupNames::FoliageConditionGroup)
+		, OffsetAttribute(InCollection, FoliageAttributeNames::ConditionOffset, GroupNames::FoliageConditionGroup)
+	{
+		DefineSchema(InCollection);
+	}
+
+	FFoliageConditionFacade::FFoliageConditionFacade(const FManagedArrayCollection& InCollection)
+	: NameAttribute(InCollection, FoliageAttributeNames::ConditionName, GroupNames::FoliageConditionGroup)
+	, WeightAttribute(InCollection, FoliageAttributeNames::ConditionWeight, GroupNames::FoliageConditionGroup)
+	, OffsetAttribute(InCollection, FoliageAttributeNames::ConditionOffset, GroupNames::FoliageConditionGroup)
+	{
+	}
+
+	void FFoliageConditionFacade::DefineSchema(FManagedArrayCollection& InCollection)
+	{
+		if (!InCollection.HasGroup(GroupNames::FoliageConditionGroup))
+		{
+			InCollection.AddGroup(GroupNames::FoliageConditionGroup);
+		}
+
+		NameAttribute.Add();
+		WeightAttribute.Add();
+		OffsetAttribute.Add();
+	}
+
+	bool FFoliageConditionFacade::IsValid() const
+	{
+		return NameAttribute.IsValid() && WeightAttribute.IsValid() && OffsetAttribute.IsValid();
+	}
+
+	FFoliageConditonInfo FFoliageConditionFacade::GetEntry(const int32 Index) const
+	{
+		FFoliageConditonInfo ReturnData = {};
+		if (IsValid() && Index > -1)
+		{
+			if (NameAttribute.IsValidIndex(Index))
+			{
+				ReturnData.Name = NameAttribute.Get()[Index];
+			}
+			if (WeightAttribute.IsValidIndex(Index))
+			{
+				ReturnData.Weight = WeightAttribute.Get()[Index];
+			}
+			if (OffsetAttribute.IsValidIndex(Index))
+			{
+				ReturnData.Offset = OffsetAttribute.Get()[Index];
+			}
+		}
+		return ReturnData;
+	}
+
+	int32 FFoliageConditionFacade::Add(const FFoliageConditonInfo& InputData)
+	{
+		if (IsValid())
+		{
+			const int32 NewIndex = NameAttribute.AddElements(1);
+			NameAttribute.Modify()[NewIndex] = InputData.Name;
+			WeightAttribute.Modify()[NewIndex] = InputData.Weight;
+			OffsetAttribute.Modify()[NewIndex] = InputData.Offset;
+
+			return NewIndex;
+		}
+		return INDEX_NONE;
+	}
+
+	void FFoliageConditionFacade::Set(const int32 Index, const FFoliageConditonInfo& InputData)
+	{
+		if (IsValid() && Index >= 0)
+		{
+			NameAttribute.ModifyAt(Index, InputData.Name);
+			WeightAttribute.ModifyAt(Index, InputData.Weight);
+			OffsetAttribute.ModifyAt(Index, InputData.Offset);
+		}
+	}
+
+	TArray<FFoliageConditonInfo> FFoliageConditionFacade::GetData() const
+	{
+		TArray<FFoliageConditonInfo> Infos;
+		
+		if (NameAttribute.IsValid())
+		{
+			for (int32 i = 0; i < NameAttribute.Num(); i++)
+			{
+				Infos.Add(GetEntry(i));
+			}
+		}
+
+		return Infos;
+	}
+
+	void FFoliageConditionFacade::SetData(const TArray<FFoliageConditonInfo>& Infos)
+	{
+		int32 NumElements = NameAttribute.AddElements(Infos.Num());
+		
+		for (int32 Index = 0; Index < Infos.Num(); ++Index)
+		{
+			FFoliageConditonInfo Info = Infos[Index];
+
+			Set(Index, Info);
+		}
+	}
+}
diff --git a/Engine/Plugins/Experimental/ProceduralVegetationEditor/Source/ProceduralVegetation/Private/Facades/PVFoliageConditionFacade.h b/Engine/Plugins/Experimental/ProceduralVegetationEditor/Source/ProceduralVegetation/Private/Facades/PVFoliageConditionFacade.h
new file mode 100644
--- /dev/null
+++ b/Engine/Plugins/Experimental/ProceduralVegetationEditor/Source/ProceduralVegetation/Private/Facades/PVFoliageConditionFacade.h
@@ -0,0 +1,63 @@
+﻿// Copyright Epic Games, Inc. All Rights Reserved.
+
+#pragma once
+
+#include "CoreMinimal.h"
+#include "GeometryCollection/ManagedArrayAccessor.h"
+#include "Implementations/PVFoliage.h"
+
+struct FPVFoliageInfo;
+
+namespace PV::Facades
+{
+	struct PROCEDURALVEGETATION_API FFoliageConditonInfo
+	{
+		FString Name;
+		float Weight;
+		float Offset;
+
+		EFoliageDistributionCondition GetType() const
+		{
+			const UEnum* EnumPtr = StaticEnum<EFoliageDistributionCondition>();
+
+			int64 EnumValue = EnumPtr->GetValueByNameString(Name);
+
+			if (EnumValue == INDEX_NONE)
+				return EFoliageDistributionCondition::None;
+
+			return static_cast<EFoliageDistributionCondition>(EnumValue);
+		}
+
+		FString ToString() const
+		{
+			return FString::Printf(TEXT("Name = %s, Weight = %f, Reduction = %f"), *Name, Weight, Offset);
+		}
+	};
+
+	class PROCEDURALVEGETATION_API FFoliageConditionFacade
+	{
+	public:
+		FFoliageConditionFacade(FManagedArrayCollection& InCollection);
+		FFoliageConditionFacade(const FManagedArrayCollection& InCollection);
+
+		bool IsValid() const;
+
+		int32 NumEntries() const { return NameAttribute.Num(); };
+		
+		int32 Add(const FFoliageConditonInfo& InputData);
+
+		void Set(const int32 Index, const FFoliageConditonInfo& InputData);
+		void SetData(const TArray<FFoliageConditonInfo>& Infos);
+		
+		FFoliageConditonInfo GetEntry(const int32 Index) const;
+
+		TArray<FFoliageConditonInfo> GetData() const;
+
+	protected:
+		void DefineSchema(FManagedArrayCollection& InCollection);
+		
+		TManagedArrayAccessor<FString> NameAttribute;
+		TManagedArrayAccessor<float> WeightAttribute;
+		TManagedArrayAccessor<float> OffsetAttribute;
+	};
+}
diff --git a/Engine/Plugins/Experimental/ProceduralVegetationEditor/Source/ProceduralVegetation/Private/Facades/PVFoliageFacade.cpp b/Engine/Plugins/Experimental/ProceduralVegetationEditor/Source/ProceduralVegetation/Private/Facades/PVFoliageFacade.cpp
--- a/Engine/Plugins/Experimental/ProceduralVegetationEditor/Source/ProceduralVegetation/Private/Facades/PVFoliageFacade.cpp
+++ b/Engine/Plugins/Experimental/ProceduralVegetationEditor/Source/ProceduralVegetation/Private/Facades/PVFoliageFacade.cpp
@@ -3,24 +3,36 @@
 #include "Facades/PVFoliageFacade.h"
 
 #include "ProceduralVegetationModule.h"
+#include "PVFoliageAttributesNames.h"
 #include "VisualizeTexture.h"
 #include "Chaos/Deformable/MuscleActivationConstraints.h"
 #include "Facades/PVAttributesNames.h"
+#include "Implementations/PVFoliage.h"
 
 namespace PV::Facades
 {
 	FFoliageFacade::FFoliageFacade(FManagedArrayCollection& InCollection, const int32 InitialSize)
 		: ConstCollection(InCollection)
 		, Collection(&InCollection)
-		, NamesAttribute(InCollection, AttributeNames::FoliageName, GroupNames::FoliageNamesGroup)
-		, NameIdsAttribute(InCollection, AttributeNames::FoliageNameID, GroupNames::FoliageGroup)
-		, BranchIdsAttribute(InCollection, AttributeNames::FoliageBranchID, GroupNames::FoliageGroup)
-		, PivotPointsAttribute(InCollection, AttributeNames::FoliagePivotPoint, GroupNames::FoliageGroup)
-		, UpVectorsAttribute(InCollection, AttributeNames::FoliageUPVector, GroupNames::FoliageGroup)
-		, NormalVectorsAttribute(InCollection, AttributeNames::FoliageNormalVector, GroupNames::FoliageGroup)
-		, ScalesAttribute(InCollection, AttributeNames::FoliageScale, GroupNames::FoliageGroup)
-		, LengthFromRootAttribute(InCollection, AttributeNames::FoliageLengthFromRoot, GroupNames::FoliageGroup)
-		, ParentBoneIdsAttribute(InCollection, AttributeNames::FoliageParentBoneID, GroupNames::FoliageGroup)
+		, FoliageInfoNamesAttribute(InCollection, FoliageAttributeNames::FoliageName, GroupNames::FoliageNamesGroup)
+		, FoliageInfoLightAttribute(InCollection, FoliageAttributeNames::FoliageAttributeLight, GroupNames::FoliageNamesGroup)
+		, FoliageInfoScaleAttribute(InCollection, FoliageAttributeNames::FoliageAttributeScale, GroupNames::FoliageNamesGroup)
+		, FoliageInfoTipAttribute(InCollection, FoliageAttributeNames::FoliageAttributeTip, GroupNames::FoliageNamesGroup)
+		, FoliageInfoUpAlignmentAttribute(InCollection, FoliageAttributeNames::FoliageAttributeUpAlignment, GroupNames::FoliageNamesGroup)
+		, FoliageInfoHealthAttribute(InCollection, FoliageAttributeNames::FoliageAttributeHealth, GroupNames::FoliageNamesGroup)
+		, NameIdsAttribute(InCollection, FoliageAttributeNames::FoliageNameID, GroupNames::FoliageGroup)
+		, BranchIdsAttribute(InCollection, FoliageAttributeNames::FoliageBranchID, GroupNames::FoliageGroup)
+		, PivotPointsAttribute(InCollection, FoliageAttributeNames::FoliagePivotPoint, GroupNames::FoliageGroup)
+		, UpVectorsAttribute(InCollection, FoliageAttributeNames::FoliageUPVector, GroupNames::FoliageGroup)
+		, NormalVectorsAttribute(InCollection, FoliageAttributeNames::FoliageNormalVector, GroupNames::FoliageGroup)
+		, ScalesAttribute(InCollection, FoliageAttributeNames::FoliageScale, GroupNames::FoliageGroup)
+		, LengthFromRootAttribute(InCollection, FoliageAttributeNames::FoliageLengthFromRoot, GroupNames::FoliageGroup)
+		, ConditionLightAttribute(InCollection, FoliageAttributeNames::FoliageConditionLight, GroupNames::FoliageGroup)
+		, ConditionScaleAttribute(InCollection, FoliageAttributeNames::FoliageConditionScale, GroupNames::FoliageGroup)
+		, ConditionTipAttribute(InCollection, FoliageAttributeNames::FoliageConditionTip, GroupNames::FoliageGroup)
+		, ConditionUpAlignmentAttribute(InCollection, FoliageAttributeNames::FoliageConditionUpAlignment, GroupNames::FoliageGroup)
+		, ConditionHealthAttribute(InCollection, FoliageAttributeNames::FoliageConditionHealth, GroupNames::FoliageGroup)
+		, ParentBoneIdsAttribute(InCollection, FoliageAttributeNames::FoliageParentBoneID, GroupNames::FoliageGroup)
 		, FoliageIdsAttribute(InCollection, AttributeNames::BranchFoliageIDs, GroupNames::BranchGroup)
 		, FoliagePathAttribute(InCollection, AttributeNames::FoliagePath, GroupNames::DetailsGroup)
 	{
@@ -30,15 +42,25 @@ namespace PV::Facades
 	FFoliageFacade::FFoliageFacade(const FManagedArrayCollection& InCollection)
 		: ConstCollection(InCollection)
 		, Collection(nullptr)
-		, NamesAttribute(InCollection, AttributeNames::FoliageName, GroupNames::FoliageNamesGroup)
-		, NameIdsAttribute(InCollection, AttributeNames::FoliageNameID, GroupNames::FoliageGroup)
-		, BranchIdsAttribute(InCollection, AttributeNames::FoliageBranchID, GroupNames::FoliageGroup)
-		, PivotPointsAttribute(InCollection, AttributeNames::FoliagePivotPoint, GroupNames::FoliageGroup)
-		, UpVectorsAttribute(InCollection, AttributeNames::FoliageUPVector, GroupNames::FoliageGroup)
-		, NormalVectorsAttribute(InCollection, AttributeNames::FoliageNormalVector, GroupNames::FoliageGroup)
-		, ScalesAttribute(InCollection, AttributeNames::FoliageScale, GroupNames::FoliageGroup)
-		, LengthFromRootAttribute(InCollection, AttributeNames::FoliageLengthFromRoot, GroupNames::FoliageGroup)
-		, ParentBoneIdsAttribute(InCollection, AttributeNames::FoliageParentBoneID, GroupNames::FoliageGroup)
+		, FoliageInfoNamesAttribute(InCollection, FoliageAttributeNames::FoliageName, GroupNames::FoliageNamesGroup)
+		, FoliageInfoLightAttribute(InCollection, FoliageAttributeNames::FoliageAttributeLight, GroupNames::FoliageNamesGroup)
+		, FoliageInfoScaleAttribute(InCollection, FoliageAttributeNames::FoliageAttributeScale, GroupNames::FoliageNamesGroup)
+		, FoliageInfoTipAttribute(InCollection, FoliageAttributeNames::FoliageAttributeTip, GroupNames::FoliageNamesGroup)
+		, FoliageInfoUpAlignmentAttribute(InCollection, FoliageAttributeNames::FoliageAttributeUpAlignment, GroupNames::FoliageNamesGroup)
+		, FoliageInfoHealthAttribute(InCollection, FoliageAttributeNames::FoliageAttributeHealth, GroupNames::FoliageNamesGroup)
+		, NameIdsAttribute(InCollection, FoliageAttributeNames::FoliageNameID, GroupNames::FoliageGroup)
+		, BranchIdsAttribute(InCollection, FoliageAttributeNames::FoliageBranchID, GroupNames::FoliageGroup)
+		, PivotPointsAttribute(InCollection, FoliageAttributeNames::FoliagePivotPoint, GroupNames::FoliageGroup)
+		, UpVectorsAttribute(InCollection, FoliageAttributeNames::FoliageUPVector, GroupNames::FoliageGroup)
+		, NormalVectorsAttribute(InCollection, FoliageAttributeNames::FoliageNormalVector, GroupNames::FoliageGroup)
+		, ScalesAttribute(InCollection, FoliageAttributeNames::FoliageScale, GroupNames::FoliageGroup)
+		, LengthFromRootAttribute(InCollection, FoliageAttributeNames::FoliageLengthFromRoot, GroupNames::FoliageGroup)
+		, ConditionLightAttribute(InCollection, FoliageAttributeNames::FoliageConditionLight, GroupNames::FoliageGroup)
+		, ConditionScaleAttribute(InCollection, FoliageAttributeNames::FoliageConditionScale, GroupNames::FoliageGroup)
+		, ConditionTipAttribute(InCollection, FoliageAttributeNames::FoliageConditionTip, GroupNames::FoliageGroup)
+		, ConditionUpAlignmentAttribute(InCollection, FoliageAttributeNames::FoliageConditionUpAlignment, GroupNames::FoliageGroup)
+		, ConditionHealthAttribute(InCollection, FoliageAttributeNames::FoliageConditionHealth, GroupNames::FoliageGroup)
+		, ParentBoneIdsAttribute(InCollection, FoliageAttributeNames::FoliageParentBoneID, GroupNames::FoliageGroup)
 		, FoliageIdsAttribute(InCollection, AttributeNames::BranchFoliageIDs, GroupNames::BranchGroup)
 		, FoliagePathAttribute(InCollection, AttributeNames::FoliagePath, GroupNames::DetailsGroup)
 	{
@@ -67,7 +89,12 @@ namespace PV::Facades
 			Collection->AddGroup(GroupNames::DetailsGroup);
 		}
 
-		NamesAttribute.Add();
+		FoliageInfoNamesAttribute.Add();
+		FoliageInfoLightAttribute.Add();
+		FoliageInfoScaleAttribute.Add();
+		FoliageInfoTipAttribute.Add();
+		FoliageInfoUpAlignmentAttribute.Add();
+		FoliageInfoHealthAttribute.Add();
 		NameIdsAttribute.Add();
 		BranchIdsAttribute.Add();
 		PivotPointsAttribute.Add();
@@ -75,6 +102,11 @@ namespace PV::Facades
 		NormalVectorsAttribute.Add();
 		ScalesAttribute.Add();
 		LengthFromRootAttribute.Add();
+		ConditionLightAttribute.Add();
+		ConditionTipAttribute.Add();
+		ConditionUpAlignmentAttribute.Add();
+		ConditionHealthAttribute.Add();
+		ConditionScaleAttribute.Add();
 		ParentBoneIdsAttribute.Add();
 		FoliageIdsAttribute.Add();
 		FoliagePathAttribute.Add();
@@ -82,7 +114,8 @@ namespace PV::Facades
 
 	bool FFoliageFacade::IsValid() const
 	{
-		return NamesAttribute.IsValid() && NameIdsAttribute.IsValid() && BranchIdsAttribute.IsValid() && PivotPointsAttribute.IsValid() &&
+		return FoliageInfoNamesAttribute.IsValid() && FoliageInfoTipAttribute.IsValid() && FoliageInfoUpAlignmentAttribute.IsValid() && FoliageInfoHealthAttribute.IsValid() && FoliageInfoLightAttribute.IsValid() && FoliageInfoScaleAttribute.IsValid() && 
+			NameIdsAttribute.IsValid() && BranchIdsAttribute.IsValid() && PivotPointsAttribute.IsValid() &&
 			UpVectorsAttribute.IsValid() && NormalVectorsAttribute.IsValid() && ScalesAttribute.IsValid() && LengthFromRootAttribute.IsValid() &&
 			FoliageIdsAttribute.IsValid() && ParentBoneIdsAttribute.IsValid();
 	}
@@ -120,6 +153,32 @@ namespace PV::Facades
 			{
 				ReturnData.LengthFromRoot = LengthFromRootAttribute.Get()[Index];
 			}
+
+			if (ConditionLightAttribute.IsValidIndex(Index))
+			{
+				ReturnData.ConditionLight = ConditionLightAttribute.Get()[Index];
+			}
+
+			if (ConditionScaleAttribute.IsValidIndex(Index))
+			{
+				ReturnData.ConditionScale = ConditionScaleAttribute.Get()[Index];
+			}
+
+			if (ConditionTipAttribute.IsValidIndex(Index))
+			{
+				ReturnData.ConditionTip = ConditionTipAttribute.Get()[Index];
+			}
+
+			if (ConditionUpAlignmentAttribute.IsValidIndex(Index))
+			{
+				ReturnData.ConditionUpAlignment = ConditionUpAlignmentAttribute.Get()[Index];
+			}
+
+			if (ConditionHealthAttribute.IsValidIndex(Index))
+			{
+				ReturnData.ConditionHealth = ConditionHealthAttribute.Get()[Index];
+			}
+			
 			if (ParentBoneIdsAttribute.IsValidIndex(Index))
 			{
 				ReturnData.ParentBoneID = ParentBoneIdsAttribute.Get()[Index];
@@ -191,6 +250,12 @@ namespace PV::Facades
 			NormalVectorsAttribute.Modify()[NewIndex] = InputData.NormalVector;
 			ScalesAttribute.Modify()[NewIndex] = InputData.Scale;
 			LengthFromRootAttribute.Modify()[NewIndex] = InputData.LengthFromRoot;
+			ConditionLightAttribute.Modify()[NewIndex] = InputData.ConditionLight;
+			ConditionScaleAttribute.Modify()[NewIndex] = InputData.ConditionScale;
+			ConditionUpAlignmentAttribute.Modify()[NewIndex] = InputData.ConditionUpAlignment;
+			ConditionHealthAttribute.Modify()[NewIndex] = InputData.ConditionHealth;
+			ConditionTipAttribute.Modify()[NewIndex] = InputData.ConditionTip;
+			
 			ParentBoneIdsAttribute.Modify()[NewIndex] = InputData.ParentBoneID;
 
 			return NewIndex;
@@ -210,8 +275,11 @@ namespace PV::Facades
 			NormalVectorsAttribute.ModifyAt(Index, InputData.NormalVector);
 			ScalesAttribute.ModifyAt(Index, InputData.Scale);
 			LengthFromRootAttribute.ModifyAt(Index, InputData.LengthFromRoot);
-			LengthFromRootAttribute.ModifyAt(Index, InputData.LengthFromRoot);
-			LengthFromRootAttribute.ModifyAt(Index, InputData.LengthFromRoot);
+			ConditionTipAttribute.ModifyAt(Index, InputData.ConditionTip);
+			ConditionScaleAttribute.ModifyAt(Index, InputData.ConditionScale);
+			ConditionLightAttribute.ModifyAt(Index, InputData.ConditionLight);
+			ConditionUpAlignmentAttribute.ModifyAt(Index, InputData.ConditionUpAlignment);
+			ConditionHealthAttribute.ModifyAt(Index, InputData.ConditionHealth);
 			ParentBoneIdsAttribute.ModifyAt(Index, InputData.ParentBoneID);
 		}
 	}
@@ -234,25 +302,57 @@ namespace PV::Facades
 		}
 	}
 
-	FString FFoliageFacade::GetFoliageName(const int32 Index) const
+	FPVFoliageInfo FFoliageFacade::GetFoliageInfo(const int32 Index) const
 	{
-		if (NamesAttribute.IsValid() && NamesAttribute.IsValidIndex(Index))
+		FPVFoliageInfo ReturnData = {};
+		
+		if (IsValid() && Index > -1)
 		{
-			return NamesAttribute[Index];
-		}
+			if (FoliageInfoNamesAttribute.IsValidIndex(Index))
+			{
+				ReturnData.Mesh = FoliageInfoNamesAttribute.Get()[Index];
+			}
 
-		return FString();
+			if (FoliageInfoLightAttribute.IsValidIndex(Index))
+			{
+				ReturnData.Attributes.Light = FoliageInfoLightAttribute.Get()[Index];
+			}
+
+			if (FoliageInfoScaleAttribute.IsValidIndex(Index))
+			{
+				ReturnData.Attributes.Scale = FoliageInfoScaleAttribute.Get()[Index];
+			}
+
+			if (FoliageInfoTipAttribute.IsValidIndex(Index))
+			{
+				ReturnData.Attributes.Tip = (bool)FoliageInfoTipAttribute.Get()[Index];
+			}
+
+			if (FoliageInfoHealthAttribute.IsValidIndex(Index))
+			{
+				ReturnData.Attributes.Health = FoliageInfoHealthAttribute.Get()[Index];
+			}
+
+			if (FoliageInfoUpAlignmentAttribute.IsValidIndex(Index))
+			{
+				ReturnData.Attributes.UpAlignment = FoliageInfoUpAlignmentAttribute.Get()[Index];
+			}
+		}
+		return ReturnData;
 	}
 
-	TArray<FString> FFoliageFacade::GetFoliageNames() const
+	TArray<FPVFoliageInfo> FFoliageFacade::GetFoliageInfos() const
 	{
-		if (NamesAttribute.IsValid())
+		TArray<FPVFoliageInfo> FoliageInfos;
+		
+		if (FoliageInfoNamesAttribute.IsValid())
 		{
-			return NamesAttribute.Get().GetConstArray();
+			for (int32 i = 0; i < FoliageInfoNamesAttribute.Num(); i++)
+			{
+				FoliageInfos.Add(GetFoliageInfo(i));
+			}
 		}
-
-		const static TArray<FString> DefaultFoliageNames;
-		return DefaultFoliageNames;
+		return FoliageInfos;
 	}
 
 	const FVector3f& FFoliageFacade::GetPivotPoint(const int32 Index) const
@@ -325,12 +425,34 @@ namespace PV::Facades
 		}
 	}
 
+	FString FFoliageFacade::GetFoliageName(const int32 Index) const
+	{
+		if (FoliageInfoNamesAttribute.IsValid() && FoliageInfoNamesAttribute.IsValidIndex(Index))
+		{
+			return FoliageInfoNamesAttribute[Index];
+		}
+
+		return FString();
+	}
+
+	TArray<FString> FFoliageFacade::GetFoliageNames() const
+	{
+		TArray<FString> FoliageNames;
+		
+		if (FoliageInfoNamesAttribute.IsValid())
+		{
+			return FoliageInfoNamesAttribute.Get().GetConstArray();
+		}
+
+		return FoliageNames;
+	}
+
 	void FFoliageFacade::SetFoliageName(const int32 Index, const FString& Input)
 	{
 		check(!IsConst());
-		if (IsValid() && NamesAttribute.IsValidIndex(Index))
+		if (IsValid() && FoliageInfoNamesAttribute.IsValidIndex(Index))
 		{
-			NamesAttribute.ModifyAt(Index, Input);
+			FoliageInfoNamesAttribute.ModifyAt(Index, Input);
 		}
 	}
 
@@ -341,7 +463,35 @@ namespace PV::Facades
 		Collection->AddElements(InputNames.Num(), GroupNames::FoliageNamesGroup);
 		for (int32 Index = 0; Index < InputNames.Num(); ++Index)
 		{
-			NamesAttribute.Modify()[Index] = InputNames[Index];
+			FoliageInfoNamesAttribute.Modify()[Index] = InputNames[Index];
+		}
+	}
+
+	void FFoliageFacade::SetFoliageInfo(const int32 Index, const FString& Input)
+	{
+		check(!IsConst());
+		if (IsValid() && FoliageInfoNamesAttribute.IsValidIndex(Index))
+		{
+			FoliageInfoNamesAttribute.ModifyAt(Index, Input);
+		}
+	}
+
+	void FFoliageFacade::SetFoliageInfos(const TArray<FPVFoliageInfo>& Infos)
+	{
+		check(!IsConst());
+		Collection->EmptyGroup(GroupNames::FoliageNamesGroup);
+		Collection->AddElements(Infos.Num(), GroupNames::FoliageNamesGroup);
+
+		for (int32 Index = 0; Index < Infos.Num(); ++Index)
+		{
+			FPVFoliageInfo Info = Infos[Index];
+			
+			FoliageInfoNamesAttribute.Modify()[Index] = Info.Mesh.ToString();
+			FoliageInfoScaleAttribute.Modify()[Index] = Info.Attributes.Scale;
+			FoliageInfoLightAttribute.Modify()[Index] = Info.Attributes.Light;
+			FoliageInfoTipAttribute.Modify()[Index] = Info.Attributes.Tip;
+			FoliageInfoUpAlignmentAttribute.Modify()[Index] = Info.Attributes.UpAlignment;
+			FoliageInfoHealthAttribute.Modify()[Index] = Info.Attributes.Health;
 		}
 	}
 
@@ -385,6 +535,11 @@ namespace PV::Facades
 			NormalVectorsAttribute.ModifyAt(ToIndex, NormalVectorsAttribute[FromIndex]);
 			ScalesAttribute.ModifyAt(ToIndex, ScalesAttribute[FromIndex]);
 			LengthFromRootAttribute.ModifyAt(ToIndex, LengthFromRootAttribute[FromIndex]);
+			ConditionTipAttribute.ModifyAt(ToIndex, ConditionTipAttribute[FromIndex]);
+			ConditionLightAttribute.ModifyAt(ToIndex, ConditionLightAttribute[FromIndex]);
+			ConditionScaleAttribute.ModifyAt(ToIndex, ConditionScaleAttribute[FromIndex]);
+			ConditionUpAlignmentAttribute.ModifyAt(ToIndex, ConditionUpAlignmentAttribute[FromIndex]);
+			ConditionHealthAttribute.ModifyAt(ToIndex, ConditionHealthAttribute[FromIndex]);
 			ParentBoneIdsAttribute.ModifyAt(ToIndex, ParentBoneIdsAttribute[FromIndex]);
 		}
 	}
diff --git a/Engine/Plugins/Experimental/ProceduralVegetationEditor/Source/ProceduralVegetation/Private/Helpers/PVAnalyticsHelper.cpp b/Engine/Plugins/Experimental/ProceduralVegetationEditor/Source/ProceduralVegetation/Private/Helpers/PVAnalyticsHelper.cpp
--- a/Engine/Plugins/Experimental/ProceduralVegetationEditor/Source/ProceduralVegetation/Private/Helpers/PVAnalyticsHelper.cpp
+++ b/Engine/Plugins/Experimental/ProceduralVegetationEditor/Source/ProceduralVegetation/Private/Helpers/PVAnalyticsHelper.cpp
@@ -10,6 +10,7 @@
 #include "GeometryCollection/GeometryCollection.h"
 #include "GeometryCollection/ManagedArrayCollection.h"
 #include "Helpers/PVUtilities.h"
+#include "Implementations/PVFoliage.h"
 #include "UObject/UnrealTypePrivate.h"
 
 namespace PV::Analytics
@@ -127,11 +128,11 @@ namespace PV::Analytics
 		OutUniqueFoliageTriangles = 0;
 		OutTotalFoliageTriangles = 0;
 		
-		for (int32 FoliageNameIndex = 0; FoliageNameIndex < InFoliageFacade.NumFoliageNames(); ++FoliageNameIndex)
+		for (int32 FoliageIndex = 0; FoliageIndex < InFoliageFacade.NumFoliageInfo(); ++FoliageIndex)
 		{
-			FString FoliagePath = InFoliageFacade.GetFoliageName(FoliageNameIndex);
+			FPVFoliageInfo Foliage = InFoliageFacade.GetFoliageInfo(FoliageIndex);
 			
-			OutUniqueFoliageTriangles += PV::Utilities::GetMeshTriangles(FoliagePath);
+			OutUniqueFoliageTriangles += PV::Utilities::GetMeshTriangles(Foliage.Mesh.ToString());
 		}
 
 		int32 NumInstances = InFoliageFacade.NumFoliageEntries();
@@ -139,9 +140,9 @@ namespace PV::Analytics
 		for (int32 Id = 0; Id < NumInstances; Id++)
 		{
 			Facades::FFoliageEntryData Data = InFoliageFacade.GetFoliageEntry(Id);
-			FString FoliagePath = InFoliageFacade.GetFoliageName(Data.NameId);
+			FPVFoliageInfo FoliageInfo = InFoliageFacade.GetFoliageInfo(Data.NameId);
 
-			OutTotalFoliageTriangles += PV::Utilities::GetMeshTriangles(FoliagePath);
+			OutTotalFoliageTriangles += PV::Utilities::GetMeshTriangles(FoliageInfo.Mesh.ToString());
 		}
 	}
 
diff --git a/Engine/Plugins/Experimental/ProceduralVegetationEditor/Source/ProceduralVegetation/Private/Helpers/PVFoliageJSONHelper.cpp b/Engine/Plugins/Experimental/ProceduralVegetationEditor/Source/ProceduralVegetation/Private/Helpers/PVFoliageJSONHelper.cpp
new file mode 100644
--- /dev/null
+++ b/Engine/Plugins/Experimental/ProceduralVegetationEditor/Source/ProceduralVegetation/Private/Helpers/PVFoliageJSONHelper.cpp
@@ -0,0 +1,407 @@
+﻿// Copyright Epic Games, Inc. All Rights Reserved.
+
+#include "Helpers/PVFoliageJSONHelper.h"
+
+#include "ProceduralVegetationModule.h"
+#include "Implementations/PVFoliage.h"
+#include "Facades/PVAttributesNames.h"
+#include "Facades/PVFoliageFacade.h"
+#include "Dom/JsonObject.h"
+#include "Dom/JsonValue.h"
+#include "Facades/PVFoliageConditionFacade.h"
+#include "Misc/FileHelper.h"
+#include "Serialization/JsonSerializer.h"
+#include "Misc/Paths.h"
+#include "Misc/PackageName.h"
+
+namespace PV
+{
+	bool PVFoliageJSONHelper::LoadJSON(const FString& InPath, TSharedPtr<FJsonObject>& OutData, FString& OutErrorMessage)
+	{
+		FString JsonString;
+		if (!FFileHelper::LoadFileToString(JsonString, *InPath))
+		{
+			OutErrorMessage = TEXT("Failed to open file");
+			return false;
+		}
+
+		OutData = MakeShareable(new FJsonObject);
+		if (!FJsonSerializer::Deserialize(TJsonReaderFactory<>::Create(JsonString), OutData))
+		{
+			OutErrorMessage = TEXT("Failed to deserialize JSON file");
+			return false;
+		}
+
+		return true;
+	}
+
+	bool GetConditionFromString(const FString& InString, EFoliageDistributionCondition& OutCondition)
+	{
+		if (InString.Equals("light", ESearchCase::IgnoreCase))
+		{
+			OutCondition = EFoliageDistributionCondition::Light;
+			return true;
+		}
+		else if (InString.Equals("scale", ESearchCase::IgnoreCase))
+		{
+			OutCondition = EFoliageDistributionCondition::Scale;
+			return true;
+		}
+		else if (InString.Equals("upAlignment", ESearchCase::IgnoreCase))
+		{
+			OutCondition = EFoliageDistributionCondition::UpAlignment;
+			return true;
+		}
+		else if (InString.Equals("health", ESearchCase::IgnoreCase))
+		{
+			OutCondition = EFoliageDistributionCondition::Health;
+			return true;
+		}
+		else if (InString.Equals("tip", ESearchCase::IgnoreCase))
+		{
+			OutCondition = EFoliageDistributionCondition::Tip;
+			return true;
+		}
+			
+		return false;
+	}
+
+	bool PVFoliageJSONHelper::LoadFoliageData(const FString& InPath, const FString& InPackagePath, FoliageVariationsMap& OutFoliageVariationsMap, FString& OutErrorMessage)
+	{
+		TSharedPtr<FJsonObject> OutData = MakeShareable(new FJsonObject);
+		
+		OutFoliageVariationsMap.Empty();
+
+		if (LoadJSON(InPath, OutData, OutErrorMessage))
+		{
+			if (OutData->HasTypedField(TEXT("Variations"), EJson::Array))
+			{
+				const TArray<TSharedPtr<FJsonValue>>& Variations = OutData->GetArrayField(TEXT("Variations"));
+
+				for (TSharedPtr<FJsonValue> FoliageInfoVariation : Variations)
+				{
+					if (FoliageInfoVariation->Type == EJson::Object)
+					{
+						const TSharedPtr<FJsonObject> VariationData = FoliageInfoVariation->AsObject();
+
+						FString VariationName = VariationData->GetStringField(TEXT("name"));
+
+						const TArray<TSharedPtr<FJsonValue>>* OutVariationData;
+
+						if (!VariationData->TryGetArrayField(TEXT("Data"), OutVariationData))
+						{
+							OutErrorMessage = FString::Printf(TEXT("Unable to load foliage data : Error: Missing data for variation : %s"), *VariationName);
+							return false;
+						}
+
+						TArray<FPVFoliageInfo> FoliageInfos;
+						
+						for (TSharedPtr<FJsonValue> FoliageInfoData : *OutVariationData)
+						{
+							if (const TSharedPtr<FJsonObject> FoliageInfoObject = FoliageInfoData->AsObject())
+							{
+								FPVFoliageInfo FoliageInfo;
+								FString BranchName;
+
+								if (!FoliageInfoObject->TryGetStringField(TEXT("Name"), BranchName))
+								{
+									OutErrorMessage = FString::Printf(TEXT("FoliageInfo Doesnt contain valid branch name for variation %s"), *VariationName);
+									return false;
+								}
+								
+								FoliageInfo.Mesh = InPackagePath / BranchName + "." + BranchName;
+								
+								if (!FoliageInfoObject->TryGetNumberField(TEXT("scale"), FoliageInfo.Attributes.Scale))
+								{
+									FoliageInfo.Attributes.Scale = 0;
+								}
+								if (!FoliageInfoObject->TryGetNumberField(TEXT("upAlignment"), FoliageInfo.Attributes.UpAlignment))
+								{
+									FoliageInfo.Attributes.UpAlignment = 0;
+								}
+
+								if (!FoliageInfoObject->TryGetNumberField(TEXT("light"), FoliageInfo.Attributes.Light))
+								{
+									FoliageInfo.Attributes.Light = 0;
+								}
+
+								if (!FoliageInfoObject->TryGetNumberField(TEXT("health"), FoliageInfo.Attributes.Health))
+								{
+									FoliageInfo.Attributes.Health = 0;
+								}
+
+								float OutTip; 
+								if (FoliageInfoObject->TryGetNumberField(TEXT("tip"), OutTip))
+								{
+									FoliageInfo.Attributes.Tip = static_cast<bool>(OutTip);	
+								}
+								else
+								{
+									FoliageInfo.Attributes.Tip = false;
+								}
+								
+
+								FoliageInfos.Add(FoliageInfo);
+							}
+						}
+
+						const TArray<TSharedPtr<FJsonValue>>* OutVariationRules;
+
+						if (!VariationData->TryGetArrayField(TEXT("Rules"), OutVariationRules))
+						{
+							OutErrorMessage = FString::Printf(TEXT("Unable to load foliage rules : Error: Missing data for variation : %s"), *VariationName);
+							return false;
+						}
+
+						TMap<EFoliageDistributionCondition, Facades::FFoliageConditonInfo> FoliageDistributionConditions;
+						
+						for (TSharedPtr<FJsonValue> FoliageRule : *OutVariationRules)
+						{
+							if (const TSharedPtr<FJsonObject> Rule = FoliageRule->AsObject())
+							{
+								FString RuleName;
+								
+								if (!Rule->TryGetStringField(TEXT("name"), RuleName))
+								{
+									OutErrorMessage = FString::Printf(TEXT("FoliageData: Rule doesnt contain valid name for variation %s"), *VariationName);
+									continue;
+								}
+
+								EFoliageDistributionCondition OutCondition;
+								
+								if (GetConditionFromString(RuleName, OutCondition))
+								{
+									Facades::FFoliageConditonInfo ConditionSettings;
+
+									const UEnum* EnumPtr = StaticEnum<EFoliageDistributionCondition>();
+									ConditionSettings.Name = EnumPtr->GetNameStringByValue(static_cast<int64>(OutCondition));
+									
+									if (!Rule->TryGetNumberField(TEXT("weight"), ConditionSettings.Weight))
+									{
+										ConditionSettings.Weight = 0;
+									}
+
+									if (!Rule->TryGetNumberField(TEXT("offset"), ConditionSettings.Offset))
+									{
+										ConditionSettings.Offset = 0;
+									}
+
+									// float OutNormalize; 
+									// if (Rule->TryGetNumberField(TEXT("normalize"), OutNormalize))
+									// {
+									// 	ConditionSettings.Normalize = static_cast<bool>(OutNormalize);	
+									// }
+									// else
+									// {
+									// 	ConditionSettings.Normalize = false;
+									// }
+
+									if (!FoliageDistributionConditions.Contains(OutCondition))
+									{
+										FoliageDistributionConditions.Add(OutCondition);
+									}
+									
+									FoliageDistributionConditions[OutCondition] = ConditionSettings;
+								}
+							}
+						}
+
+						if (!OutFoliageVariationsMap.Contains(VariationName))
+						{
+							OutFoliageVariationsMap.Add(VariationName);
+						}
+
+						OutFoliageVariationsMap[VariationName].VariationName = VariationName;
+						OutFoliageVariationsMap[VariationName].FoliageInfos = FoliageInfos;
+						OutFoliageVariationsMap[VariationName].Conditions = FoliageDistributionConditions;
+					}
+				}
+			}
+			else
+			{
+				OutErrorMessage = FString::Printf(TEXT("Unable to load foliage data : Error: Missing variation data"));
+				
+				return false;	
+			}
+		}
+
+		for (auto Mapping : OutFoliageVariationsMap)
+		{
+			UE_LOG(LogProceduralVegetation, Log, TEXT("Variation name : %s : Start"), *Mapping.Key);
+
+			for (const FPVFoliageInfo& FoliageInfo : Mapping.Value.FoliageInfos)
+			{
+				UE_LOG(LogProceduralVegetation, Log, TEXT("FoliageInfo : Mesh : %s : Scale : %f : Light : %f : Tip : %i : UpAlignment :  %f : Health : %f"),
+					*FoliageInfo.Mesh.ToString(),
+					FoliageInfo.Attributes.Scale,
+					FoliageInfo.Attributes.Light,
+					FoliageInfo.Attributes.Tip,
+					FoliageInfo.Attributes.UpAlignment,
+					FoliageInfo.Attributes.Health);	
+			}
+
+			for (const auto& ConditionEntry : Mapping.Value.Conditions)
+			{
+				UE_LOG(LogProceduralVegetation, Log, TEXT("FoliageCondition : %s"), *ConditionEntry.Value.ToString());	
+			}
+
+			UE_LOG(LogProceduralVegetation, Log, TEXT("Variation name : %s : End"), *Mapping.Key);
+		}
+
+		return true;
+	}
+
+	bool PVFoliageJSONHelper::LoadFoliageDataInCollection(FManagedArrayCollection& Collection, const FString& InPresetPath, const FPVFoliageVariationData& VariationData, FString& OutErrorMessage)
+	{
+		TSharedPtr<FJsonObject> OutData = MakeShareable(new FJsonObject);
+
+		FString MegaPlantsData;
+		if (!FFileHelper::LoadFileToString(MegaPlantsData, *InPresetPath))
+		{
+			OutErrorMessage = TEXT("Failed to open skeleton file");
+			return false;
+		}
+
+		TSharedPtr<FJsonObject> LoadedData = MakeShareable(new FJsonObject);
+		if (!FJsonSerializer::Deserialize(TJsonReaderFactory<>::Create(MegaPlantsData), LoadedData))
+		{
+			OutErrorMessage = TEXT("Failed to deserialize JSON file for skeleton");
+			return false;
+		}
+
+		Collection.RemoveGroup(GroupNames::FoliageGroup);
+		Collection.RemoveGroup(GroupNames::FoliageNamesGroup);
+		
+		const TSharedPtr<FJsonObject>& Primitives = LoadedData->GetObjectField(TEXT("primitives"));
+		const TSharedPtr<FJsonObject>& PrimitiveAttributes = Primitives->GetObjectField(TEXT("attributes"));
+
+		const TSharedPtr<FJsonObject>& InstancerNameObject = PrimitiveAttributes->GetObjectField(TEXT("instancer_name"));
+		const TArray<TSharedPtr<FJsonValue>>& InstancerNameValues = InstancerNameObject->GetArrayField(TEXT("values"));
+
+		const TSharedPtr<FJsonObject>& InstancerPivotObject = PrimitiveAttributes->GetObjectField(TEXT("instancer_pivot"));
+		const TArray<TSharedPtr<FJsonValue>>& InstancerPivotValues = InstancerPivotObject->GetArrayField(TEXT("values"));
+
+		const TSharedPtr<FJsonObject>& InstancerUpVectorObject = PrimitiveAttributes->GetObjectField(TEXT("instancer_UP"));
+		const TArray<TSharedPtr<FJsonValue>>& InstancerUpVectorValues = InstancerUpVectorObject->GetArrayField(TEXT("values"));
+
+		const TSharedPtr<FJsonObject>& InstancerNormalVectorObject = PrimitiveAttributes->GetObjectField(TEXT("instancer_N"));
+		const TArray<TSharedPtr<FJsonValue>>& InstancerNormalVectorValues = InstancerNormalVectorObject->GetArrayField(TEXT("values"));
+
+		const TSharedPtr<FJsonObject>& InstancerScaleObject = PrimitiveAttributes->GetObjectField(TEXT("instancer_scale"));
+		const TArray<TSharedPtr<FJsonValue>>& InstancerScaleValues = InstancerScaleObject->GetArrayField(TEXT("values"));
+
+		const TSharedPtr<FJsonObject>& InstancerLFRObject = PrimitiveAttributes->GetObjectField(TEXT("instancer_LFR"));
+		const TArray<TSharedPtr<FJsonValue>>& InstancerLFRValues = InstancerLFRObject->GetArrayField(TEXT("values"));
+
+		
+		TMap<FString, int32> FoliageNamesToIndex;
+		for (int FoliageIndex = 0; FoliageIndex < VariationData.FoliageInfos.Num(); ++FoliageIndex)
+		{
+			auto FoliageInfo = VariationData.FoliageInfos[FoliageIndex];
+			FoliageNamesToIndex.Emplace(FoliageInfo.Mesh.GetAssetName(), FoliageIndex);
+		}
+		
+		// Compute unique names and number of elements for group
+		int32 TotalValues = 0;
+		TArray<FString> FoliageNames;
+		for (int32 i = 0; i < InstancerNameValues.Num(); ++i)
+		{
+			const TArray<TSharedPtr<FJsonValue>>& InstancerNames = InstancerNameValues[i]->AsArray();
+			for (int32 Index = 0; Index < InstancerNames.Num(); ++Index)
+			{
+				const FString& Name = InstancerNames[Index]->AsString();
+
+				if (!FoliageNamesToIndex.Contains(Name))
+				{
+					OutErrorMessage = FString::Printf(TEXT("FoliageData doesnt contain branch info : %s"), *Name);
+					return false;
+				}
+			}
+			TotalValues += InstancerNames.Num();
+		}
+		
+		Facades::FFoliageFacade Facade(Collection, TotalValues);
+		Facade.SetFoliageInfos(VariationData.FoliageInfos);
+
+		TArray<Facades::FFoliageConditonInfo> Conditions;
+
+		for (auto ConditionEntry : VariationData.Conditions)
+		{
+			Conditions.Add(ConditionEntry.Value);
+		}
+
+		Facades::FFoliageConditionFacade ConditionFacade(Collection);
+		ConditionFacade.SetData(Conditions);
+		
+		int32 FoliageEntryIndex = 0;
+		for (int32 Id = 0; Id < InstancerNameValues.Num(); ++Id)
+		{
+			const TArray<TSharedPtr<FJsonValue>>& InstancerNames = InstancerNameValues[Id]->AsArray();
+			const TArray<TSharedPtr<FJsonValue>>& PivotPoints = InstancerPivotValues[Id]->AsArray();
+			const TArray<TSharedPtr<FJsonValue>>& UpVectors = InstancerUpVectorValues[Id]->AsArray();
+			const TArray<TSharedPtr<FJsonValue>>& NormalVectors = InstancerNormalVectorValues[Id]->AsArray();
+			const TArray<TSharedPtr<FJsonValue>>& Scales = InstancerScaleValues[Id]->AsArray();
+			const TArray<TSharedPtr<FJsonValue>>& LFRs = InstancerLFRValues[Id]->AsArray();
+		
+			TArray<int32> FoliageEntryIds;
+			for (int32 j = 0; j < InstancerNames.Num(); ++j)
+			{
+				const FString Name = InstancerNames[j]->AsString();
+				const FVector3f PivotPoint = FVector3f(
+					PivotPoints[j * 3]->AsNumber(),
+					PivotPoints[(j * 3) + 2]->AsNumber(),
+					PivotPoints[(j * 3) + 1]->AsNumber()) * 100.0f;
+				const FVector3f UpVector = FVector3f(
+					UpVectors[j * 3]->AsNumber(),
+					UpVectors[(j * 3) + 2]->AsNumber(),
+					UpVectors[(j * 3) + 1]->AsNumber());
+				const FVector3f NormalVector = FVector3f(
+					NormalVectors[j * 3]->AsNumber(),
+					NormalVectors[(j * 3) + 2]->AsNumber(),
+					NormalVectors[(j * 3) + 1]->AsNumber());
+				const float Scale = Scales[j]->AsNumber();
+				const float LengthFromRoot = LFRs[j]->AsNumber();
+		
+				Facade.SetFoliageEntry(FoliageEntryIndex, {
+					.NameId = FoliageNamesToIndex[Name],
+					.BranchId = Id,
+					.PivotPoint = PivotPoint,
+					.UpVector = UpVector,
+					.NormalVector = NormalVector,
+					.Scale = Scale,
+					.LengthFromRoot = LengthFromRoot,
+				});
+		
+				FoliageEntryIds.Add(FoliageEntryIndex);
+				FoliageEntryIndex++;
+			}
+		
+			Facade.SetFoliageIdsArray(Id, FoliageEntryIds);
+		}
+
+		return true;
+	}
+
+
+	void PVFoliageJSONHelper::SetFoliagePaths(FManagedArrayCollection& Collection, const FString& FilePath)
+	{
+		PV::Facades::FFoliageFacade Facade(Collection);
+
+		FString AbsolutePath = FPaths::ConvertRelativePathToFull(FilePath);
+		FString PackagePath;
+
+		if (FPackageName::TryConvertFilenameToLongPackageName(AbsolutePath, PackagePath))
+		{
+			TArray<FPVFoliageInfo> Infos = Facade.GetFoliageInfos();
+			FString PackageFolder = FPackageName::GetLongPackagePath(PackagePath);
+			Facade.SetFoliagePath(PackageFolder);
+
+			for (FPVFoliageInfo& Info : Infos)
+			{
+				Info.Mesh = PackageFolder / Info.Mesh.ToString();
+			}
+			
+			Facade.SetFoliageInfos(Infos);
+		}
+	}
+}
diff --git a/Engine/Plugins/Experimental/ProceduralVegetationEditor/Source/ProceduralVegetation/Private/Helpers/PVFoliageJSONHelper.h b/Engine/Plugins/Experimental/ProceduralVegetationEditor/Source/ProceduralVegetation/Private/Helpers/PVFoliageJSONHelper.h
new file mode 100644
--- /dev/null
+++ b/Engine/Plugins/Experimental/ProceduralVegetationEditor/Source/ProceduralVegetation/Private/Helpers/PVFoliageJSONHelper.h
@@ -0,0 +1,34 @@
+// Copyright Epic Games, Inc. All Rights Reserved.
+#pragma once
+
+#include "CoreMinimal.h"
+#include "Facades/PVFoliageConditionFacade.h"
+#include "Templates/SharedPointer.h"
+
+struct FPVFoliageDistributionConditionSettings;
+enum class EFoliageDistributionCondition : uint8;
+class FJsonObject;
+struct FPVFoliageInfo;
+struct FManagedArrayCollection;
+
+namespace PV
+{
+	typedef TMap<EFoliageDistributionCondition, Facades::FFoliageConditonInfo> FoliageConditionMap;
+	
+	struct FPVFoliageVariationData
+	{
+		FString VariationName;
+		TArray<FPVFoliageInfo> FoliageInfos;
+		FoliageConditionMap Conditions;
+	};
+
+	typedef TMap<FString, FPVFoliageVariationData> FoliageVariationsMap;
+	
+	struct PVFoliageJSONHelper
+	{
+		static bool LoadJSON(const FString& InPath, TSharedPtr<FJsonObject>& OutData, FString& OutErrorMessage);
+		static bool LoadFoliageData(const FString& InPath, const FString& InPackagePath, FoliageVariationsMap& OutFoliageVariationsMap, FString& OutErrorMessage);
+		static bool LoadFoliageDataInCollection(FManagedArrayCollection& Collection, const FString& InPresetPath, const FPVFoliageVariationData& VariationData, FString& OutErrorMessage);
+		static void SetFoliagePaths(FManagedArrayCollection& Collection, const FString& FilePath);
+	};	
+}
\ No newline at end of file
diff --git a/Engine/Plugins/Experimental/ProceduralVegetationEditor/Source/ProceduralVegetation/Private/Helpers/PVUtilities.cpp b/Engine/Plugins/Experimental/ProceduralVegetationEditor/Source/ProceduralVegetation/Private/Helpers/PVUtilities.cpp
--- a/Engine/Plugins/Experimental/ProceduralVegetationEditor/Source/ProceduralVegetation/Private/Helpers/PVUtilities.cpp
+++ b/Engine/Plugins/Experimental/ProceduralVegetationEditor/Source/ProceduralVegetation/Private/Helpers/PVUtilities.cpp
@@ -15,6 +15,7 @@
 #include "AssetRegistry/AssetData.h"
 #include "Facades/PVBranchFacade.h"
 #include "Facades/PVPointFacade.h"
+#include "Implementations/PVFoliage.h"
 #include "Misc/FileHelper.h"
 #include "Rendering/SkeletalMeshRenderData.h"
 
@@ -23,7 +24,7 @@ namespace PV::Utilities
 #if WITH_EDITOR
 	static TAutoConsoleVariable<bool> CVarEnablePVDebugMode(
 		TEXT("PV.DebugMode.Enabled"),
-		false,
+		true,
 		TEXT("Enables debug mode for the Procedural Vegetation Editor"), 
 		FConsoleVariableDelegate::CreateLambda([](IConsoleVariable* ConsoleVariable)
 			{
@@ -67,9 +68,10 @@ namespace PV::Utilities
 
 		if (NumInstances > 0)
 		{
-			for (int32 FoliageNameIndex = 0; FoliageNameIndex < InFoliageFacade.NumFoliageNames(); ++FoliageNameIndex)
+			for (int32 FoliageNameIndex = 0; FoliageNameIndex < InFoliageFacade.NumFoliageInfo(); ++FoliageNameIndex)
 			{
-				FString FoliageName = InFoliageFacade.GetFoliageName(FoliageNameIndex);
+				FPVFoliageInfo FoliageInfo = InFoliageFacade.GetFoliageInfo(FoliageNameIndex);
+				FString FoliageName = FoliageInfo.Mesh.ToString();
 				FSoftObjectPath MeshPath(FoliageName);
 
 				UStaticMesh* StaticFoliageMesh = PackageExists(MeshPath.GetLongPackageName(), UStaticMesh::StaticClass()) ? LoadObject<UStaticMesh>(nullptr, *FoliageName) : nullptr;
@@ -108,8 +110,9 @@ namespace PV::Utilities
 			for (int32 Id = 0; Id < NumInstances; Id++)
 			{
 				Facades::FFoliageEntryData Data = InFoliageFacade.GetFoliageEntry(Id);
-				FString FoliageName = InFoliageFacade.GetFoliageName(Data.NameId);
-
+				FPVFoliageInfo FoliageInfo = InFoliageFacade.GetFoliageInfo(Data.NameId);
+				FString FoliageName = FoliageInfo.Mesh.ToString();
+				
 				const FTransform Transform = InFoliageFacade.GetFoliageTransform(Id);
 
 				InstanceTransformsMap.FindOrAdd(FoliageName).Add(Transform);
diff --git a/Engine/Plugins/Experimental/ProceduralVegetationEditor/Source/ProceduralVegetation/Private/Implementations/PVFoliage.cpp b/Engine/Plugins/Experimental/ProceduralVegetationEditor/Source/ProceduralVegetation/Private/Implementations/PVFoliage.cpp
--- a/Engine/Plugins/Experimental/ProceduralVegetationEditor/Source/ProceduralVegetation/Private/Implementations/PVFoliage.cpp
+++ b/Engine/Plugins/Experimental/ProceduralVegetationEditor/Source/ProceduralVegetation/Private/Implementations/PVFoliage.cpp
@@ -4,20 +4,79 @@
 
 #include "ProceduralVegetationModule.h"
 #include "PVFloatRamp.h"
+#include "Algo/RandomShuffle.h"
 
 #include "Facades/PVAttributesNames.h"
 #include "Facades/PVBranchFacade.h"
+#include "Facades/PVBudVectorsFacade.h"
 #include "Facades/PVFoliageFacade.h"
 #include "Facades/PVPointFacade.h"
 
+struct ConditionAttributes
+{
+	TArray<float> UpAlignment;
+	TArray<float> Light;
+	TArray<float> Scale;
+	TArray<float> Tip;
+	TArray<float> Health;
+	
+
+	float GetConditionValue(EFoliageDistributionCondition InCondition, const int Index) const
+	{
+		TArray<float> Result;
+		switch (InCondition)
+		{
+		case EFoliageDistributionCondition::Light:
+			{
+				Result = Light;
+				break;
+			}
+		case EFoliageDistributionCondition::Scale:
+			{
+				Result = Scale;
+				break;;
+			}
+		case EFoliageDistributionCondition::Tip:
+			{
+				Result = Tip;
+				break;
+			}
+		case EFoliageDistributionCondition::UpAlignment:
+			{
+				Result = UpAlignment;
+				break;
+			}
+		case EFoliageDistributionCondition::Health:
+			{
+				Result = Health;
+				break;
+			}
+		}
+
+		if (Result.IsValidIndex(Index))
+		{
+			return Result[Index];
+		}
+
+		return 0;
+	}
+};
+
+// Condition Weight
+struct FFoliageWeightInfo
+{
+	int FoliageIndex = 0;
+	float Weight = 0;
+};
+
 void FPVFoliage::DistributeFoliage(FManagedArrayCollection& OutCollection, const FManagedArrayCollection& SourceCollection,
                                    const FDistributionSettings& DistributionSettings, const FScaleSettings& ScaleSettings,
                                    const FVectorSettings& VectorSettings, const FPhyllotaxySettings& PhyllotaxySettings,
-                                   const FMiscSettings& MiscSettings)
+                                   const FMiscSettings& MiscSettings, const FPVFoliageDistributionConditions& DistributionConditions)
 {
 	PV::Facades::FFoliageFacade FoliageFacadeOutput(OutCollection);
 
-	if (const int32 NumFoliageMeshes = FoliageFacadeOutput.NumFoliageNames(); NumFoliageMeshes == 0)
+	if (const int32 NumFoliageMeshes = FoliageFacadeOutput.NumFoliageInfo(); NumFoliageMeshes == 0)
 	{
 		UE_LOG(LogProceduralVegetation, Warning, TEXT("There are no foliage meshes available in input for distribution"));
 		return;
@@ -31,13 +90,19 @@ void FPVFoliage::DistributeFoliage(FManagedArrayCollection& OutCollection, const
 
 	TManagedArrayAccessor<TArray<FVector3f>> BudDirectionAttribute(OutCollection, PV::AttributeNames::BudDirection,
 		PV::GroupNames::PointGroup);
+	
 	TManagedArrayAccessor<float> BranchGradientsAttribute(OutCollection, PV::AttributeNames::BranchGradient,
 		PV::GroupNames::PointGroup);
+	
 	TManagedArrayAccessor<TArray<float>> LeafPhyllotaxyAttribute(OutCollection, PV::AttributeNames::LeafPhyllotaxy,
 		PV::GroupNames::DetailsGroup);
+	
 	TManagedArrayAccessor<TArray<float>> BudHormoneLevels(OutCollection, PV::AttributeNames::BudHormoneLevels,
 		PV::GroupNames::PointGroup);
 
+	TManagedArrayAccessor<TArray<int>> BudDevelopment(OutCollection, PV::AttributeNames::BudDevelopment,
+		PV::GroupNames::PointGroup);
+
 	float PhyllotaxyFormationLeaf = LeafPhyllotaxyAttribute[0][1];
 	float MinBudsLeaf = LeafPhyllotaxyAttribute[0][3];
 	float MaxBudsLeaf = LeafPhyllotaxyAttribute[0][4];
@@ -90,6 +155,74 @@ void FPVFoliage::DistributeFoliage(FManagedArrayCollection& OutCollection, const
 	{
 		TRACE_CPUPROFILER_EVENT_SCOPE(PVFoliageDistributorNode::OverrideDistribution);
 
+		TArray<float> NormalisedPScales;
+		NormalisedPScales.Init(0, PointFacadeSource.GetElementCount());
+
+		TArray<float> NormalisedBudLightDetected;
+		NormalisedBudLightDetected.Init(0, PointFacadeSource.GetElementCount());
+
+		TArray<float> NormalisedHealth;
+		NormalisedHealth.Init(0, PointFacadeSource.GetElementCount());
+
+		TArray<PV::Facades::FFoliageEntryData> FoliageEntries;
+
+		// TODO : Optimize
+		// if (DistributionConditions.IsActiveCondition(EFoliageDistributionCondition::Scale) || DistributionConditions.IsActiveCondition(EFoliageDistributionCondition::Light))
+		// {
+		float MinPScale = std::numeric_limits<float>::max();
+		float MaxPScale = std::numeric_limits<float>::min();
+			
+		float MinBudLightDetected = std::numeric_limits<float>::max();
+		float MaxBudLightDetected = std::numeric_limits<float>::min();
+
+		float MinHealth = std::numeric_limits<float>::max();
+		float MaxHealth = std::numeric_limits<float>::min();
+
+		for (int PointIndex = 0; PointIndex < PointFacadeSource.GetElementCount(); PointIndex++)
+		{
+			// Scale
+			float PScale = PointFacadeSource.GetPointScale(PointIndex);
+			MinPScale = FMath::Min(MinPScale, PScale);	
+			MaxPScale = FMath::Max(MaxPScale, PScale);
+
+			// Light
+			const TArray<float>& CurrentBudLightDetected = PointFacadeSource.GetBudLightDetected(PointIndex);
+			const float CurrentLightAvailable = CurrentBudLightDetected.IsValidIndex(PV::Facades::BudLightDetectedAvailableIndex) ? CurrentBudLightDetected[PV::Facades::BudLightDetectedAvailableIndex] : 0.0f;
+				
+			MinBudLightDetected = FMath::Min(MinBudLightDetected, CurrentLightAvailable);	
+			MaxBudLightDetected = FMath::Max(MaxBudLightDetected, CurrentLightAvailable);
+
+			// Health
+			float Ethylene = BudHormoneLevels[PointIndex][4];
+			int BudAgeSenescence = BudDevelopment[PointIndex][3];
+			int BudLightSenescence = BudDevelopment[PointIndex][4];
+			
+			int Senescence = FMath::Max(BudAgeSenescence,BudLightSenescence) + 1;
+			NormalisedHealth[PointIndex] = 1 - (Ethylene * (1/Senescence));
+			
+			MinHealth = FMath::Min(MinHealth, NormalisedHealth[PointIndex]);
+			MaxHealth = FMath::Max(MaxHealth, NormalisedHealth[PointIndex]);
+		}
+
+		for (int PointIndex = 0; PointIndex < PointFacadeSource.GetElementCount(); PointIndex++)
+		{
+			// Scale
+			float PScale = PointFacadeSource.GetPointScale(PointIndex);
+			NormalisedPScales[PointIndex] = FMath::GetMappedRangeValueClamped(FVector2f(MinPScale, MaxPScale),FVector2f(0, 1),PScale);
+
+			// UE_LOG(LogTemp, Log, TEXT("Normalized PScale : Index : %i : Value : %f"),
+			// 					PointIndex, NormalisedPScales[PointIndex]);
+
+			// Light
+			const TArray<float>& CurrentBudLightDetected = PointFacadeSource.GetBudLightDetected(PointIndex);
+			const float CurrentLightAvailable = CurrentBudLightDetected.IsValidIndex(PV::Facades::BudLightDetectedAvailableIndex) ? CurrentBudLightDetected[PV::Facades::BudLightDetectedAvailableIndex] : 0.0f;
+			NormalisedBudLightDetected[PointIndex] = FMath::GetMappedRangeValueClamped(FVector2f(MinBudLightDetected, MaxBudLightDetected),FVector2f(0, 1), CurrentLightAvailable);
+
+			// Health
+			NormalisedHealth[PointIndex] = FMath::GetMappedRangeValueClamped(FVector2f(MinHealth, MaxHealth),FVector2f(0, 1),NormalisedHealth[PointIndex]);
+		}
+		// }
+		
 		// Compute attachment points in terms of normalized distances
 		OutCollection.EmptyGroup(PV::GroupNames::FoliageGroup);
  
@@ -113,6 +246,7 @@ void FPVFoliage::DistributeFoliage(FManagedArrayCollection& OutCollection, const
 		
 		for (int32 BranchIndex = 0; BranchIndex < NumBranches; ++BranchIndex)
 		{
+			ConditionAttributes ConditionAttributes;
 			TArray<float> NormalizedAttachmentPoints;
 			TArray<float> NormalizedAttachmentPointsFinal;
 			TArray<FVector3f> AttachmentPointsPositions;
@@ -120,6 +254,9 @@ void FPVFoliage::DistributeFoliage(FManagedArrayCollection& OutCollection, const
 			TArray<FVector3f> AttachmentPointsUpVectors;
 			TArray<FVector3f> AttachmentPointsNormalVectors;
 			TArray<float> AttachmentPointsScales;
+			
+			const int32 NumFoliageInfo = FoliageFacadeOutput.NumFoliageInfo();
+			
 			TArray<int32> BranchPoints = BranchFacadeSource.GetPoints(BranchIndex);
 
 			if (BranchPoints.Num() < 2)
@@ -181,15 +318,20 @@ void FPVFoliage::DistributeFoliage(FManagedArrayCollection& OutCollection, const
 
 					const int32 CurrentPointIndex = BranchPoints[BranchPointIndex];
 					const int32 NextPointIndex = BranchPoints[BranchPointIndex + 1];
+
 					const FVector3f CurrentPointPosition = PointFacadeSource.GetPosition(CurrentPointIndex);
 					const FVector3f NextPointPosition = PointFacadeSource.GetPosition(NextPointIndex);
+
 					float EthyleneLevelAtCurrentPoint = BudHormoneLevels[CurrentPointIndex][4];
 					float EthyleneLevelAtNextPoint = BudHormoneLevels[NextPointIndex][4];
+
 					const float CurrentPointLengthFromRoot = PointFacadeSource.GetLengthFromRoot(CurrentPointIndex);
 					const float NextPointLengthFromRoot = PointFacadeSource.GetLengthFromRoot(NextPointIndex);
+					
 					const float CurrentPointLfrNormalized = FMath::GetMappedRangeValueClamped(
 						FVector2f(FirstPointLengthFromRoot, LastPointLengthFromRoot),
 						FVector2f(0.0f, 1.0f), CurrentPointLengthFromRoot);
+
 					const float NextPointLfrNormalized = FMath::GetMappedRangeValueClamped(
 						FVector2f(FirstPointLengthFromRoot, LastPointLengthFromRoot),
 						FVector2f(0.0f, 1.0f),
@@ -245,12 +387,55 @@ void FPVFoliage::DistributeFoliage(FManagedArrayCollection& OutCollection, const
 						Scale = Scale * RandomValue;
 
 						AttachmentPointsScales.Add(Scale);
-
+						
 						const FVector3f CurrentPointApicalDirectionVector(BudDirectionAttribute[CurrentPointIndex][0]);
 						const FVector3f NextPointApicalDirectionVector(BudDirectionAttribute[NextPointIndex][0]);
 						const FVector3f BudApicalDirectionVector = FMath::Lerp(CurrentPointApicalDirectionVector,
 							NextPointApicalDirectionVector, DistanceAlpha);
 
+						// Condition - Y Dot
+						float YDot = FVector3f::DotProduct(BudApicalDirectionVector, FVector3f::UpVector);
+						YDot = FMath::GetMappedRangeValueClamped(FVector2f(-1, 1),FVector2f(0, 1),YDot);
+						ConditionAttributes.UpAlignment.Add(YDot);
+
+						// Condition - Light
+						// const TArray<float>& CurrentBudLightDetected = PointFacadeSource.GetBudLightDetected(CurrentPointIndex);
+						// const float CurrentLightAvailable = CurrentBudLightDetected.IsValidIndex(PV::Facades::BudLightDetectedAvailableIndex) ? CurrentBudLightDetected[PV::Facades::BudLightDetectedAvailableIndex] : 0.0f;
+						//
+						// const TArray<float>& NextBudLightDetected = PointFacadeSource.GetBudLightDetected(NextPointIndex);
+						// const float NextLightAvailable = NextBudLightDetected.IsValidIndex(PV::Facades::BudLightDetectedAvailableIndex) ? NextBudLightDetected[PV::Facades::BudLightDetectedAvailableIndex] : 0.0f;	
+
+						const float CurrentLightAvailable = NormalisedBudLightDetected[CurrentPointIndex];
+						const float NextLightAvailable = NormalisedBudLightDetected[NextPointIndex];
+						const float BudLightReceived = FMath::Lerp(CurrentLightAvailable,NextLightAvailable, DistanceAlpha);
+						ConditionAttributes.Light.Add(BudLightReceived);
+
+						// Condition - Scale
+						// const float CurrentPointScaleGradient = PointFacadeSource.GetPointScaleGradient(CurrentPointIndex);
+						// const float NextPointScaleGradient = PointFacadeSource.GetPointScaleGradient(NextPointIndex);
+						// const float PointScaleGradient = FMath::Lerp(CurrentPointScaleGradient, NextPointScaleGradient, DistanceAlpha);
+
+						const float CurrentPScale = NormalisedPScales[CurrentPointIndex];
+						const float NextPScale = NormalisedPScales[NextPointIndex];
+						const float PointPScale = FMath::Lerp(CurrentPScale, NextPScale, DistanceAlpha);
+						ConditionAttributes.Scale.Add(PointPScale);
+
+						// Health
+						const float CurrentHealth = NormalisedHealth[CurrentPointIndex];
+						const float NextHealth = NormalisedHealth[NextPointIndex];
+						const float BudHealth = FMath::Lerp(CurrentHealth,NextHealth, DistanceAlpha);
+						ConditionAttributes.Health.Add(BudHealth);
+
+						// Condition - Tip
+						if (BranchPointIndex == BranchPoints.Num() - 2 && DistanceAlpha == 1)
+						{
+							ConditionAttributes.Tip.Add(1);
+						}
+						else
+						{
+							ConditionAttributes.Tip.Add(0);
+						}
+
 						int32 BudsToAdd = 1;
 
 						// Compute Up Vector
@@ -260,64 +445,80 @@ void FPVFoliage::DistributeFoliage(FManagedArrayCollection& OutCollection, const
 							AttachmentPointsUpVectors.Add(BudApicalDirectionVector);
 						}
 						else // Points other than last
-						[[likely]]
-						{
-							const FQuat4f RotationQuat(InitialApicalDirection.GetSafeNormal(), AxillaryRotationDegree);
-							FVector3f NewAxillaryDirection = RotationQuat.RotateVector(InitialAxillaryDirection.GetSafeNormal());
-							FQuat4f ApicalCorrection = FQuat4f::FindBetweenNormals(
-								InitialApicalDirection.GetSafeNormal(), BudApicalDirectionVector.GetSafeNormal());
-							NewAxillaryDirection = ApicalCorrection.RotateVector(NewAxillaryDirection).GetSafeNormal();
-
-							int32 NumBuds = RandomStream.RandRange(MinBudsLeaf, MaxBudsLeaf);
-							float RotationPerBud = 360.0f / static_cast<float>(NumBuds);
-
-							for (int32 i = 0; i < NumBuds; ++i)
+							[[likely]]
 							{
-								float RotationAngle = FMath::DegreesToRadians(RotationPerBud * i);
-								const FQuat4f AxillaryRotation(InitialApicalDirection.GetSafeNormal(), RotationAngle);
-								FVector3f AxillaryDirectionModified = AxillaryRotation.RotateVector(NewAxillaryDirection.GetSafeNormal());
-
-								if (VectorSettings.OverrideAxilAngle)
-								{
-									float AxilAngleRampValue = VectorSettings.AxilAngleRamp
-										? FMath::Clamp(
-											VectorSettings.AxilAngleRamp->GetRichCurveConst()->Eval(1.0f - BranchGradient),
-											0.0f,
-											1.0f)
-										: 1.0f;
+								const FQuat4f RotationQuat(InitialApicalDirection.GetSafeNormal(), AxillaryRotationDegree);
+								FVector3f NewAxillaryDirection = RotationQuat.RotateVector(InitialAxillaryDirection.GetSafeNormal());
+								FQuat4f ApicalCorrection = FQuat4f::FindBetweenNormals(
+									InitialApicalDirection.GetSafeNormal(), BudApicalDirectionVector.GetSafeNormal());
+								NewAxillaryDirection = ApicalCorrection.RotateVector(NewAxillaryDirection).GetSafeNormal();
 
-									AxilAngleRampValue = AxilAngleRampValue * VectorSettings.AxilAngleRampUpperValue;
-									const float AxilAngleBlended = FMath::Lerp(VectorSettings.AxilAngle, AxilAngleRampValue,
-										VectorSettings.AxilAngleRampEffect);
+								int32 NumBuds = RandomStream.RandRange(MinBudsLeaf, MaxBudsLeaf);
+								float RotationPerBud = 360.0f / static_cast<float>(NumBuds);
 
-									const FVector3f RotationAxis = FVector3f::CrossProduct(
-										BudApicalDirectionVector.GetSafeNormal(), AxillaryDirectionModified.GetSafeNormal());
-
-									const FQuat4f Rotation(RotationAxis.GetSafeNormal(), FMath::DegreesToRadians(AxilAngleBlended));
-									const FVector3f AxilVector = Rotation.RotateVector(BudApicalDirectionVector.GetSafeNormal());
-
-									AttachmentPointsUpVectors.Add(AxilVector);
-								}
-								else
+								for (int32 i = 0; i < NumBuds; ++i)
 								{
-									AttachmentPointsUpVectors.Add(AxillaryDirectionModified);
+									float RotationAngle = FMath::DegreesToRadians(RotationPerBud * i);
+									const FQuat4f AxillaryRotation(InitialApicalDirection.GetSafeNormal(), RotationAngle);
+									FVector3f AxillaryDirectionModified = AxillaryRotation.RotateVector(NewAxillaryDirection.GetSafeNormal());
+
+									if (VectorSettings.OverrideAxilAngle)
+									{
+										float AxilAngleRampValue = VectorSettings.AxilAngleRamp
+											? FMath::Clamp(
+												VectorSettings.AxilAngleRamp->GetRichCurveConst()->Eval(1.0f - BranchGradient),
+												0.0f,
+												1.0f)
+											: 1.0f;
+
+										AxilAngleRampValue = AxilAngleRampValue * VectorSettings.AxilAngleRampUpperValue;
+										const float AxilAngleBlended = FMath::Lerp(VectorSettings.AxilAngle, AxilAngleRampValue,
+											VectorSettings.AxilAngleRampEffect);
+
+										const FVector3f RotationAxis = FVector3f::CrossProduct(
+											BudApicalDirectionVector.GetSafeNormal(), AxillaryDirectionModified.GetSafeNormal());
+
+										const FQuat4f Rotation(RotationAxis.GetSafeNormal(), FMath::DegreesToRadians(AxilAngleBlended));
+										const FVector3f AxilVector = Rotation.RotateVector(BudApicalDirectionVector.GetSafeNormal());
+
+										AttachmentPointsUpVectors.Add(AxilVector);
+									}
+									else
+									{
+										AttachmentPointsUpVectors.Add(AxillaryDirectionModified);
+									}
+
+									if (i > 0)
+									{
+										FVector3f LastPositionAdded = AttachmentPointsPositions.Last();
+										float LastScaleAdded = AttachmentPointsScales.Last();
+										float LastLengthFromRootAdded = AttachmentPointsLengthsFromRoot.Last();
+										AttachmentPointsPositions.Add(LastPositionAdded);
+										AttachmentPointsScales.Add(LastScaleAdded);
+										AttachmentPointsLengthsFromRoot.Add(LastLengthFromRootAdded);
+
+										// Conditions
+										float LastConditionScale = ConditionAttributes.Scale.Last();
+										ConditionAttributes.Scale.Add(LastConditionScale);
+									
+										float LastConditionLight = ConditionAttributes.Light.Last();
+										ConditionAttributes.Light.Add(LastConditionLight);
+									
+										float LastConditionUpAlignment  = ConditionAttributes.UpAlignment.Last();
+										ConditionAttributes.UpAlignment.Add(LastConditionUpAlignment);
+
+										float LastConditionTip = ConditionAttributes.Tip.Last();
+										ConditionAttributes.Tip.Add(LastConditionTip);
+
+										float LastConditionHealth = ConditionAttributes.Health.Last();
+										ConditionAttributes.Health.Add(LastConditionHealth);
+									}
 								}
 
-								if (i > 0)
-								{
-									FVector3f LastPositionAdded = AttachmentPointsPositions.Last();
-									float LastScaleAdded = AttachmentPointsScales.Last();
-									float LastLengthFromRootAdded = AttachmentPointsLengthsFromRoot.Last();
-									AttachmentPointsPositions.Add(LastPositionAdded);
-									AttachmentPointsScales.Add(LastScaleAdded);
-									AttachmentPointsLengthsFromRoot.Add(LastLengthFromRootAdded);
-								}
+								BudsToAdd = NumBuds;
+								AxillaryRotationIteration++;
 							}
 
-							BudsToAdd = NumBuds;
-							AxillaryRotationIteration++;
-						}
-
 						// Compute Normal vectors based on light optimal direction
 						const FVector3f CurrentPointLightOptimalVector(BudDirectionAttribute[CurrentPointIndex][2]);
 						const FVector3f NextPointLightOptimalVector(BudDirectionAttribute[NextPointIndex][2]);
@@ -339,24 +540,141 @@ void FPVFoliage::DistributeFoliage(FManagedArrayCollection& OutCollection, const
 			// Set foliage entries
 			TManagedArrayAccessor<FVector3f> PivotPointsAttribute(OutCollection, PV::AttributeNames::FoliagePivotPoint,
 				PV::GroupNames::FoliageGroup);
-			const int32 NumFoliageMeshes = NonEmptyFoliageMeshNameIndices.Num();
-			int32 CurrentFoliageMeshIndex = BranchIndex % NumFoliageMeshes;
+
+			const TArray<FPVFoliageInfo> FoliageInfos = FoliageFacadeOutput.GetFoliageInfos();
+			
+
+			const int32 BranchNumber = BranchFacadeSource.GetBranchNumber(BranchIndex);
+
+			FRandomStream FoliageRandomStream;
+			FoliageRandomStream.Initialize(BranchNumber);
+			
 			for (int32 Idx = 0; Idx < AttachmentPointsLengthsFromRoot.Num(); ++Idx)
 			{
+				TArray<FFoliageWeightInfo> AvailableFoliageWeights;
+				
+				for (int i = 0; i < NumFoliageInfo; ++i)
+				{
+					auto FoliageInfo = FoliageInfos[i];
+
+					if (FoliageInfo.IsValid())
+					{
+						AvailableFoliageWeights.Add(FFoliageWeightInfo(i, 0));
+					}
+				}
+
+				// Shuffle foliage
+				for (int i = AvailableFoliageWeights.Num() - 1; i > 0; --i)
+				{
+					int Index = RandomStream.RandRange(0, i);
+					AvailableFoliageWeights.Swap(i, Index);
+				}
+
+				if(DistributionConditions.HasActiveCondition())
+				{
+					for (EFoliageDistributionCondition Condition : TEnumRange<EFoliageDistributionCondition>())
+					{
+						FPVFoliageDistributionConditionSettings OutConditionSettings;
+						
+						if (DistributionConditions.IsActiveCondition(Condition) && DistributionConditions.GetSettings(Condition, OutConditionSettings))
+						{
+							float PointAttributeValue = ConditionAttributes.GetConditionValue(Condition,Idx);
+					
+							const float AttributeValue = PointAttributeValue + OutConditionSettings.Offset;
+
+							for (int Index = 0; Index < AvailableFoliageWeights.Num(); ++Index)
+							{
+								int FoliageIndex = AvailableFoliageWeights[Index].FoliageIndex;
+
+								if (FoliageInfos.IsValidIndex(FoliageIndex))
+								{
+									auto FoliageInfo = FoliageInfos[FoliageIndex];
+									float FoliageValue = FoliageInfo.Attributes.GetAttributeValue(Condition);
+									float FoliageDistribution = FMath::Abs(FoliageValue - AttributeValue) * OutConditionSettings.Weight;
+									
+									// UE_LOG(LogProceduralVegetation, Log, TEXT("FoliageWeight : Index : %i : FoliageValue : %f : AttributeValue : %f : Weight : %f : Final Value: %f"),
+									// 	FoliageIndex,
+									// 	FoliageValue,
+									// 	AttributeValue,
+									// 	OutConditionSettings.Weight,
+									// 	FoliageDistribution);	
+
+									AvailableFoliageWeights[Index].Weight += FoliageDistribution;
+								}
+								else
+								{
+									check(false);
+								}
+							}	
+						}
+					}
+
+					float MinFoliageWeight = Algo::MinElementBy(AvailableFoliageWeights, [](const FFoliageWeightInfo& Info) { return Info.Weight; })->Weight;
+					float MaxFoliageWeight = Algo::MaxElementBy(AvailableFoliageWeights, [](const FFoliageWeightInfo& Info) { return Info.Weight; })->Weight;
+
+					for (FFoliageWeightInfo& Info : AvailableFoliageWeights)
+					{
+						Info.Weight = FMath::GetMappedRangeValueClamped(
+									FVector2f(MinFoliageWeight, MaxFoliageWeight),
+									FVector2f(0, 1),
+									Info.Weight);
+					}
+				}
+
+				// Pick foliage
+				AvailableFoliageWeights.Sort([&](const FFoliageWeightInfo& A, const FFoliageWeightInfo& B)
+					{
+						return A.Weight < B.Weight; // Ascending
+					});
+
+				
+				TArray<int> PickedFoliage;
+
+				
+				for (const FFoliageWeightInfo& Info : AvailableFoliageWeights)
+				{
+					if (Info.Weight < DistributionConditions.CutoffThreshold)
+					{
+						PickedFoliage.Add(Info.FoliageIndex);
+					}
+					else if (PickedFoliage.Num() < DistributionConditions.MinimumCandidates)
+					{
+						PickedFoliage.Add(Info.FoliageIndex);
+					}
+					else
+					{
+						break;
+					}
+				}
+
+				const int PickedIndex = RandomStream.RandRange(0, PickedFoliage.Num() - 1);
+				int PickedFoliageIndex = PickedFoliage[PickedIndex];
+
 				const int32 NewIndex = PivotPointsAttribute.AddElements(1);
 				FoliageFacadeOutput.SetFoliageEntry(NewIndex, {
-					.NameId = NonEmptyFoliageMeshNameIndices[CurrentFoliageMeshIndex],
+					.NameId = PickedFoliageIndex,
 					.BranchId = BranchIndex,
 					.PivotPoint = AttachmentPointsPositions[Idx],
 					.UpVector = AttachmentPointsUpVectors[Idx],
 					.NormalVector = AttachmentPointsNormalVectors[Idx],
 					.Scale = AttachmentPointsScales[Idx],
-					.LengthFromRoot = AttachmentPointsLengthsFromRoot[Idx]
+					.LengthFromRoot = AttachmentPointsLengthsFromRoot[Idx],
+					.ConditionUpAlignment = ConditionAttributes.UpAlignment[Idx],
+					.ConditionTip = ConditionAttributes.Tip[Idx],
+					.ConditionLight = ConditionAttributes.Light[Idx],
+					.ConditionScale = ConditionAttributes.Scale[Idx],
+					.ConditionHealth = ConditionAttributes.Health[Idx]
 				});
-				CurrentFoliageMeshIndex = (CurrentFoliageMeshIndex + 1) % NumFoliageMeshes;
 			}
 		}
 
+		// TArray<float> Scales;
+		//
+		// for (int32 i = 0; i < FoliageFacadeOutput.NumFoliageEntries(); ++i)
+		// {
+		// 	FoliageFacadeOutput.GetFoliageEntry(i).Scale;
+		// }
+		
 		// Re-calculate Foliage IDs for each branch
 		TMap<int32, TArray<int32>> BranchIndexToFoliageIDs;
 		for (int32 i = 0; i < FoliageFacadeOutput.NumFoliageEntries(); ++i)
diff --git a/Engine/Plugins/Experimental/ProceduralVegetationEditor/Source/ProceduralVegetation/Private/Implementations/PVFoliage.h b/Engine/Plugins/Experimental/ProceduralVegetationEditor/Source/ProceduralVegetation/Private/Implementations/PVFoliage.h
--- a/Engine/Plugins/Experimental/ProceduralVegetationEditor/Source/ProceduralVegetation/Private/Implementations/PVFoliage.h
+++ b/Engine/Plugins/Experimental/ProceduralVegetationEditor/Source/ProceduralVegetation/Private/Implementations/PVFoliage.h
@@ -3,8 +3,7 @@
 #pragma once
 
 #include "CoreMinimal.h"
-
-#include "Curves/CurveFloat.h"
+#include "Misc/EnumRange.h"
 
 #include "GeometryCollection/ManagedArrayCollection.h"
 
@@ -72,15 +71,288 @@ struct FPhyllotaxySettings
 	const float PhyllotaxyAdditionalAngle;
 };
 
+USTRUCT()
+struct FPVFoliageDistributionConditionSettings
+{
+	GENERATED_BODY()
+
+	UPROPERTY(EditAnywhere, Category="ConditionSettings", meta=(ClampMin=0.0f, ClampMax=1.0f, UIMin=0.0f, UIMax=1.0f, Tooltip=""))
+	float Weight = 1.0f;
+
+	UPROPERTY(EditAnywhere, Category="ConditionSettings", meta=(ClampMin=-1.0f, ClampMax=1.0f, UIMin=-1.0f, UIMax=1.0f, Tooltip=""))
+	float Offset = 0.0f;
+
+	FString ToString() const
+	{
+		return FString::Printf(TEXT("Weight = %f, Offset = %f"), Weight, Offset);
+	}
+};
+
+UENUM()
+enum class EFoliageDistributionCondition : uint8
+{
+	Light,
+	Scale,
+	UpAlignment,
+	Tip,
+	Health,
+	Count UMETA(Hidden),
+	None UMETA(Hidden),
+};
+
+ENUM_RANGE_BY_COUNT(EFoliageDistributionCondition, EFoliageDistributionCondition::Count)
+
+
+USTRUCT()
+struct FPVFoliageDistributionConditions
+{
+	GENERATED_BODY()
+
+	UPROPERTY(EditAnywhere, Category="ConditionSettings", meta=(ClampMin=0.0f, ClampMax=1.0f, UIMin=0.0f, UIMax=1.0f, Tooltip=""))
+	float CutoffThreshold = 0.3f;
+
+	UPROPERTY(EditAnywhere, Category="ConditionSettings", meta=(ClampMin=1, ClampMax=10, UIMin=0, UIMax=10, Tooltip=""))
+	int32 MinimumCandidates = 1;
+	
+	UPROPERTY(EditAnywhere, Category="ConditionSettings", meta=(PinHiddenByDefault, InlineEditConditionToggle))
+	bool bActivateLight = false;
+
+	UPROPERTY(EditAnywhere, Category="ConditionSettings", meta=(EditCondition="bActivateLight", Tooltip=""))
+	FPVFoliageDistributionConditionSettings Light;
+
+	UPROPERTY()
+	bool bActivateScale = false;
+
+	UPROPERTY()
+	FPVFoliageDistributionConditionSettings Scale;
+
+	UPROPERTY(EditAnywhere, Category="ConditionSettings", meta=(PinHiddenByDefault, InlineEditConditionToggle))
+	bool bActivateUpAlignment = false;
+
+	UPROPERTY(EditAnywhere, Category="ConditionSettings", meta=(EditCondition="bActivateUpAlignment", Tooltip=""))
+	FPVFoliageDistributionConditionSettings UpAlignment;
+
+	UPROPERTY(EditAnywhere, Category="ConditionSettings", meta=(PinHiddenByDefault, InlineEditConditionToggle))
+	bool bActivateTip = false;
+	
+	UPROPERTY(EditAnywhere, Category="ConditionSettings", meta=(EditCondition="bActivateTip", Tooltip=""))
+	FPVFoliageDistributionConditionSettings Tip;
+
+	UPROPERTY(EditAnywhere, Category="ConditionSettings", meta=(PinHiddenByDefault, InlineEditConditionToggle))
+	bool bActivateHealth = false;
+	
+	UPROPERTY(EditAnywhere, Category="ConditionSettings", meta=(EditCondition="bActivateHealth", Tooltip=""))
+	FPVFoliageDistributionConditionSettings Health;
+
+	bool HasActiveCondition() const
+	{
+		for (EFoliageDistributionCondition Condition : TEnumRange<EFoliageDistributionCondition>())
+		{
+			if (IsActiveCondition(Condition))
+			{
+				return true;
+			}
+		}
+		return false;
+	}
+
+	bool IsActiveCondition(const EFoliageDistributionCondition InCondition) const 
+	{
+		switch (InCondition)
+		{
+		case EFoliageDistributionCondition::Light:
+			return bActivateLight && Light.Weight > 0.0f;
+			
+		case EFoliageDistributionCondition::Scale:
+			return bActivateScale && Scale.Weight > 0.0f; 
+			
+		case EFoliageDistributionCondition::UpAlignment:
+			return bActivateUpAlignment && UpAlignment.Weight > 0.0f;
+			
+		case EFoliageDistributionCondition::Tip:
+			return bActivateTip && Tip.Weight > 0.0f;
+
+		case EFoliageDistributionCondition::Health:
+			return bActivateHealth && Health.Weight > 0.0f;
+
+		default:
+			return false;
+		}
+	}
+
+	void SetConditionState(const EFoliageDistributionCondition InCondition, const bool InNewState) 
+	{
+		switch (InCondition)
+		{
+		case EFoliageDistributionCondition::Light:
+			bActivateLight = InNewState;
+			break;
+			
+			
+		case EFoliageDistributionCondition::Scale:
+			bActivateScale = InNewState;
+			break;
+			
+		case EFoliageDistributionCondition::UpAlignment:
+			bActivateUpAlignment = InNewState;
+			break;
+			
+		case EFoliageDistributionCondition::Tip:
+			bActivateTip = InNewState;
+			break;
+
+		case EFoliageDistributionCondition::Health:
+			bActivateHealth = InNewState;
+			break;
+
+			default:
+			check(false);
+		}
+	}
+
+	FPVFoliageDistributionConditionSettings* GetSettings(const EFoliageDistributionCondition InCondition)
+	{
+		switch (InCondition)
+		{
+			case EFoliageDistributionCondition::Light:
+				return &Light;
+				
+			case EFoliageDistributionCondition::Scale:
+				return &Scale;
+			
+				
+			case EFoliageDistributionCondition::UpAlignment:
+				return &UpAlignment;
+				
+			case EFoliageDistributionCondition::Tip:
+				return &Tip;
+
+			case EFoliageDistributionCondition::Health:
+				return &Health;
+
+			default:
+			return nullptr;
+		}
+	}
+
+	bool GetSettings(const EFoliageDistributionCondition InCondition, FPVFoliageDistributionConditionSettings& OutSettings) const
+	{
+		switch (InCondition)
+		{
+			case EFoliageDistributionCondition::Light:
+			{
+				OutSettings =  Light;
+				return true;
+			}
+			
+			case EFoliageDistributionCondition::Scale:
+			{
+				OutSettings = Scale;
+				return true;
+			}
+			
+			case EFoliageDistributionCondition::UpAlignment:
+			{
+				OutSettings = UpAlignment;
+				return true;
+			}
+			
+			case EFoliageDistributionCondition::Tip:
+			{
+				OutSettings = Tip;
+				return true;
+			}
+
+			case EFoliageDistributionCondition::Health:
+			{
+				OutSettings = Health;
+				return true;
+			}
+		}
+
+		return false;
+	}
+};
+
 struct FMiscSettings
 {
 	const int32 RandomSeed;
 };
 
+
+USTRUCT()
+struct FPVFoliageAttributes
+{
+	GENERATED_BODY()
+
+	UPROPERTY(EditAnywhere, Category="Info", meta=(ClampMin=0.0f, ClampMax=1.0f, UIMin=0.0f, UIMax=1.0f, Tooltip=""))
+	float Light = 0.0f;
+
+	UPROPERTY()
+	float Scale = 0.0f;
+
+	UPROPERTY(EditAnywhere, Category="Info", meta=(ClampMin=0.0f, ClampMax=1.0f, UIMin=0.0f, UIMax=1.0f, Tooltip=""))
+	float UpAlignment = 0.0f;
+
+	UPROPERTY(EditAnywhere, Category="Info", meta=(ClampMin=0.0f, ClampMax=1.0f, UIMin=0.0f, UIMax=1.0f, Tooltip=""))
+	float Health = 0.0f;
+
+	UPROPERTY(EditAnywhere, Category="Info", meta=(Tooltip=""))
+	bool Tip = false;
+
+	float GetAttributeValue(const EFoliageDistributionCondition InCondition) const
+	{
+		switch (InCondition)
+		{
+			case EFoliageDistributionCondition::Light:
+			return Light;
+			
+			case EFoliageDistributionCondition::Scale:
+			return Scale;
+			
+			case EFoliageDistributionCondition::UpAlignment:
+			return UpAlignment;
+			
+			case EFoliageDistributionCondition::Tip:
+			return Tip ? 1.0f : 0.0f;
+
+			case EFoliageDistributionCondition::Health:
+			return Health;
+
+			default:
+			return 0.0f;
+		}
+	}
+
+	FString ToString() const
+	{
+		return FString::Printf(TEXT("Light = %f, Scale = %f, UpAlignment = %f, Health = %f, Tip = %s"), Light, Scale, UpAlignment, Health, Tip ? TEXT("True") : TEXT("False"));
+	}
+};
+
+USTRUCT()
+struct FPVFoliageInfo
+{
+	GENERATED_BODY()
+	
+	UPROPERTY(EditAnywhere, Category = "Info")
+	TSoftObjectPtr<UObject> Mesh;
+
+	UPROPERTY(EditAnywhere, Category="Info", meta=(ShowInnerProperties, FullyExpand="true", Tooltip=""))
+	FPVFoliageAttributes Attributes;
+
+	bool IsValid() const
+	{
+		// Asset is loaded and pointer is valid
+		// Path is valid, but asset is NOT loaded yet
+		return Mesh.IsValid() || Mesh.ToSoftObjectPath().IsValid(); 
+	}
+};
+
 struct FPVFoliage
 {
 	static void DistributeFoliage(FManagedArrayCollection& OutCollection, const FManagedArrayCollection& SourceCollection,
 	                              const FDistributionSettings& DistributionSettings, const FScaleSettings& ScaleSettings,
 	                              const FVectorSettings& VectorSettings, const FPhyllotaxySettings& PhyllotaxySettings,
-	                              const FMiscSettings& MiscSettings);
+	                              const FMiscSettings& MiscSettings, const FPVFoliageDistributionConditions& DistributionConditions);
 };
diff --git a/Engine/Plugins/Experimental/ProceduralVegetationEditor/Source/ProceduralVegetation/Private/Nodes/PVFoliageDistributorSettings.cpp b/Engine/Plugins/Experimental/ProceduralVegetationEditor/Source/ProceduralVegetation/Private/Nodes/PVFoliageDistributorSettings.cpp
--- a/Engine/Plugins/Experimental/ProceduralVegetationEditor/Source/ProceduralVegetation/Private/Nodes/PVFoliageDistributorSettings.cpp
+++ b/Engine/Plugins/Experimental/ProceduralVegetationEditor/Source/ProceduralVegetation/Private/Nodes/PVFoliageDistributorSettings.cpp
@@ -5,6 +5,7 @@
 #include "DataTypes/PVMeshData.h"
 #include "DataTypes/PVFoliageMeshData.h"
 #include "Data/PCGBasePointData.h"
+#include "Facades/PVFoliageConditionFacade.h"
 #include "GeometryCollection/ManagedArrayCollection.h"
 #include "Helpers/PCGHelpers.h"
 #include "Helpers/PCGSettingsHelpers.h"
@@ -28,6 +29,24 @@ FText UPVFoliageDistributorSettings::GetNodeTooltipText() const
 		"\n\nPress Ctrl + L to lock/unlock node output"
 	); 
 }
+
+void UPVFoliageDistributorSettings::PostEditChangeProperty(FPropertyChangedEvent& PropertyChangedEvent)
+{
+	Super::PostEditChangeProperty(PropertyChangedEvent);
+
+	if (PropertyChangedEvent.Property
+		&& PropertyChangedEvent.Property->GetFName() == GET_MEMBER_NAME_CHECKED(UPVFoliageDistributorSettings, OverrideDistribution))
+	{
+		if (OverrideDistribution)
+		{
+			Conditions = CustomConditions;
+		}
+		else
+		{
+			Conditions = DefaultConditions;
+		}	
+	}
+}
 #endif
 
 UPVFoliageDistributorSettings::UPVFoliageDistributorSettings()
@@ -81,6 +100,9 @@ bool FPVFoliageDistributorElement::ExecuteInternal(FPCGContext* InContext) const
 	const UPVFoliageDistributorSettings* Settings = InContext->GetInputSettings<UPVFoliageDistributorSettings>();
 	check(Settings);
 
+	UPVFoliageDistributorSettings* MutableSettings = const_cast<UPVFoliageDistributorSettings*>(Settings);
+	check(MutableSettings);
+
 	for (const FPCGTaggedData& Input : InContext->InputData.GetInputsByPin(PCGPinConstants::DefaultInputLabel))
 	{
 		if(const UPVMeshData* InputData = Cast<UPVMeshData>(Input.Data))
@@ -129,8 +151,47 @@ bool FPVFoliageDistributorElement::ExecuteInternal(FPCGContext* InContext) const
 				.RandomSeed = Settings->RandomSeed
 			};
 
+			if (!Settings->OverrideDistribution)
+			{
+				PV::Facades::FFoliageConditionFacade FoliageConditionFacade(Collection);
+
+				auto DefaultConditionsData = FoliageConditionFacade.GetData();
+
+				for (auto ConditonInfo : DefaultConditionsData)
+				{
+					auto ConditionType = ConditonInfo.GetType();
+
+					if (ConditionType != EFoliageDistributionCondition::None)
+					{
+						
+						if (auto ConditionSettings = MutableSettings->DefaultConditions.GetSettings(ConditionType))
+						{
+							ConditionSettings->Weight = ConditonInfo.Weight;
+							ConditionSettings->Offset = ConditonInfo.Offset;
+
+							MutableSettings->DefaultConditions.SetConditionState(ConditionType, true);	
+						}
+					}
+				}
+
+				if (!MutableSettings->bCustomConditionSet)
+				{
+					MutableSettings->CustomConditions = MutableSettings->DefaultConditions;
+					MutableSettings->bCustomConditionSet = true;
+				}
+			}
+
+			if (!MutableSettings->OverrideDistribution)
+			{
+				MutableSettings->Conditions = MutableSettings->DefaultConditions;
+			}
+			else
+			{
+				MutableSettings->CustomConditions = MutableSettings->Conditions;
+			}	
+
 			FPVFoliage::DistributeFoliage(Collection, Collection, DistributionSettings, ScaleSettings, VectorSettings,
-												 PhyllotaxySettings, MiscSettings);
+												 PhyllotaxySettings, MiscSettings, Settings->Conditions);
 
 			OutManagedArrayCollectionData->Initialize(MoveTemp(Collection));
 			InContext->OutputData.TaggedData.Emplace(OutManagedArrayCollectionData);
diff --git a/Engine/Plugins/Experimental/ProceduralVegetationEditor/Source/ProceduralVegetation/Private/Nodes/PVFoliageDistributorSettings.h b/Engine/Plugins/Experimental/ProceduralVegetationEditor/Source/ProceduralVegetation/Private/Nodes/PVFoliageDistributorSettings.h
--- a/Engine/Plugins/Experimental/ProceduralVegetationEditor/Source/ProceduralVegetation/Private/Nodes/PVFoliageDistributorSettings.h
+++ b/Engine/Plugins/Experimental/ProceduralVegetationEditor/Source/ProceduralVegetation/Private/Nodes/PVFoliageDistributorSettings.h
@@ -29,6 +29,7 @@ public:
 	virtual FLinearColor GetNodeTitleColor() const override { return FLinearColor::Blue; }
 
 	virtual TArray<EPVRenderType> GetDefaultRenderType() const override { return TArray{ EPVRenderType::Mesh, EPVRenderType::Foliage }; }
+	virtual void PostEditChangeProperty(FPropertyChangedEvent& PropertyChangedEvent) override;
 	#endif
 
 protected:
@@ -44,7 +45,7 @@ private:
 	float EthyleneThreshold = 0.85f;
 
 	UPROPERTY(EditAnywhere, Category="Distribution Settings", meta=(Tooltip="Override default placement with custom distribution rules.\n\nOverides the distribution settings loaded from the procedural vegetation preset."))
-	bool OverrideDistribution = true;
+	bool OverrideDistribution = false;
 
 	UPROPERTY(EditAnywhere, Category="Distribution Settings", meta=(EditCondition="OverrideDistribution", ClampMin=0.0f, ClampMax=1.0f, UIMin=0.0f, UIMax=1.0f, Tooltip="Minimum distance between placed instances.\n\nHigher spacing yields a more open distribution; lower spacing allows denser placement."))
 	float InstanceSpacing = 0.3f;
@@ -114,6 +115,18 @@ private:
 
 	UPROPERTY(EditAnywhere, Category="Misc Settings", meta=(EditCondition="OverrideDistribution"))
 	int32 RandomSeed = 123456;
+
+	UPROPERTY(EditAnywhere, Category="ConditionSettings", meta = (ShowOnlyInnerProperties, EditCondition="OverrideDistribution"))
+	FPVFoliageDistributionConditions Conditions;
+	
+	UPROPERTY()
+	FPVFoliageDistributionConditions DefaultConditions;
+	
+	UPROPERTY()
+	FPVFoliageDistributionConditions CustomConditions;
+
+	UPROPERTY()
+	bool bCustomConditionSet = false;
 };
 
 class FPVFoliageDistributorElement : public FPVBaseElement
diff --git a/Engine/Plugins/Experimental/ProceduralVegetationEditor/Source/ProceduralVegetation/Private/Nodes/PVFoliagePaletteSettings.cpp b/Engine/Plugins/Experimental/ProceduralVegetationEditor/Source/ProceduralVegetation/Private/Nodes/PVFoliagePaletteSettings.cpp
--- a/Engine/Plugins/Experimental/ProceduralVegetationEditor/Source/ProceduralVegetation/Private/Nodes/PVFoliagePaletteSettings.cpp
+++ b/Engine/Plugins/Experimental/ProceduralVegetationEditor/Source/ProceduralVegetation/Private/Nodes/PVFoliagePaletteSettings.cpp
@@ -12,6 +12,7 @@
 #include "Helpers/PCGHelpers.h"
 #include "Helpers/PCGSettingsHelpers.h"
 #include "Helpers/PVAnalyticsHelper.h"
+#include "Implementations/PVFoliage.h"
 
 #define LOCTEXT_NAMESPACE "PVFoliagePaletteSettings"
 
@@ -56,21 +57,29 @@ void UPVFoliagePaletteSettings::PostEditChangeChainProperty(FPropertyChangedChai
 
 	if (MemberProperty && Property)
 	{
-		if (MemberProperty->GetFName() == GET_MEMBER_NAME_CHECKED(UPVFoliagePaletteSettings, FoliageMeshes))
+		if (MemberProperty->GetFName() == GET_MEMBER_NAME_CHECKED(UPVFoliagePaletteSettings, FoliageInfos))
 		{
 			FName Name = Property->GetFName();
 			
 			int Index = PropertyChangedEvent.GetArrayIndex(Name.ToString());
 
-			if (Index >= 0 && FoliageMeshes.Num() > Index)
+			if (Index >= 0 && FoliageInfos.Num() > Index)
 			{
-				if (!FoliageMeshes[Index].IsNull())
+				if (!FoliageInfos[Index].Mesh.IsNull())
 				{
-					FString MeshPath = FoliageMeshes[Index].ToString();
+					FString MeshPath = FoliageInfos[Index].Mesh.ToString();
 
 					PV::Analytics::SendFoliageMeshChangeEvent(MeshPath);	
 				}
 			}
+
+			// For backwards compatibility 
+			FoliageMeshes.Empty();
+
+			for (int i = 0; i < FoliageInfos.Num(); i++)
+			{
+				FoliageMeshes.Add(FoliageInfos[i].Mesh);
+			}
 		}
 	}
 }
@@ -92,6 +101,22 @@ FPCGElementPtr UPVFoliagePaletteSettings::CreateElement() const
 	return MakeShared<FPVFoliageElement>();
 }
 
+void UPVFoliagePaletteSettings::PostLoad()
+{
+	Super::PostLoad();
+
+	if (bOverrideFoliage)
+	{
+		if (FoliageInfos.Num() == 0)
+		{
+			for (int32 i = 0; i < FoliageMeshes.Num(); i++)
+			{
+				FoliageInfos.Emplace(FoliageMeshes[i]);
+			}
+		}
+	}
+}
+
 bool FPVFoliageElement::ExecuteInternal(FPCGContext* InContext) const
 {
 	TRACE_CPUPROFILER_EVENT_SCOPE(FPVFoliageElement::Execute);
@@ -117,45 +142,47 @@ bool FPVFoliageElement::ExecuteInternal(FPCGContext* InContext) const
 
 			if (!Settings->bOverrideFoliage)
 			{
-				MutableSettings->FoliageMeshes.Empty();
+				MutableSettings->FoliageInfos.Empty();
 				MutableSettings->bOverrideFoliage = false;
-				for (int32 FoliageIndex = 0; FoliageIndex < Facade.NumFoliageNames(); ++FoliageIndex)
+				for (int32 FoliageIndex = 0; FoliageIndex < Facade.NumFoliageInfo(); ++FoliageIndex)
 				{
-					MutableSettings->FoliageMeshes.Emplace(FSoftObjectPath(Facade.GetFoliageName(FoliageIndex)));
+					FPVFoliageInfo FoliageInfo = Facade.GetFoliageInfo(FoliageIndex);
+					MutableSettings->FoliageInfos.Emplace(FoliageInfo);
 				}
 #if WITH_EDITOR
 				MutableSettings->PostEditChange();
 #endif
 			}
 
-			const int32 IterationLimit = FMath::Min(Settings->FoliageMeshes.Num(), Facade.NumFoliageNames());
-			TArray<FString> FoliageNames;
+			const int32 IterationLimit = FMath::Min(Settings->FoliageInfos.Num(), Facade.NumFoliageInfo());
+			TArray<FPVFoliageInfo> FoliageInfos;
+			
 			for (int32 i = 0; i < IterationLimit; ++i)
 			{
-				FoliageNames.Emplace(Settings->FoliageMeshes[i].ToString());
+				FoliageInfos.Emplace(Settings->FoliageInfos[i]);
 			}
 
-			if (Settings->FoliageMeshes.Num() > Facade.NumFoliageNames())
+			if (Settings->FoliageInfos.Num() > Facade.NumFoliageInfo())
 			{
 				UE_LOG(LogProceduralVegetation, Warning,
 				       TEXT("Num of Foliage meshes specified is greater than the number of available entries, ignoring the last %d entries"),
-				       Settings->FoliageMeshes.Num() - Facade.NumFoliageNames());
+				       Settings->FoliageInfos.Num() - Facade.NumFoliageInfo());
 			}
-			else if (Settings->FoliageMeshes.Num() < Facade.NumFoliageNames())
+			else if (Settings->FoliageInfos.Num() < Facade.NumFoliageInfo())
 			{
-				if (Settings->FoliageMeshes.Num() > 0)
+				if (Settings->FoliageInfos.Num() > 0)
 				{
 					UE_LOG(LogProceduralVegetation, Warning,
 					       TEXT("Num of Foliage meshes specified is less than the number of available entries, repeating the last entry"));
 
-					for (int32 i = Settings->FoliageMeshes.Num(); i < Facade.NumFoliageNames(); ++i)
+					for (int32 i = Settings->FoliageInfos.Num(); i < Facade.NumFoliageInfo(); ++i)
 					{
-						FoliageNames.Emplace(Settings->FoliageMeshes.Last().ToString());
+						FoliageInfos.Emplace(Settings->FoliageInfos.Last());
 					}
 				}
 			}
 
-			Facade.SetFoliageNames(FoliageNames);
+			Facade.SetFoliageInfos(FoliageInfos);
 
 			OutManagedArrayCollectionData->Initialize(MoveTemp(Collection));
 			InContext->OutputData.TaggedData.Emplace(OutManagedArrayCollectionData);
diff --git a/Engine/Plugins/Experimental/ProceduralVegetationEditor/Source/ProceduralVegetation/Private/Nodes/PVFoliagePaletteSettings.h b/Engine/Plugins/Experimental/ProceduralVegetationEditor/Source/ProceduralVegetation/Private/Nodes/PVFoliagePaletteSettings.h
--- a/Engine/Plugins/Experimental/ProceduralVegetationEditor/Source/ProceduralVegetation/Private/Nodes/PVFoliagePaletteSettings.h
+++ b/Engine/Plugins/Experimental/ProceduralVegetationEditor/Source/ProceduralVegetation/Private/Nodes/PVFoliagePaletteSettings.h
@@ -31,14 +31,18 @@ protected:
 	virtual FPCGElementPtr CreateElement() const override;
 
 public:
+	virtual void PostLoad() override;
+
 	UPROPERTY(EditAnywhere, Category="Meshes", meta=(PinHiddenByDefault, InlineEditConditionToggle))
 	bool bOverrideFoliage = false;
 
+	UPROPERTY()
+	TArray<TSoftObjectPtr<UObject>> FoliageMeshes;
+	
 	UPROPERTY(EditAnywhere, Category = "Meshes", EditFixedSize,
-		Meta = (NoResetToDefault, EditCondition="bOverrideFoliage", AllowedClasses = "/Script/Engine.StaticMesh,/Script/Engine.SkeletalMesh", Tooltip=
+		Meta = (NoResetToDefault, EditCondition="bOverrideFoliage", AllowedClasses = "/Script/Engine.StaticMesh,/Script/Engine.SkeletalMesh", ShowOnlyInnerProperties, FullyExpand="true", Tooltip=
 			"List of Foliage meshes used for this plant\n\nList of all foliage meshes applied based on procedural vegetation preset."))
-	TArray<TSoftObjectPtr<UObject>> FoliageMeshes;
-
+	TArray<struct FPVFoliageInfo> FoliageInfos;
 };
 
 class FPVFoliageElement : public FPVBaseElement
diff --git a/Engine/Plugins/Experimental/ProceduralVegetationEditor/Source/ProceduralVegetation/Private/Nodes/PVImporterSettings.cpp b/Engine/Plugins/Experimental/ProceduralVegetationEditor/Source/ProceduralVegetation/Private/Nodes/PVImporterSettings.cpp
--- a/Engine/Plugins/Experimental/ProceduralVegetationEditor/Source/ProceduralVegetation/Private/Nodes/PVImporterSettings.cpp
+++ b/Engine/Plugins/Experimental/ProceduralVegetationEditor/Source/ProceduralVegetation/Private/Nodes/PVImporterSettings.cpp
@@ -11,6 +11,7 @@
 
 #include "Helpers/PCGHelpers.h"
 #include "Helpers/PCGSettingsHelpers.h"
+#include "Helpers/PVFoliageJSONHelper.h"
 #include "Helpers/PVJSONHelper.h"
 
 UPVImporterSettings::UPVImporterSettings()
@@ -68,7 +69,7 @@ bool FPVImporterElement::ExecuteInternal(FPCGContext* InContext) const
 		return true;
 	}
 
-	PV::SetFoliagePaths(Collection, SkeletonFilePath);
+	PV::PVFoliageJSONHelper::SetFoliagePaths(Collection, SkeletonFilePath);
 	
 	UPVGrowthData* OutManagedArrayCollectionData = FPCGContext::NewObject_AnyThread<UPVGrowthData>(InContext);
 	OutManagedArrayCollectionData->Initialize(MoveTemp(Collection));
diff --git a/Engine/Plugins/Experimental/ProceduralVegetationEditor/Source/ProceduralVegetation/Private/ProceduralVegetationPreset.cpp b/Engine/Plugins/Experimental/ProceduralVegetationEditor/Source/ProceduralVegetation/Private/ProceduralVegetationPreset.cpp
--- a/Engine/Plugins/Experimental/ProceduralVegetationEditor/Source/ProceduralVegetation/Private/ProceduralVegetationPreset.cpp
+++ b/Engine/Plugins/Experimental/ProceduralVegetationEditor/Source/ProceduralVegetation/Private/ProceduralVegetationPreset.cpp
@@ -7,17 +7,19 @@
 #include "AssetRegistry/IAssetRegistry.h"
 #include "Facades/PVBranchFacade.h"
 #include "Facades/PVProfileFacade.h"
+#include "Helpers/PVFoliageJSONHelper.h"
 #include "Helpers/PVJSONHelper.h"
 #include "Helpers/PVUtilities.h"
+#include "Implementations/PVFoliage.h"
 
 void FPVPresetVariationInfo::Fill(FName InName, FManagedArrayCollection Collection)
 {
 	Name = InName;
 	
 	const PV::Facades::FFoliageFacade Facade(Collection);
-	for (int32 FoliageIndex = 0; FoliageIndex < Facade.NumFoliageNames(); ++FoliageIndex)
+	for (int32 FoliageIndex = 0; FoliageIndex < Facade.NumFoliageInfo(); ++FoliageIndex)
 	{
-		FoliageMeshes.Emplace(FSoftObjectPath(Facade.GetFoliageName(FoliageIndex)));
+		FoliageMeshes.Emplace(Facade.GetFoliageInfo(FoliageIndex).Mesh);
 	}
 
 	const PV::Facades::FBranchFacade BranchFacade(Collection);
@@ -71,20 +73,43 @@ void UProceduralVegetationPreset::UpdateDataAsset()
 		UE_LOG(LogProceduralVegetation, Warning, TEXT("Please specify a JSON path"));
 		return;
 	}
-	
+	FString OutError;
 	IFileManager& FileManager = IFileManager::Get();
 	
 	TArray<FString> FileNames;
 	FileManager.FindFiles(FileNames, *JsonDirectoryPath.Path, TEXT("*.json"));
 
 	Variants.Empty();
+
+	FString JSONExtension = ".json";
+	FString FoliageDataFileName = "FoliageData" + JSONExtension;
+
+	if (!FileNames.Contains(FoliageDataFileName))
+	{
+		UE_LOG(LogProceduralVegetation, Warning, TEXT("Unable to find Foliage Data file"));
+		return;
+	}
+	else
+	{
+		FileNames.Remove(FoliageDataFileName);
+	}
+
+	FString FoliageDataPath = JsonDirectoryPath.Path / FoliageDataFileName;
+
+	PV::FoliageVariationsMap FoliageVariationsMap;
+
+	if (!PV::PVFoliageJSONHelper::LoadFoliageData(FoliageDataPath, FoliageFolder.Path, FoliageVariationsMap, OutError))
+	{
+		UE_LOG(LogProceduralVegetation, Warning, TEXT("Unable to load Foliage Data :  Error : %s"), *OutError);
+	}
 	
 	TArray<TSharedPtr<FJsonObject>> MetaFiles;
 	for (const auto& FileName : FileNames)
 	{
 		FString FullPath = JsonDirectoryPath.Path / FileName;
 		FManagedArrayCollection Collection;
-		FString OutError;
+		
+
 		if ( PV::LoadMegaPlantsJsonToCollection(Collection, FullPath, OutError))
 		{
 			UE_LOG(LogProceduralVegetation, Log, TEXT("Variant %s loaded from file"), *FPaths::GetBaseFilename(FullPath));
@@ -111,6 +136,17 @@ void UProceduralVegetationPreset::UpdateDataAsset()
 		{
 			PV::LoadMetaJsonToCollection(Pair.Value, MetaFiles[0]);
 		}
+
+		// load Foliage Data again
+		if (FoliageVariationsMap.Contains(Pair.Key))
+		{
+			FString VariationPresetPath = JsonDirectoryPath.Path / Pair.Key + JSONExtension;
+			
+			if (!PV::PVFoliageJSONHelper::LoadFoliageDataInCollection(Pair.Value, VariationPresetPath, FoliageVariationsMap[Pair.Key], OutError))
+			{
+				UE_LOG(LogProceduralVegetation, Warning, TEXT("Unable to load foliage data in collection:  Error : %s"), *OutError);
+			}
+		}
 		
 		VariantIndex++;
 	}
@@ -197,13 +233,8 @@ void UProceduralVegetationPreset::UpdateFoliageAndMaterialPath()
 		PV::Facades::FFoliageFacade FoliageFacade(Pair.Value);
 		PV::Facades::FBranchFacade BranchFacade(Pair.Value);
 			
-		for (int32 FoliageIndex = 0; FoliageIndex < FoliageFacade.NumFoliageNames(); FoliageIndex++)
+		for (int32 FoliageIndex = 0; FoliageIndex < FoliageFacade.NumFoliageInfo(); FoliageIndex++)
 		{
-			FString ShortName = FoliageFacade.GetFoliageName(FoliageIndex);
-			//FString ObjectName = FString::Format(TEXT("{0}.{1}"),{ ShortName, ShortName});
-			FString FullObjectName = FPaths::Combine(*FoliageFolder.Path, *ShortName);
-			FoliageFacade.SetFoliageName(FoliageIndex, FullObjectName);
-
 			if (!TrunkMaterialName.IsEmpty())
 			{
 				FString FullMaterialName = FPaths::Combine(*MaterialsFolder.Path, TrunkMaterialName + '.' + TrunkMaterialName);
diff --git a/Engine/Plugins/Experimental/ProceduralVegetationEditor/Source/ProceduralVegetation/Public/Facades/PVFoliageFacade.h b/Engine/Plugins/Experimental/ProceduralVegetationEditor/Source/ProceduralVegetation/Public/Facades/PVFoliageFacade.h
--- a/Engine/Plugins/Experimental/ProceduralVegetationEditor/Source/ProceduralVegetation/Public/Facades/PVFoliageFacade.h
+++ b/Engine/Plugins/Experimental/ProceduralVegetationEditor/Source/ProceduralVegetation/Public/Facades/PVFoliageFacade.h
@@ -6,6 +6,8 @@
 #include "IShrinkableFacade.h"
 #include "GeometryCollection/ManagedArrayAccessor.h"
 
+struct FPVFoliageInfo;
+
 namespace PV::Facades
 {
 	struct PROCEDURALVEGETATION_API FFoliageEntryData
@@ -17,6 +19,11 @@ namespace PV::Facades
 		FVector3f NormalVector;
 		float Scale;
 		float LengthFromRoot;
+		float ConditionUpAlignment;
+		float ConditionTip;
+		float ConditionLight;
+		float ConditionScale;
+		float ConditionHealth;
 		int32 ParentBoneID = INDEX_NONE;
 	};
 
@@ -58,12 +65,18 @@ namespace PV::Facades
 
 		void SetFoliageBranchId(const int32 Index, const int32 InputId);
 
-		int32 NumFoliageNames() const { return NamesAttribute.Num(); };
+		int32 NumFoliageNames() const { return FoliageInfoNamesAttribute.Num(); };
 
 		FString GetFoliageName(const int32 Index) const;
 
 		TArray<FString> GetFoliageNames() const;
 
+		int32 NumFoliageInfo() const { return FoliageInfoNamesAttribute.Num(); };
+
+		FPVFoliageInfo GetFoliageInfo(const int32 Index) const;
+
+		TArray<FPVFoliageInfo> GetFoliageInfos() const;
+
 		const FVector3f& GetPivotPoint(const int32 Index) const;
 
 		const FVector3f& GetUpVector(const int32 Index) const;
@@ -82,6 +95,10 @@ namespace PV::Facades
 
 		void SetFoliageNames(const TArray<FString>& InputNames);
 
+		void SetFoliageInfo(const int32 Index, const FString& Input);
+
+		void SetFoliageInfos(const TArray<FPVFoliageInfo>& InputNames);
+
 		void SetFoliagePath(const FString InPath);
 
 		FString GetFoliagePath() const;
@@ -102,7 +119,13 @@ namespace PV::Facades
 		const FManagedArrayCollection& ConstCollection;
 		FManagedArrayCollection* Collection = nullptr;
 
-		TManagedArrayAccessor<FString> NamesAttribute;
+		TManagedArrayAccessor<FString> FoliageInfoNamesAttribute;
+		TManagedArrayAccessor<float> FoliageInfoLightAttribute;
+		TManagedArrayAccessor<float> FoliageInfoScaleAttribute;
+		TManagedArrayAccessor<float> FoliageInfoTipAttribute;
+		TManagedArrayAccessor<float> FoliageInfoUpAlignmentAttribute;
+		TManagedArrayAccessor<float> FoliageInfoHealthAttribute;
+		
 		TManagedArrayAccessor<int32> NameIdsAttribute;
 		TManagedArrayAccessor<int32> BranchIdsAttribute;
 		TManagedArrayAccessor<FVector3f> PivotPointsAttribute;
@@ -110,6 +133,11 @@ namespace PV::Facades
 		TManagedArrayAccessor<FVector3f> NormalVectorsAttribute;
 		TManagedArrayAccessor<float> ScalesAttribute;
 		TManagedArrayAccessor<float> LengthFromRootAttribute;
+		TManagedArrayAccessor<float> ConditionLightAttribute;
+		TManagedArrayAccessor<float> ConditionScaleAttribute;
+		TManagedArrayAccessor<float> ConditionTipAttribute;
+		TManagedArrayAccessor<float> ConditionUpAlignmentAttribute;
+		TManagedArrayAccessor<float> ConditionHealthAttribute;
 		TManagedArrayAccessor<int32> ParentBoneIdsAttribute;
 		TManagedArrayAccessor<TArray<int32>> FoliageIdsAttribute;
 		TManagedArrayAccessor<FString> FoliagePathAttribute;
diff --git a/Engine/Plugins/Experimental/ProceduralVegetationEditor/Source/ProceduralVegetationEditor/Private/Helpers/PVExportHelper.cpp b/Engine/Plugins/Experimental/ProceduralVegetationEditor/Source/ProceduralVegetationEditor/Private/Helpers/PVExportHelper.cpp
--- a/Engine/Plugins/Experimental/ProceduralVegetationEditor/Source/ProceduralVegetationEditor/Private/Helpers/PVExportHelper.cpp
+++ b/Engine/Plugins/Experimental/ProceduralVegetationEditor/Source/ProceduralVegetationEditor/Private/Helpers/PVExportHelper.cpp
@@ -54,6 +54,7 @@
 #include "PhysicsEngine/PhysicsAsset.h"
 #include "PhysicsEngine/PhysicsConstraintTemplate.h"
 #include "PhysicsEngine/SkeletalBodySetup.h"
+#include "ProceduralVegetation/Private/Implementations/PVFoliage.h"
 
 #include "UObject/Package.h"
 
@@ -317,7 +318,10 @@ namespace PV::Export::Internal
 	FString GetFoliageName(const Facades::FFoliageFacade& FoliageFacade, int32 Id)
 	{
 		Facades::FFoliageEntryData Data = FoliageFacade.GetFoliageEntry(Id);
-		return FoliageFacade.GetFoliageName(Data.NameId);
+		
+		FPVFoliageInfo Info = FoliageFacade.GetFoliageInfo(Data.NameId);
+		
+		return Info.Mesh.ToString();
 	}
 
 	TArray<TObjectPtr<UMaterialInterface>> GetMaterialsFromCollection(const FManagedArrayCollection& Collection)
@@ -368,7 +372,7 @@ namespace PV::Export::Internal
 
 		const Facades::FFoliageFacade FoliageFacade(Collection);
 		const int32 NumInstances = FoliageFacade.NumFoliageEntries();
-		const int32 NumFoliageNames = FoliageFacade.NumFoliageNames();
+		const int32 NumFoliageNames = FoliageFacade.NumFoliageInfo();
 		
 		if (NumInstances == 0 || NumFoliageNames == 0)
 		{
@@ -394,8 +398,9 @@ namespace PV::Export::Internal
 
 			for (int32 FoliageNameIndex = 0; FoliageNameIndex < NumFoliageNames; ++FoliageNameIndex)
 			{
-				const FString FoliageName = FoliageFacade.GetFoliageName(FoliageNameIndex);
-
+				FPVFoliageInfo Info = FoliageFacade.GetFoliageInfo(FoliageNameIndex);
+				const FString FoliageName = Info.Mesh.ToString();
+				
 				if (FoliageNameToDynamicMeshResult.Contains(FoliageName))
 				{
 					UE_LOG(LogProceduralVegetationEditor, Warning, TEXT("Duplicate foliage name %s in foliage facade"), *FoliageName);
@@ -603,11 +608,13 @@ namespace PV::Export::Internal
 	TArray<FString> GetUsedFoliage(const Facades::FFoliageFacade& FoliageFacade)
 	{
 		TArray<FString> OutFoliageNames;
-		OutFoliageNames.Reserve(FoliageFacade.NumFoliageNames());
+		OutFoliageNames.Reserve(FoliageFacade.NumFoliageInfo());
 
-		for(int32 FoliageNameIndex = 0; FoliageNameIndex < FoliageFacade.NumFoliageNames(); ++FoliageNameIndex)
+		for(int32 FoliageNameIndex = 0; FoliageNameIndex < FoliageFacade.NumFoliageInfo(); ++FoliageNameIndex)
 		{
-			FString FoliageName = FoliageFacade.GetFoliageName(FoliageNameIndex);
+			FPVFoliageInfo Info = FoliageFacade.GetFoliageInfo(FoliageNameIndex);
+			FString FoliageName = Info.Mesh.ToString();
+			
 			if (ContainsFoliageName(FoliageFacade, FoliageName))
 			{
 				OutFoliageNames.Add(FoliageName);
diff --git a/Engine/Plugins/Experimental/ProceduralVegetationEditor/Source/ProceduralVegetationEditor/Private/Visualizations/PVDataVisualization.cpp b/Engine/Plugins/Experimental/ProceduralVegetationEditor/Source/ProceduralVegetationEditor/Private/Visualizations/PVDataVisualization.cpp
--- a/Engine/Plugins/Experimental/ProceduralVegetationEditor/Source/ProceduralVegetationEditor/Private/Visualizations/PVDataVisualization.cpp
+++ b/Engine/Plugins/Experimental/ProceduralVegetationEditor/Source/ProceduralVegetationEditor/Private/Visualizations/PVDataVisualization.cpp
@@ -27,6 +27,7 @@
 #include "Components/StaticMeshComponent.h"
 #include "Components/SkeletalMeshComponent.h"
 #include "Materials/Material.h"
+#include "ProceduralVegetation/Private/Implementations/PVFoliage.h"
 
 #define LOCTEXT_NAMESPACE "PCGManagedArrayCollectionDataVisualization"
 
@@ -250,7 +251,7 @@ void FPVDataVisualization::FoliageGridRenderer(const FManagedArrayCollection& In
 {
 	const PV::Facades::FFoliageFacade FoliageFacade(InCollection);
 
-	const TArray<FString> FoliageNames = FoliageFacade.GetFoliageNames();
+	const TArray<FPVFoliageInfo> FoliageInfos = FoliageFacade.GetFoliageInfos();
 
 	FBox MaxFoliageBounds(ForceInit);
 	int ValidFoliageEntries = 0;
@@ -275,8 +276,10 @@ void FPVDataVisualization::FoliageGridRenderer(const FManagedArrayCollection& In
 			}
 		};
 
-	for (const FString& FoliageName : FoliageNames)
+	for (const FPVFoliageInfo& Info : FoliageInfos)
 	{
+		FString FoliageName = Info.Mesh.ToString();
+		
 		if (UObject* FoliageMesh = LoadObject<UObject>(nullptr, *FoliageName))
 		{
 			FBox Bounds = GetMeshBounds(Cast<UStaticMesh>(FoliageMesh));
@@ -296,10 +299,12 @@ void FPVDataVisualization::FoliageGridRenderer(const FManagedArrayCollection& In
 	const float FoliageHorizontalPadding = MaxFoliageBounds.GetSize().GetMax() + 100.0f;
 	const float FoliageVerticalPadding = MaxFoliageBounds.GetSize().GetMax() + 100.0f;
 
-	FVector FoliagePos = FVector(FoliageHorizontalPadding * FoliageNames.Num() / -2.0f, 0.0f, 0.0f);
+	FVector FoliagePos = FVector(FoliageHorizontalPadding * FoliageInfos.Num() / -2.0f, 0.0f, 0.0f);
 
-	for (const FString& FoliagePath : FoliageNames)
+	for (const FPVFoliageInfo& Info : FoliageInfos)
 	{
+		const FString& FoliagePath = Info.Mesh.ToString();
+		
 		UObject* FoliageObject = LoadObject<UObject>(nullptr, FoliagePath);
 		if (UStaticMesh* StaticFoliageMesh = Cast<UStaticMesh>(FoliageObject))
 		{
