[PVE] Fix potential bug where the UpVector is parallel to the forward vector

diff --git a/Engine/Plugins/Experimental/ProceduralVegetationEditor/Source/ProceduralVegetation/Private/Facades/PVFoliageFacade.cpp b/Engine/Plugins/Experimental/ProceduralVegetationEditor/Source/ProceduralVegetation/Private/Facades/PVFoliageFacade.cpp
--- a/Engine/Plugins/Experimental/ProceduralVegetationEditor/Source/ProceduralVegetation/Private/Facades/PVFoliageFacade.cpp
+++ b/Engine/Plugins/Experimental/ProceduralVegetationEditor/Source/ProceduralVegetation/Private/Facades/PVFoliageFacade.cpp
@@ -416,7 +416,8 @@ namespace PV::Facades
 			(NormalVectorN - FVector::DotProduct(NormalVectorN, UpVectorN) * UpVectorN).GetSafeNormal();
 		if (OrthogonalizedNormalVectorN.IsNearlyZero())
 		{
-			OrthogonalizedNormalVectorN = FVector::CrossProduct(UpVectorN, FVector::ForwardVector).GetSafeNormal();
+			FVector TempVector;
+			UpVectorN.FindBestAxisVectors(TempVector, OrthogonalizedNormalVectorN);
 		}
 
 		const FQuat RotationQuat = FQuat(FRotationMatrix::MakeFromYZ(OrthogonalizedNormalVectorN, UpVectorN));
