[PVE] Multiple Fixes
- Foliage attachment points visualization
- Foliage visibility toggle
- Fixed viewport stats not updating on graph re execution

#virtualized

diff --git a/Engine/Plugins/Experimental/ProceduralVegetationEditor/Source/ProceduralVegetation/Private/Helpers/PVUtilities.cpp b/Engine/Plugins/Experimental/ProceduralVegetationEditor/Source/ProceduralVegetation/Private/Helpers/PVUtilities.cpp
--- a/Engine/Plugins/Experimental/ProceduralVegetationEditor/Source/ProceduralVegetation/Private/Helpers/PVUtilities.cpp
+++ b/Engine/Plugins/Experimental/ProceduralVegetationEditor/Source/ProceduralVegetation/Private/Helpers/PVUtilities.cpp
@@ -77,7 +77,7 @@ namespace PV::Utilities
 
 				if(StaticFoliageMesh)
 				{
-					UInstancedStaticMeshComponent* InstancedComponent = NewObject<UInstancedStaticMeshComponent>(InParent != nullptr ? InParent->GetPackage() : GetTransientPackage(), NAME_None);
+					UInstancedStaticMeshComponent* InstancedComponent = NewObject<UInstancedStaticMeshComponent>(InParent != nullptr ? InParent->GetPackage() : GetTransientPackage(), StaticFoliageMesh->GetFName());
 					InstancedComponent->SetStaticMesh(StaticFoliageMesh);
 					OutInstancedComponentMap.Add(FoliageName, InstancedComponent);
 					InCallback.ExecuteIfBound(InstancedComponent);
@@ -96,7 +96,7 @@ namespace PV::Utilities
 					UE_LOG(LogProceduralVegetation, Warning, TEXT("Instanced Skinned Mesh component only works with nanite meshes make sure nanite is enabled for mesh {%s}"), *FoliageName);
 #endif
 					
-					UInstancedSkinnedMeshComponent* InstancedComponent = NewObject<UInstancedSkinnedMeshComponent>(InParent != nullptr ? InParent->GetPackage() : GetTransientPackage(), NAME_None);
+					UInstancedSkinnedMeshComponent* InstancedComponent = NewObject<UInstancedSkinnedMeshComponent>(InParent != nullptr ? InParent->GetPackage() : GetTransientPackage(), SkeletalFoliageMesh->GetFName());
 					InstancedComponent->SetSkinnedAssetAndUpdate(SkeletalFoliageMesh);
 					OutInstancedComponentMap.Add(FoliageName, InstancedComponent);
 					InCallback.ExecuteIfBound(InstancedComponent);
diff --git a/Engine/Plugins/Experimental/ProceduralVegetationEditor/Source/ProceduralVegetation/Private/Nodes/PVBoneReductionSettings.h b/Engine/Plugins/Experimental/ProceduralVegetationEditor/Source/ProceduralVegetation/Private/Nodes/PVBoneReductionSettings.h
--- a/Engine/Plugins/Experimental/ProceduralVegetationEditor/Source/ProceduralVegetation/Private/Nodes/PVBoneReductionSettings.h
+++ b/Engine/Plugins/Experimental/ProceduralVegetationEditor/Source/ProceduralVegetation/Private/Nodes/PVBoneReductionSettings.h
@@ -20,6 +20,7 @@ public:
 	virtual FLinearColor GetNodeTitleColor() const override { return FLinearColor(1.0f, 0.0f, 0.599f); }
 
 	virtual TArray<EPVRenderType> GetDefaultRenderType() const override { return TArray{EPVRenderType::Bones, EPVRenderType::Mesh}; }
+	virtual TArray<EPVRenderType> GetSupportedRenderTypes() const override { return TArray{ EPVRenderType::Bones, EPVRenderType::Mesh, EPVRenderType::FoliageAttachments }; }
 #endif
 
 protected:
diff --git a/Engine/Plugins/Experimental/ProceduralVegetationEditor/Source/ProceduralVegetation/Private/Nodes/PVFoliageDistributorSettings.h b/Engine/Plugins/Experimental/ProceduralVegetationEditor/Source/ProceduralVegetation/Private/Nodes/PVFoliageDistributorSettings.h
--- a/Engine/Plugins/Experimental/ProceduralVegetationEditor/Source/ProceduralVegetation/Private/Nodes/PVFoliageDistributorSettings.h
+++ b/Engine/Plugins/Experimental/ProceduralVegetationEditor/Source/ProceduralVegetation/Private/Nodes/PVFoliageDistributorSettings.h
@@ -29,6 +29,7 @@ public:
 	virtual FLinearColor GetNodeTitleColor() const override { return FLinearColor::Blue; }
 
 	virtual TArray<EPVRenderType> GetDefaultRenderType() const override { return TArray{ EPVRenderType::Mesh, EPVRenderType::Foliage }; }
+	virtual TArray<EPVRenderType> GetSupportedRenderTypes() const override { return TArray{ EPVRenderType::Mesh, EPVRenderType::Foliage, EPVRenderType::FoliageAttachments }; }
 	#endif
 
 protected:
diff --git a/Engine/Plugins/Experimental/ProceduralVegetationEditor/Source/ProceduralVegetation/Private/Nodes/PVFoliagePaletteSettings.h b/Engine/Plugins/Experimental/ProceduralVegetationEditor/Source/ProceduralVegetation/Private/Nodes/PVFoliagePaletteSettings.h
--- a/Engine/Plugins/Experimental/ProceduralVegetationEditor/Source/ProceduralVegetation/Private/Nodes/PVFoliagePaletteSettings.h
+++ b/Engine/Plugins/Experimental/ProceduralVegetationEditor/Source/ProceduralVegetation/Private/Nodes/PVFoliagePaletteSettings.h
@@ -19,7 +19,7 @@ public:
 	virtual FLinearColor GetNodeTitleColor() const override { return FLinearColor::Green; }
 
 	virtual TArray<EPVRenderType> GetDefaultRenderType() const override { return TArray{EPVRenderType::FoliageGrid}; }
-	virtual TArray<EPVRenderType> GetSupportedRenderTypes() const override { return TArray{EPVRenderType::FoliageGrid, EPVRenderType::Foliage, EPVRenderType::Mesh}; }
+	virtual TArray<EPVRenderType> GetSupportedRenderTypes() const override { return TArray{EPVRenderType::FoliageGrid, EPVRenderType::Foliage, EPVRenderType::FoliageAttachments, EPVRenderType::Mesh}; }
 	
 	virtual void PostEditChangeProperty(FPropertyChangedEvent& PropertyChangedEvent) override;
 	virtual void PostEditChangeChainProperty(FPropertyChangedChainEvent& PropertyChangedEvent) override;
diff --git a/Engine/Plugins/Experimental/ProceduralVegetationEditor/Source/ProceduralVegetation/Private/Nodes/PVMeshBuilderSettings.h b/Engine/Plugins/Experimental/ProceduralVegetationEditor/Source/ProceduralVegetation/Private/Nodes/PVMeshBuilderSettings.h
--- a/Engine/Plugins/Experimental/ProceduralVegetationEditor/Source/ProceduralVegetation/Private/Nodes/PVMeshBuilderSettings.h
+++ b/Engine/Plugins/Experimental/ProceduralVegetationEditor/Source/ProceduralVegetation/Private/Nodes/PVMeshBuilderSettings.h
@@ -20,8 +20,8 @@ public:
 	virtual FText GetDefaultNodeTitle() const override;
 	virtual FText GetNodeTooltipText() const override;
 	virtual FLinearColor GetNodeTitleColor() const override { return FLinearColor(FColor::Cyan); }
-	virtual TArray<EPVRenderType> GetDefaultRenderType() const override { return TArray{EPVRenderType::Mesh}; }
-	virtual TArray<EPVRenderType> GetSupportedRenderTypes() const override { return TArray{EPVRenderType::PointData, EPVRenderType::Mesh}; }
+	virtual TArray<EPVRenderType> GetDefaultRenderType() const override { return TArray{ EPVRenderType::Mesh }; }
+	virtual TArray<EPVRenderType> GetSupportedRenderTypes() const override { return TArray{ EPVRenderType::PointData, EPVRenderType::Mesh, EPVRenderType::FoliageAttachments }; }
 
 	virtual void PostEditChangeProperty(FPropertyChangedEvent& PropertyChangedEvent) override;
 	virtual void PostEditChangeChainProperty(FPropertyChangedChainEvent& PropertyChangedEvent) override;
diff --git a/Engine/Plugins/Experimental/ProceduralVegetationEditor/Source/ProceduralVegetation/Public/DataTypes/PVData.h b/Engine/Plugins/Experimental/ProceduralVegetationEditor/Source/ProceduralVegetation/Public/DataTypes/PVData.h
--- a/Engine/Plugins/Experimental/ProceduralVegetationEditor/Source/ProceduralVegetation/Public/DataTypes/PVData.h
+++ b/Engine/Plugins/Experimental/ProceduralVegetationEditor/Source/ProceduralVegetation/Public/DataTypes/PVData.h
@@ -13,12 +13,13 @@ struct FPCGContext;
 UENUM()
 enum class EPVRenderType : uint8
 {
-	None 		= 0,
-	PointData 	= 1 << 0,
-	Mesh 		= 1 << 1,
-	Foliage 	= 1 << 2,
-	Bones 		= 1 << 3,
-	FoliageGrid = 1 << 4,
+	None 				= 0,
+	PointData 			= 1 << 0,
+	Mesh 				= 1 << 1,
+	Foliage 			= 1 << 2,
+	Bones 				= 1 << 3,
+	FoliageGrid 		= 1 << 4,
+	FoliageAttachments	= 1 << 5,
 
 	// Add new renter type above
 	Count	
diff --git a/Engine/Plugins/Experimental/ProceduralVegetationEditor/Source/ProceduralVegetation/Public/Nodes/PVOutputSettings.h b/Engine/Plugins/Experimental/ProceduralVegetationEditor/Source/ProceduralVegetation/Public/Nodes/PVOutputSettings.h
--- a/Engine/Plugins/Experimental/ProceduralVegetationEditor/Source/ProceduralVegetation/Public/Nodes/PVOutputSettings.h
+++ b/Engine/Plugins/Experimental/ProceduralVegetationEditor/Source/ProceduralVegetation/Public/Nodes/PVOutputSettings.h
@@ -19,7 +19,7 @@ public:
 	virtual FText GetDefaultNodeTitle() const override;
 	virtual FText GetNodeTooltipText() const override;
 	virtual EPCGSettingsType GetType() const override { return EPCGSettingsType::InputOutput; }
-	virtual TArray<EPVRenderType> GetSupportedRenderTypes() const override { return TArray{EPVRenderType::PointData, EPVRenderType::Foliage, EPVRenderType::FoliageGrid, EPVRenderType::Mesh, EPVRenderType::Bones}; }
+	virtual TArray<EPVRenderType> GetSupportedRenderTypes() const override { return TArray{EPVRenderType::PointData, EPVRenderType::Foliage, EPVRenderType::FoliageGrid, EPVRenderType::FoliageAttachments, EPVRenderType::Mesh, EPVRenderType::Bones}; }
 	virtual TArray<EPVRenderType> GetDefaultRenderType() const override { return TArray{EPVRenderType::Foliage, EPVRenderType::Mesh}; }
 #endif
 
diff --git a/Engine/Plugins/Experimental/ProceduralVegetationEditor/Source/ProceduralVegetationEditor/Private/Customizations/RichTextColorDecorator.cpp b/Engine/Plugins/Experimental/ProceduralVegetationEditor/Source/ProceduralVegetationEditor/Private/Customizations/RichTextColorDecorator.cpp
new file mode 100644
--- /dev/null
+++ b/Engine/Plugins/Experimental/ProceduralVegetationEditor/Source/ProceduralVegetationEditor/Private/Customizations/RichTextColorDecorator.cpp
@@ -0,0 +1,143 @@
+// Copyright Epic Games, Inc. All Rights Reserved.
+
+#include "RichTextColorDecorator.h"
+
+#include "Fonts/FontMeasure.h"
+#include "Fonts/SlateFontInfo.h"
+
+#include "Framework/Application/SlateApplication.h"
+#include "Framework/Text/SlateTextRun.h"
+
+#include "Internationalization/Text.h"
+
+#include "Styling/CoreStyle.h"
+#include "Styling/ISlateStyle.h"
+#include "Styling/SlateColor.h"
+
+TOptional<FLinearColor> ParseNamed(const FString& Name)
+{
+	if (Name == TEXT("red")) return FLinearColor::Red;
+	if (Name == TEXT("green")) return FLinearColor::Green;
+	if (Name == TEXT("blue")) return FLinearColor::Blue;
+	if (Name == TEXT("yellow")) return FLinearColor::Yellow;
+	if (Name == TEXT("white")) return FLinearColor::White;
+	if (Name == TEXT("black")) return FLinearColor::Black;
+	if (Name == TEXT("orange")) return FLinearColor(1.f, 0.5f, 0.f, 1.f);
+	if (Name == TEXT("purple")) return FLinearColor(0.5f, 0.f, 0.5f, 1.f);
+
+	return {};
+}
+
+bool GetAttr(const TMap<FString, FTextRange>& Meta, const FString& Key, FString& Out,
+             const FString& Source)
+{
+	if (const FTextRange* Range = Meta.Find(Key))
+	{
+		Out = Source.Mid(Range->BeginIndex, Range->Len());
+		return true;
+	}
+	return false;
+}
+
+
+bool GetAttr(const TMap<FString, FTextRange>& Meta, const FString& Key, int32& Out,
+             const FString& Source)
+{
+	if (FString Value; GetAttr(Meta, Key, Value, Source))
+	{
+		Out = FCString::Atoi(*Value);
+		return true;
+	}
+	return false;
+}
+
+FColorBlockTextRun::FColorBlockTextRun(
+	const FRunInfo& InRunInfo,
+	const TSharedRef<const FString>& InText,
+	const FTextBlockStyle& InStyle,
+	const FTextRange& InRange,
+	const FLinearColor& InBg
+)
+	: FSlateTextRun(InRunInfo, InText, InStyle, InRange)
+	, BgColor(InBg)
+{}
+
+TSharedRef<FColorBlockTextRun> FColorBlockTextRun::Create(
+	const FRunInfo& InRunInfo,
+	const TSharedRef<const FString>& InText,
+	const FTextBlockStyle& InStyle,
+	const FTextRange& InRange,
+	const FLinearColor& InBg
+)
+{
+	return MakeShared<FColorBlockTextRun>(InRunInfo, InText, InStyle, InRange, InBg);
+}
+
+int32 FColorBlockTextRun::OnPaint(const FPaintArgs& Args, const FTextArgs& TextArgs, const FGeometry& Allotted, const FSlateRect& Culling,
+                                  FSlateWindowElementList& Out, int32 LayerId, const FWidgetStyle& InWidgetStyle, bool bParentEnabled) const
+{
+	// Measure the exact width of this run’s substring (the spaces inside <block>...</>)
+	const FString RunText = Text->Mid(Range.BeginIndex, Range.Len());
+
+	// Use the font set on this run (from your TextStyle)
+	const FSlateFontMeasure& FM = *FSlateApplication::Get().GetRenderer()->GetFontMeasureService();
+	const FVector2D Size(
+		FM.Measure(RunText, Style.Font).X,
+		FM.GetMaxCharacterHeight(Style.Font)
+	);
+
+	// The base class will draw the text at the proper offset; paint the bg at the same origin
+	const FSlateBrush* Brush = FCoreStyle::Get().GetBrush("WhiteBrush");
+	FSlateDrawElement::MakeBox(
+		Out, LayerId,
+		Allotted.ToPaintGeometry(Size, FSlateLayoutTransform()),
+		Brush, ESlateDrawEffect::None, BgColor
+	);
+
+	return FSlateTextRun::OnPaint(Args, TextArgs, Allotted, Culling, Out, LayerId + 1, InWidgetStyle, bParentEnabled);
+}
+
+TSharedRef<FColorBlockDecorator> FColorBlockDecorator::CreateDecorator()
+{
+	return MakeShared<FColorBlockDecorator>();
+}
+
+bool FColorBlockDecorator::Supports(const FTextRunParseResults& RunParseResults, const FString& Text) const
+{
+	return RunParseResults.Name.Equals(TEXT("block"), ESearchCase::IgnoreCase);
+}
+
+TSharedRef<ISlateRun> FColorBlockDecorator::Create(
+	const TSharedRef<FTextLayout>& TextLayout,
+	const FTextRunParseResults& RunInfo,
+	const FString& OriginalText,
+	const TSharedRef<FString>& ModelText,
+	const ISlateStyle* Style
+)
+{
+	FLinearColor Color = FLinearColor::White;
+	if (FString HexAttr; GetAttr(RunInfo.MetaData, TEXT("hex"), HexAttr, OriginalText))
+	{
+		Color = FLinearColor::FromSRGBColor(FColor::FromHex(HexAttr));
+	}
+	else if (FString NameAttr; GetAttr(RunInfo.MetaData, TEXT("name"), NameAttr, OriginalText))
+	{
+		if (TOptional<FLinearColor> Parsed = ParseNamed(NameAttr))
+		{
+			Color = Parsed.GetValue();
+		}
+	}
+
+	FTextBlockStyle TextStyle = Style
+		? Style->GetWidgetStyle<FTextBlockStyle>("NormalText")
+		: FTextBlockStyle(FCoreStyle::Get().GetWidgetStyle<FTextBlockStyle>("NormalText"));
+	TextStyle.SetColorAndOpacity(Color);
+
+	FTextRange ModelRange;
+	ModelRange.BeginIndex = ModelText->Len();
+	ModelText->Append(OriginalText.Mid(RunInfo.ContentRange.BeginIndex, RunInfo.ContentRange.Len()));
+	ModelRange.EndIndex = ModelText->Len();
+
+	const FRunInfo Info(RunInfo.Name);
+	return FColorBlockTextRun::Create(Info, ModelText, TextStyle, ModelRange, Color);
+}
diff --git a/Engine/Plugins/Experimental/ProceduralVegetationEditor/Source/ProceduralVegetationEditor/Private/Customizations/RichTextColorDecorator.h b/Engine/Plugins/Experimental/ProceduralVegetationEditor/Source/ProceduralVegetationEditor/Private/Customizations/RichTextColorDecorator.h
new file mode 100644
--- /dev/null
+++ b/Engine/Plugins/Experimental/ProceduralVegetationEditor/Source/ProceduralVegetationEditor/Private/Customizations/RichTextColorDecorator.h
@@ -0,0 +1,58 @@
+// Copyright Epic Games, Inc. All Rights Reserved.
+
+#pragma once
+
+#include "CoreMinimal.h"
+
+#include "Framework/Text/ITextDecorator.h"
+#include "Framework/Text/SlateTextRun.h"
+
+class FColorBlockTextRun : public FSlateTextRun
+{
+public:
+	FColorBlockTextRun(
+		const FRunInfo& InRunInfo,
+		const TSharedRef<const FString>& InText,
+		const FTextBlockStyle& InStyle,
+		const FTextRange& InRange,
+		const FLinearColor& InBg
+	);
+
+	static TSharedRef<FColorBlockTextRun> Create(
+		const FRunInfo& InRunInfo,
+		const TSharedRef<const FString>& InText,
+		const FTextBlockStyle& InStyle,
+		const FTextRange& InRange,
+		const FLinearColor& InBg
+	);
+
+	virtual int32 OnPaint(
+		const FPaintArgs& Args,
+		const FTextArgs& TextArgs,
+		const FGeometry& Allotted,
+		const FSlateRect& Culling,
+		FSlateWindowElementList& Out,
+		int32 LayerId,
+		const FWidgetStyle& InWidgetStyle,
+		bool bParentEnabled
+	) const override;
+
+private:
+	FLinearColor BgColor;
+};
+
+class FColorBlockDecorator : public ITextDecorator
+{
+public:
+	static TSharedRef<FColorBlockDecorator> CreateDecorator();
+
+	virtual bool Supports(const FTextRunParseResults& RunParseResults, const FString& Text) const override;
+
+	virtual TSharedRef<ISlateRun> Create(
+		const TSharedRef<FTextLayout>& TextLayout,
+		const FTextRunParseResults& RunInfo,
+		const FString& OriginalText,
+		const TSharedRef<FString>& ModelText,
+		const ISlateStyle* Style
+	) override;
+};
diff --git a/Engine/Plugins/Experimental/ProceduralVegetationEditor/Source/ProceduralVegetationEditor/Private/PVEditor.cpp b/Engine/Plugins/Experimental/ProceduralVegetationEditor/Source/ProceduralVegetationEditor/Private/PVEditor.cpp
--- a/Engine/Plugins/Experimental/ProceduralVegetationEditor/Source/ProceduralVegetationEditor/Private/PVEditor.cpp
+++ b/Engine/Plugins/Experimental/ProceduralVegetationEditor/Source/ProceduralVegetationEditor/Private/PVEditor.cpp
@@ -80,6 +80,8 @@ void FPVEditor::Initialize(const EToolkitMode::Type InMode, const TSharedPtr<cla
 	UPCGEngineSubsystem* Subsystem = GEngine->GetEngineSubsystem<UPCGEngineSubsystem>();
 
 	check(Subsystem);
+
+	Subsystem->GetOnPCGSourceGenerationDone().AddRaw(this, &FPVEditor::OnGraphExecuted);
 	
 	const FPCGDefaultExecutionSourceParams Params
 	{
@@ -319,7 +321,6 @@ void FPVEditor::ChangeNodeInspection(UPCGEditorGraphNodeBase* InNode)
 		ViewportWidgets[0]->OnNodeInspectionChanged(Settings);
 		
 		NodeBeingInspected = InNode;
-		CollectionBeingInspected = SelectedCollection;
 		
 		UpdateCollectionSpreadSheet(InNode->GetPCGNode());
 		
@@ -347,25 +348,34 @@ void FPVEditor::OnSelectedNodesChanged(const TSet<UObject*>& InNewSelection)
 
 	if (UPCGNode* SelectedNode = GetSelectedNode() ? GetSelectedNode()->GetPCGNode() : nullptr)
 	{
-		SelectedCollection = MakeShared<FManagedArrayCollection>();
-		FName OutputName;
-		GetCollectionFromNode(SelectedNode, OutputName, *SelectedCollection);
-		
 		if (UPVBaseSettings* Settings = Cast<UPVBaseSettings>(SelectedNode->GetSettings()))
 		{
 			const UPVBaseSettings* InspectionSettings = GetNodeSettings(NodeBeingInspected);
 			if (!InspectionSettings || !InspectionSettings->bLockInspection)
 			{
 				ChangeNodeInspection(EdNode);
-
-				TArray<FText> OverlayStats;
-				GatherStats(OverlayStats);
-				ViewportWidgets[0]->PopulateStatsOverlayText(OverlayStats);
 			}
 		}	
 	}
 }
 
+void FPVEditor::OnGraphExecuted(IPCGBaseSubsystem* InSubsystem, IPCGGraphExecutionSource* InExecutionSource,
+	EPCGGenerationStatus InGenerationStatus)
+{
+	UpdateInspectedCollection();
+	UpdateStats();
+}
+
+void FPVEditor::UpdateInspectedCollection()
+{
+	if (const UPCGNode* InspectedNode = NodeBeingInspected ? NodeBeingInspected->GetPCGNode() : nullptr)
+	{
+		CollectionBeingInspected = MakeShared<FManagedArrayCollection>();
+		FName OutputName;
+		GetCollectionFromNode(InspectedNode, OutputName, *CollectionBeingInspected);
+	}
+}
+
 void FPVEditor::OnExport_Clicked()
 {
 	ExportInternal();
@@ -764,6 +774,26 @@ FReply FPVEditor::ExportInternal()
 	return FReply::Handled();
 }
 
+void FPVEditor::UpdateStats()
+{
+	if (ViewportWidgets.IsEmpty() || NodeBeingInspected == nullptr)
+	{
+		return;
+	}
+
+	const UPVBaseSettings* NodeBeingInspectedSettings = GetNodeSettings(NodeBeingInspected);
+
+	TArray<FText> OverlayStats;
+	GatherStats(OverlayStats);
+	ViewportWidgets[0]->PopulateStatsOverlayText(OverlayStats);
+	
+	ViewportWidgets[0]->SetOverlayText(
+		NodeBeingInspectedSettings->bLockInspection
+		? NodeBeingInspected->GetNodeTitle(ENodeTitleType::FullTitle)
+		: FText::GetEmpty()
+	);
+}
+
 void FPVEditor::OnLockNodeSelection()
 {
 	UPVBaseSettings* NodeBeingInspectedSettings = GetNodeSettings(NodeBeingInspected);
@@ -782,22 +812,8 @@ void FPVEditor::OnLockNodeSelection()
 			ChangeNodeInspection(GetSelectedNode());
 		}
 	}
-	if (ViewportWidgets.IsEmpty())
-	{
-		return;
-	}
-
-	NodeBeingInspectedSettings = GetNodeSettings(NodeBeingInspected);
-
-	TArray<FText> OverlayStats;
-	GatherStats(OverlayStats);
-	ViewportWidgets[0]->PopulateStatsOverlayText(OverlayStats);
 	
-	ViewportWidgets[0]->SetOverlayText(
-		NodeBeingInspectedSettings->bLockInspection
-		? NodeBeingInspected->GetNodeTitle(ENodeTitleType::FullTitle)
-		: FText::GetEmpty()
-	);
+	UpdateStats();
 }
 
 TSharedRef<SPCGEditorViewport> FPVEditor::CreateViewportWidget()
@@ -930,6 +946,11 @@ void FPVEditor::UpdateCollectionSpreadSheet(const UPCGNode* InSelectedNode)
 
 void FPVEditor::GatherStats(TArray<FText>& OverlayStats)
 {
+	if (CollectionBeingInspected == nullptr)
+	{
+		return;
+	}
+	
 	if (CollectionBeingInspected->HasGroup(PV::GroupNames::PointGroup))
 	{
 		const int32 NumPoints = CollectionBeingInspected->NumElements(PV::GroupNames::PointGroup);
diff --git a/Engine/Plugins/Experimental/ProceduralVegetationEditor/Source/ProceduralVegetationEditor/Private/PVEditor.h b/Engine/Plugins/Experimental/ProceduralVegetationEditor/Source/ProceduralVegetationEditor/Private/PVEditor.h
--- a/Engine/Plugins/Experimental/ProceduralVegetationEditor/Source/ProceduralVegetationEditor/Private/PVEditor.h
+++ b/Engine/Plugins/Experimental/ProceduralVegetationEditor/Source/ProceduralVegetationEditor/Private/PVEditor.h
@@ -46,7 +46,11 @@ protected:
 	virtual void OnSelectedNodesChanged(const TSet<UObject*>& InNewSelection) override;
 	// ~End FPCGEditor interface
 
-	bool bGenerationInProgress = false;
+	void OnGraphExecuted(IPCGBaseSubsystem* InSubsystem, IPCGGraphExecutionSource* InExecutionSource,
+	EPCGGenerationStatus InGenerationStatus);
+
+	void UpdateInspectedCollection();
+
 	TObjectPtr<UPCGDefaultExecutionSource> ExecutionSource;
 
 private:
@@ -56,7 +60,6 @@ private:
 
 	TArray<TObjectPtr<UPCGEditorGraphNodeBase>> SelectedNodes;
 	TSharedPtr<FManagedArrayCollection> CollectionBeingInspected = nullptr;
-	TSharedPtr<FManagedArrayCollection> SelectedCollection = nullptr;
 	TObjectPtr<UPCGEditorGraphNodeBase> NodeBeingInspected = nullptr;
 	
 	TArray<TSharedRef<SPVEditorViewport>> ViewportWidgets;
@@ -77,6 +80,8 @@ private:
 	void OnExport_Clicked();
 
 	FReply ExportInternal();
+
+	void UpdateStats();
 	
 	void OnLockNodeSelection();
 
diff --git a/Engine/Plugins/Experimental/ProceduralVegetationEditor/Source/ProceduralVegetationEditor/Private/SPVEditorViewport.cpp b/Engine/Plugins/Experimental/ProceduralVegetationEditor/Source/ProceduralVegetationEditor/Private/SPVEditorViewport.cpp
--- a/Engine/Plugins/Experimental/ProceduralVegetationEditor/Source/ProceduralVegetationEditor/Private/SPVEditorViewport.cpp
+++ b/Engine/Plugins/Experimental/ProceduralVegetationEditor/Source/ProceduralVegetationEditor/Private/SPVEditorViewport.cpp
@@ -13,9 +13,12 @@
 
 #include "Brushes/SlateColorBrush.h"
 
+#include "Components/MeshComponent.h"
 #include "Components/StaticMeshComponent.h"
 #include "Components/TextRenderComponent.h"
 
+#include "Customizations/RichTextColorDecorator.h"
+
 #include "DataTypes/PVData.h"
 
 #include "Editor/EditorPerProjectUserSettings.h"
@@ -24,17 +27,55 @@
 
 #include "Nodes/PVBaseSettings.h"
 
+#include "Styling/SlateIconFinder.h"
+
 #include "ViewportToolbar/UnrealEdViewportToolbar.h"
 
 #include "Visualizations/PVScaleVisualizationComponent.h"
 
 #include "Widgets/Input/SNumericEntryBox.h"
 #include "Widgets/Text/SRichTextBlock.h"
+#include "Widgets/Views/SListView.h"
+#include "Widgets/Views/STableRow.h"
 
 #define LOCTEXT_NAMESPACE "PVEditorViewportToolbarSections"
 
+static FTableViewStyle GTransparentViewStyle = []
+	{
+		const FSlateColorBrush Brush(FLinearColor::Transparent);
+
+		FTableViewStyle ViewStyle = FAppStyle::Get().GetWidgetStyle<FTableViewStyle>("ListView");
+		ViewStyle.SetBackgroundBrush(Brush);
+
+		return ViewStyle;
+	}();
+
+static FTableRowStyle GTransparentRowStyle = []
+	{
+		FTableRowStyle Style = FAppStyle::Get().GetWidgetStyle<FTableRowStyle>("TableView.Row");
+
+		const FSlateColorBrush Brush(FLinearColor::Transparent);
+
+		Style.SetEvenRowBackgroundBrush(Brush)
+		     .SetEvenRowBackgroundHoveredBrush(Brush)
+		     .SetOddRowBackgroundBrush(Brush)
+		     .SetOddRowBackgroundHoveredBrush(Brush)
+		     .SetSelectorFocusedBrush(Brush)
+		     .SetActiveBrush(Brush)
+		     .SetActiveHoveredBrush(Brush)
+		     .SetInactiveBrush(Brush)
+		     .SetInactiveHoveredBrush(Brush)
+		     .SetDropIndicator_Above(Brush)
+		     .SetDropIndicator_Onto(Brush)
+		     .SetDropIndicator_Below(Brush);
+
+		return Style;
+	}();
+
 SPVEditorViewport::SPVEditorViewport()
 {
+	FoliageComponents = MakeShared<TArray<TObjectPtr<UMeshComponent>>>();
+	
 	PreviewNodeBackgroundBrush = MakeShared<FSlateColorBrush>(FColor{70, 100, 200});
 	
 	if (AdvancedPreviewScene)
@@ -106,6 +147,7 @@ TSharedPtr<SWidget> SPVEditorViewport::BuildViewportToolbar()
 		Context.AddExtender(GetExtenders());
 		Context.AddObject(UE::UnrealEd::CreateViewportToolbarDefaultContext(GetViewportWidget()));
 		Context.AddObject(CreateMannequinWidgetContext(CreateMannequinOffsetWidget()));
+		Context.AddObject(CreateFoliageComponentsContext());
 	}
 	
 	return UToolMenus::Get()->GenerateWidget(ToolbarName, Context);
@@ -190,6 +232,56 @@ void SPVEditorViewport::PopulateViewportOverlays(TSharedRef<SOverlay> Overlay)
 				SAssignNew(StatsOverlayText, SRichTextBlock)
 			]
 		];
+
+	Overlay
+		->AddSlot()
+		.VAlign(VAlign_Top)
+		.HAlign(HAlign_Right)
+		.Padding(6.0f, 25.f)
+		[
+			SNew(SBorder)
+			.BorderImage(FAppStyle::Get().GetBrush("FloatingBorder"))
+			.Visibility_Lambda([this]
+				{
+					return !LegendData.IsEmpty()
+						? EVisibility::Visible
+						: EVisibility::Collapsed;
+				})
+			.Padding(4.f)
+			[
+				SNew(SVerticalBox)
+				+ SVerticalBox::Slot()
+				.AutoHeight()
+				.Padding(0, 0, 0, 4)
+				[
+					SNew(SRichTextBlock)
+					.Text(FText::FromString(TEXT("<TextBlock.ShadowedText>Foliage Legend</>")))
+					.Justification(ETextJustify::Center)
+					.TextStyle(&FCoreStyle::Get().GetWidgetStyle<FTextBlockStyle>( "BoldFont" ))
+				]
+				+ SVerticalBox::Slot()
+				.FillHeight(1.f)
+				[
+					SNew(SListView<TSharedPtr<FText>>)
+					.ListItemsSource(&LegendData)
+					.ListViewStyle(&GTransparentViewStyle)
+					.OnGenerateRow_Lambda(
+						[](const TSharedPtr<FText>& Item, const TSharedRef<STableViewBase>& OwningTable) -> TSharedRef<ITableRow>
+							{
+								return SNew(STableRow<TSharedPtr<FText>>, OwningTable)
+									.Style(&GTransparentRowStyle)
+									.ShowSelection(false)
+									.Padding(FMargin(0))
+									[
+										SNew(SRichTextBlock)
+										.Text(*Item)
+										.Decorators({FColorBlockDecorator::CreateDecorator()})
+									];
+							}
+					)
+				]
+			]
+		];
 }
 
 void SPVEditorViewport::InitVisualizationScene()
@@ -355,6 +447,13 @@ UPVMannequinWidgetContext* SPVEditorViewport::CreateMannequinWidgetContext(TShar
 	return Context;
 }
 
+UPVFoliageComponentsContext* SPVEditorViewport::CreateFoliageComponentsContext()
+{
+	UPVFoliageComponentsContext* Context = NewObject<UPVFoliageComponentsContext>();
+	Context->FoliageComponents = FoliageComponents;
+	return Context;
+}
+
 void SPVEditorViewport::OnSetupScene()
 {
 	SPCGEditorViewport::OnSetupScene();
@@ -377,6 +476,8 @@ void SPVEditorViewport::OnResetScene()
 	SPCGEditorViewport::OnResetScene();
 
 	ResetVisualizationScene();
+
+	FoliageComponents->Empty();
 }
 
 FToolMenuEntry SPVEditorViewport::CreateSettingsToolbarMenu()
@@ -389,7 +490,8 @@ FToolMenuEntry SPVEditorViewport::CreateSettingsToolbarMenu()
 				FText::GetEmpty(),
 				FNewToolMenuDelegate::CreateLambda([](UToolMenu* Submenu)
 					{
-						const UPVMannequinWidgetContext* Context = Submenu->FindContext<UPVMannequinWidgetContext>();
+						const UPVMannequinWidgetContext* MannequinContext = Submenu->FindContext<UPVMannequinWidgetContext>();
+						const UPVFoliageComponentsContext* FoliageContext = Submenu->FindContext<UPVFoliageComponentsContext>();
 						
 						// Add preference toggles to this section.
 						FToolMenuSection& MannequinSection = Submenu->FindOrAddSection(TEXT("PVE_MannequinSettings"),
@@ -398,7 +500,7 @@ FToolMenuEntry SPVEditorViewport::CreateSettingsToolbarMenu()
 						MannequinSection.AddEntry(
 							FToolMenuEntry::InitWidget(
 								"PVE_MannequinOffset",
-								Context->MannequinOffsetWidget.ToSharedRef(),
+								MannequinContext->MannequinOffsetWidget.ToSharedRef(),
 								LOCTEXT("PVE_MannequinOffsetLabel", "Mannequin Offset"),
 								false,
 								true,
@@ -410,6 +512,57 @@ FToolMenuEntry SPVEditorViewport::CreateSettingsToolbarMenu()
 							LOCTEXT("PVE_ScaleVisSettingsSectionLabel", "Scale Vis. Settings"));
 						ScaleVisSection.AddMenuEntry(FPVEditorCommands::Get().ShowScaleVisualization);
 
+						if (FoliageContext && FoliageContext->FoliageComponents && FoliageContext->FoliageComponents->Num())
+						{
+							FToolMenuSection& FoliageSection = Submenu->FindOrAddSection(TEXT("PVE_FoliageSettings"),
+								LOCTEXT("PVE_FoliageSettingsSectionLabel", "Foliage Settings"));
+							FoliageSection.AddSubMenu(
+								"PVE_ToggleFoliage",
+								LOCTEXT("PVE_ToggleFoliageSettingsLabel", "Toggle Foliage"),
+								LOCTEXT("PVE_ToggleFoliageSettingsTooltip", "Toggle Individual Foliage"),
+								FNewToolMenuChoice(
+									FNewToolMenuDelegate::CreateLambda([](UToolMenu* FoliageSubMenu)
+										{
+											const UPVFoliageComponentsContext* Context = FoliageSubMenu->FindContext<UPVFoliageComponentsContext>();
+
+											if (Context->FoliageComponents)
+											{
+												for (TObjectPtr<UMeshComponent> InstancedStaticMeshComponent : *Context->FoliageComponents)
+												{
+													const FString ObjectName = InstancedStaticMeshComponent->GetName();
+													FoliageSubMenu->AddMenuEntry(
+														"Toggle Foliage",
+														FToolMenuEntry::InitMenuEntry(
+															*ObjectName,
+															FText::FromString(ObjectName),
+															FText::Format(FText::FromString("Toggle visibility of: {0}"),
+																FText::FromString(ObjectName)),
+															FSlateIcon(),
+															FUIAction(
+																FExecuteAction::CreateLambda([InstancedStaticMeshComponent]()
+																	{
+																		InstancedStaticMeshComponent->ToggleVisibility(true);
+																	}
+																),
+																FCanExecuteAction(),
+																FIsActionChecked::CreateLambda([InstancedStaticMeshComponent]
+																	{
+																		return InstancedStaticMeshComponent->IsVisible();
+																	})
+															),
+															EUserInterfaceActionType::ToggleButton
+														)
+													);
+												}
+											}
+										}
+									)
+								),
+								false,
+								FSlateIconFinder::FindIcon("ClassIcon.FoliageType_Actor")
+							);
+						}
+
 						FToolMenuSection& AutoFocusViewportSection = Submenu->FindOrAddSection(TEXT("PVE_InteractionSettings"),
 							LOCTEXT("PVE_InteractionSettingsSectionLabel", "Interaction Settings"));
 						AutoFocusViewportSection.AddMenuEntry(FPVEditorCommands::Get().AutoFocusViewport);
@@ -425,19 +578,19 @@ EPVVisualizationMode PVRenderModesToVisualizationMode(const TArray<EPVRenderType
 	{
 		return EPVVisualizationMode::FoliageMesh;
 	}
-	else if (InRenderType.Contains(EPVRenderType::FoliageGrid))
+	if (InRenderType.Contains(EPVRenderType::FoliageGrid))
 	{
 		return EPVVisualizationMode::FoliageGrid;
 	}
-	else if (InRenderType.Contains(EPVRenderType::Mesh))
+	if (InRenderType.Contains(EPVRenderType::Mesh))
 	{
 		return EPVVisualizationMode::Mesh;
 	}
-	else if (InRenderType.Contains(EPVRenderType::Bones))
+	if (InRenderType.Contains(EPVRenderType::Bones))
 	{
 		return EPVVisualizationMode::Bones;
 	}
-	else if (InRenderType.Contains(EPVRenderType::PointData))
+	if (InRenderType.Contains(EPVRenderType::PointData))
 	{
 		return EPVVisualizationMode::PointData;
 	}
@@ -453,7 +606,10 @@ TArray<EPVRenderType> VisualizationModeToRenderModes(EPVVisualizationMode InMode
 		return DefautlSettings;
 		
 	case EPVVisualizationMode::FoliageMesh:
-		return {EPVRenderType::Mesh, EPVRenderType::Foliage};			
+		return {EPVRenderType::Mesh, EPVRenderType::Foliage};
+
+	case EPVVisualizationMode::FoliageAttachmentMesh:
+		return {EPVRenderType::Mesh, EPVRenderType::FoliageAttachments};	
 
 	case EPVVisualizationMode::FoliageGrid:
 		return {EPVRenderType::FoliageGrid};
@@ -473,6 +629,9 @@ TArray<EPVRenderType> VisualizationModeToRenderModes(EPVVisualizationMode InMode
 	case EPVVisualizationMode::PointDataMesh:
 		return {EPVRenderType::PointData, EPVRenderType::Mesh};
 
+	case EPVVisualizationMode::PointDataFoliageAttachment:
+		return {EPVRenderType::PointData, EPVRenderType::FoliageAttachments};
+
 	default:
 		return TArray<EPVRenderType>{};
 		
@@ -481,92 +640,85 @@ TArray<EPVRenderType> VisualizationModeToRenderModes(EPVVisualizationMode InMode
 
 FToolMenuEntry SPVEditorViewport::CreateVisualizationModeToolbarMenu()
 {
-	return FToolMenuEntry::InitDynamicEntry
-	(
-	    "VisualizationMode",
-	    FNewToolMenuSectionDelegate::CreateSPLambda(this,[this](FToolMenuSection& InnerSection) -> void
-	       {
-	       	TAttribute<FText> LabelAttribute = UE::UnrealEd::GetViewModesSubmenuLabel(nullptr);
+	return FToolMenuEntry::InitDynamicEntry(
+		"VisualizationMode",
+		FNewToolMenuSectionDelegate::CreateSPLambda(this, [this](FToolMenuSection& InnerSection) -> void
+			{
+				TAttribute<FText> LabelAttribute = UE::UnrealEd::GetViewModesSubmenuLabel(nullptr);
 
-	       	LabelAttribute = TAttribute<FText>::CreateSPLambda(this,
-					[this]()
+				LabelAttribute = TAttribute<FText>::CreateSPLambda(this, [this]
 					{
-						UEnum* EnumPtr = StaticEnum<EPVVisualizationMode>();
-							return EnumPtr->GetDisplayNameTextByIndex((int)CurrentVisualizationMode);
+						const UEnum* EnumPtr = StaticEnum<EPVVisualizationMode>();
+						return EnumPtr->GetDisplayNameTextByIndex(static_cast<int>(CurrentVisualizationMode));
 					}
 				);
-	       		
-	          FToolMenuEntry& Entry = InnerSection.AddSubMenu(
-	             "VisualizationModes",
-	             LabelAttribute,
-	             FText::GetEmpty(),
-	             FNewToolMenuDelegate::CreateSPLambda
-	             (this, [this](UToolMenu* Submenu) -> void
-	                   {
-	                      FToolMenuSection& RenderModeSelectionSection = Submenu->FindOrAddSection(
-	                         "VisualizationModeSection",
-	                         LOCTEXT("FoliageVisualizationModeSelectionSectionLabel", "Visualization modes")
-	                      );
-
-	                      UEnum* EnumPtr = StaticEnum<EPVVisualizationMode>();
-	                      
-	                      if (EnumPtr)
-	                      {
-	                         for (int32 i = 0; i < EnumPtr->NumEnums() - 1; i++) // last entry is _MAX or hidden
-	                         {
-	                            FText ModeName = EnumPtr->GetDisplayNameTextByIndex(i);
-	                            EPVVisualizationMode Mode = static_cast<EPVVisualizationMode>(EnumPtr->GetValueByIndex(i));
-	                         	
-	                         	RenderModeSelectionSection.AddMenuEntry
-	                         	(
-								NAME_None,
-								ModeName,
-								FText(),
-								FSlateIcon(),
-								FUIAction(
-								   FExecuteAction::CreateSPLambda(this,
-									  [this, Mode]()
-									  {
-									  		OnVisualizationModeChanged(Mode);
-									  }
-								   ),
-								   FCanExecuteAction::CreateSPLambda(this,[this, Mode]()
-									   {
-								   		if (Mode != EPVVisualizationMode::Default)
-								   		{
-								   			auto RenderTypes = VisualizationModeToRenderModes(Mode);
-
-											   for (int i = 0; i < RenderTypes.Num(); i++)
-											   {
-												   if (!SupportedRenderTypes.Contains(RenderTypes[i]))
-												   {
-													   return false;
-												   }
-											   }
-								   		}
-								   		
-								   		return true;
-								   		
-									   }),
-								   FIsActionChecked::CreateSPLambda(this,
-									  [this, Mode]()
-									  {
-									  		return CurrentVisualizationMode == Mode;
-									  }
-								   )
-								),EUserInterfaceActionType::RadioButton
-								);
-	                         }
-	                      }
+
+				FToolMenuEntry& Entry = InnerSection.AddSubMenu(
+					"VisualizationModes",
+					LabelAttribute,
+					FText::GetEmpty(),
+					FNewToolMenuDelegate::CreateSPLambda(this, [this](UToolMenu* Submenu) -> void
+						{
+							FToolMenuSection& RenderModeSelectionSection = Submenu->FindOrAddSection(
+								"VisualizationModeSection",
+								LOCTEXT("FoliageVisualizationModeSelectionSectionLabel", "Visualization modes")
+							);
+
+							UEnum* EnumPtr = StaticEnum<EPVVisualizationMode>();
+
+							if (EnumPtr)
+							{
+								for (int32 i = 0; i < EnumPtr->NumEnums() - 1; i++) // last entry is _MAX or hidden
+								{
+									FText ModeName = EnumPtr->GetDisplayNameTextByIndex(i);
+									EPVVisualizationMode Mode = static_cast<EPVVisualizationMode>(EnumPtr->GetValueByIndex(i));
+
+									RenderModeSelectionSection.AddMenuEntry(
+										NAME_None,
+										ModeName,
+										FText(),
+										FSlateIcon(),
+										FUIAction(
+											FExecuteAction::CreateSPLambda(this, [this, Mode]
+												{
+													OnVisualizationModeChanged(Mode);
+												}
+											),
+											FCanExecuteAction::CreateSPLambda(this, [this, Mode]
+												{
+													if (Mode != EPVVisualizationMode::Default)
+													{
+														TArray<EPVRenderType> RenderTypes = VisualizationModeToRenderModes(Mode);
+														for (EPVRenderType RenderType : RenderTypes)
+														{
+															if (!SupportedRenderTypes.Contains(RenderType))
+															{
+																return false;
+															}
+														}
+													}
+
+													return true;
+												}),
+											FIsActionChecked::CreateSPLambda(this, [this, Mode]
+												{
+													return CurrentVisualizationMode == Mode;
+												}
+											)
+										),
+										EUserInterfaceActionType::RadioButton
+									);
+								}
+							}
 						}
-	             ),
-	             false,
-	             FSlateIcon(FAppStyle::GetAppStyleSetName(), "EditorViewport.Visualizers")
-	          );
-	          Entry.ToolBarData.ResizeParams.ClippingPriority = 1100;
-	       }
-	    )
-	 );
+					),
+					false,
+					FSlateIcon(FAppStyle::GetAppStyleSetName(), "EditorViewport.Visualizers")
+				);
+				Entry.ToolBarData.ResizeParams.ClippingPriority = 1100;
+			}
+		)
+	);
 }
 
 void SPVEditorViewport::OnVisualizationModeChanged(EPVVisualizationMode InMode)
@@ -627,4 +779,15 @@ void SPVEditorViewport::PopulateStatsOverlayText(const TArrayView<FText> TextIte
 	StatsOverlayText->SetText(FinalText.ToText());
 }
 
+void SPVEditorViewport::PopulateLegendOverlayText(const TArrayView<TPair<FColor, FText>> LegendItems)
+{
+	static FText LegendStyle = FText::FromString(TEXT("<block hex=\"{0}\">          </>\t<TextBlock.ShadowedText>{1}</>"));
+
+	LegendData.Empty();
+	for (const auto& [Color, Text] : LegendItems)
+	{
+		LegendData.Add(MakeShared<FText>(FText::Format(LegendStyle, FText::FromString(Color.ToHex()), Text)));
+	}
+}
+
 #undef LOCTEXT_NAMESPACE
diff --git a/Engine/Plugins/Experimental/ProceduralVegetationEditor/Source/ProceduralVegetationEditor/Private/SPVEditorViewport.h b/Engine/Plugins/Experimental/ProceduralVegetationEditor/Source/ProceduralVegetationEditor/Private/SPVEditorViewport.h
--- a/Engine/Plugins/Experimental/ProceduralVegetationEditor/Source/ProceduralVegetationEditor/Private/SPVEditorViewport.h
+++ b/Engine/Plugins/Experimental/ProceduralVegetationEditor/Source/ProceduralVegetationEditor/Private/SPVEditorViewport.h
@@ -12,16 +12,26 @@
 UENUM()
 enum class EPVVisualizationMode : uint8
 {
-	Default			UMETA(DisplayName = "Default"),
-	PointData		UMETA(DisplayName = "Point data"),
-	Mesh			UMETA(DisplayName = "Mesh"),
-	FoliageGrid		UMETA(DisplayName = "Foliage grid"),
-	Bones			UMETA(DisplayName = "Bones"),
-	PointDataMesh	UMETA(DisplayName = "Point data + Mesh"),
-	FoliageMesh		UMETA(DisplayName = "Foliage + Mesh"),
-	BonesMesh		UMETA(DisplayName = "Bones + Mesh")
+	Default							UMETA(DisplayName = "Default"),
+	PointData						UMETA(DisplayName = "Point data"),
+	Mesh							UMETA(DisplayName = "Mesh"),
+	FoliageGrid						UMETA(DisplayName = "Foliage grid"),
+	Bones							UMETA(DisplayName = "Bones"),
+	PointDataMesh					UMETA(DisplayName = "Point data + Mesh"),
+	PointDataFoliageAttachment		UMETA(DisplayName = "Point data + Foliage attachments"),
+	FoliageMesh						UMETA(DisplayName = "Foliage + Mesh"),
+	FoliageAttachmentMesh			UMETA(DisplayName = "Foliage attachments + Mesh"),
+	BonesMesh						UMETA(DisplayName = "Bones + Mesh")
 };
 
+struct FLegendData
+{
+	FLinearColor Color;
+	FText Text;
+};
+
+class UMeshComponent;
+
 class UPVBaseSettings;
 class UPVEditorSettings;
 class UPVScaleVisualizationComponent;
@@ -37,6 +47,15 @@ public:
 	TSharedPtr<SWidget> MannequinOffsetWidget = nullptr;
 };
 
+UCLASS(MinimalAPI)
+class UPVFoliageComponentsContext : public UObject
+{
+	GENERATED_BODY()
+
+public:
+	TSharedPtr<TArray<TObjectPtr<UMeshComponent>>> FoliageComponents = nullptr;
+};
+
 class SPVEditorViewport : public SPCGEditorViewport
 {
 public:
@@ -51,6 +70,9 @@ public:
 
 	void SetOverlayText(const FText& CurrentlyLockedNodeName = FText::GetEmpty());
 	void PopulateStatsOverlayText(const TArrayView<FText> TextItems);
+	void PopulateLegendOverlayText(const TArrayView<TPair<FColor, FText>> LegendItems);
+
+	TSharedPtr<TArray<TObjectPtr<UMeshComponent>>> FoliageComponents;
 
 protected:
 	//~ Begin SEditorViewport Interface
@@ -61,6 +83,7 @@ protected:
 	virtual TSharedRef<FEditorViewportClient> MakeEditorViewportClient() override;
 	virtual void PopulateViewportOverlays(TSharedRef<SOverlay> Overlay) override;
 
+
 private:
 	void InitVisualizationScene();
 	void ResetVisualizationScene();
@@ -82,6 +105,7 @@ private:
 	void SetMannequinOffset(float NewValue, bool bSaveConfig = false) const;
 
 	static UPVMannequinWidgetContext* CreateMannequinWidgetContext(TSharedPtr<SWidget> InOffsetWidget);
+	UPVFoliageComponentsContext* CreateFoliageComponentsContext();
 
 	TArray<EPVRenderType> SupportedRenderTypes;
 	EPVVisualizationMode CurrentVisualizationMode = EPVVisualizationMode::Default;
@@ -102,6 +126,7 @@ private:
 	bool bIsPreviewingLockedNode = false;
 	TSharedPtr<SRichTextBlock> OverlayText;
 	TSharedPtr<SRichTextBlock> StatsOverlayText;
+	TArray<TSharedPtr<FText>> LegendData;
 
 	TSharedPtr<FSlateBrush> PreviewNodeBackgroundBrush;
 };
diff --git a/Engine/Plugins/Experimental/ProceduralVegetationEditor/Source/ProceduralVegetationEditor/Private/Visualizations/PVDataVisualization.cpp b/Engine/Plugins/Experimental/ProceduralVegetationEditor/Source/ProceduralVegetationEditor/Private/Visualizations/PVDataVisualization.cpp
--- a/Engine/Plugins/Experimental/ProceduralVegetationEditor/Source/ProceduralVegetationEditor/Private/Visualizations/PVDataVisualization.cpp
+++ b/Engine/Plugins/Experimental/ProceduralVegetationEditor/Source/ProceduralVegetationEditor/Private/Visualizations/PVDataVisualization.cpp
@@ -6,27 +6,38 @@
 #include "EditorViewportClient.h"
 #include "GeometryCollectionToDynamicMesh.h"
 #include "PCGComponent.h"
-#include "PCGContext.h"
 #include "PCGDataVisualization.h"
-#include "PCGEditorSettings.h"
 #include "PlanarCut.h"
 #include "PVBoneComponent.h"
+#include "SPVEditorViewport.h"
+
 #include "Components/DynamicMeshComponent.h"
+#include "Components/InstancedStaticMeshComponent.h"
 #include "Components/LineBatchComponent.h"
+#include "Components/SkeletalMeshComponent.h"
+#include "Components/StaticMeshComponent.h"
 #include "Components/TextRenderComponent.h"
+
 #include "DataTypes/PVData.h"
+
 #include "DebugVisualization/PVDebugVisualizer.h"
+
+#include "Engine/SkeletalMesh.h"
+#include "Engine/StaticMesh.h"
+
 #include "Facades/PVFoliageFacade.h"
+
 #include "GeometryCollection/GeometryCollectionObject.h"
+
 #include "Helpers/PVUtilities.h"
+
+#include "Materials/Material.h"
+
 #include "Nodes/PVBaseSettings.h"
+
 #include "Rendering/SkeletalMeshRenderData.h"
+
 #include "Visualizations/PVSkeletonVisualizerComponent.h"
-#include "Engine/StaticMesh.h"
-#include "Engine/SkeletalMesh.h"
-#include "Components/StaticMeshComponent.h"
-#include "Components/SkeletalMeshComponent.h"
-#include "Materials/Material.h"
 
 #define LOCTEXT_NAMESPACE "PCGManagedArrayCollectionDataVisualization"
 
@@ -46,6 +57,7 @@ void FPVDataVisualization::RegisterVisualizations()
 	RenderMap.Add(EPVRenderType::Foliage, FPVRenderCallback::CreateRaw(this, &FPVDataVisualization::FoliageRenderer));
 	RenderMap.Add(EPVRenderType::Bones, FPVRenderCallback::CreateRaw(this, &FPVDataVisualization::BonesRenderer));
 	RenderMap.Add(EPVRenderType::FoliageGrid, FPVRenderCallback::CreateRaw(this, &FPVDataVisualization::FoliageGridRenderer));
+	RenderMap.Add(EPVRenderType::FoliageAttachments, FPVRenderCallback::CreateRaw(this, &FPVDataVisualization::FoliageAttachmentPointsRenderer));
 }
 
 void FPVDataVisualization::ExecuteDebugDisplay(FPCGContext* Context, const UPCGSettingsInterface* SettingsInterface, const UPCGData* Data,
@@ -59,8 +71,7 @@ FPCGSetupSceneFunc FPVDataVisualization::GetViewportSetupFunc(const UPCGSettings
 	return [this, WeakData=TWeakObjectPtr<const UPVData>(Cast<UPVData>(Data)), WeakSettings = TWeakObjectPtr<const UPVBaseSettings>(Cast<UPVBaseSettings>(SettingsInterface))](
 		FPCGSceneSetupParams& InOutParams)
 	{
-		check(InOutParams.Scene);
-
+		check(InOutParams.Scene)
 		if (!WeakData.IsValid())
 		{
 			UE_LOG(LogProceduralVegetationDataVisualization, Error, TEXT("Failed to setup data viewport, the data was lost or invalid."));
@@ -185,8 +196,11 @@ void FPVDataVisualization::FoliageRenderer(const FManagedArrayCollection& InColl
 	                                           PV::Utilities::FFoliageComponentCreatedCallback::CreateLambda(OnComponentCreated),
 	                                           OutInstancedComponentMap);
 
+	const TSharedPtr<SPVEditorViewport> ViewportWidget = StaticCastSharedPtr<SPVEditorViewport>(
+			InOutParams.EditorViewportClient->GetEditorViewportWidget());
 	for (const auto& [FoliageName , Component] : OutInstancedComponentMap)
 	{
+		ViewportWidget->FoliageComponents->Add(Component);
 		OutBounds = OutBounds + Component->CalcLocalBounds();
 	}
 }
@@ -360,4 +374,78 @@ void FPVDataVisualization::FoliageGridRenderer(const FManagedArrayCollection& In
 	}
 }
 
+void FPVDataVisualization::FoliageAttachmentPointsRenderer(const FManagedArrayCollection& InCollection, FPCGSceneSetupParams& InOutParams,
+	FBoxSphereBounds& OutBounds)
+{
+	const PV::Facades::FFoliageFacade FoliageFacade(InCollection);
+
+	const TObjectPtr<UPVLineBatchComponent> LineBatchComponent = NewObject<UPVLineBatchComponent>(GetTransientPackage(), NAME_None, RF_Transient);
+
+	const TObjectPtr<UInstancedStaticMeshComponent> SphereInstancer = NewObject<UInstancedStaticMeshComponent>(GetTransientPackage(), NAME_None, RF_Transient);
+	SphereInstancer->SetStaticMesh(LoadObject<UStaticMesh>(nullptr, TEXT("/Engine/BasicShapes/Sphere.Sphere")));
+	SphereInstancer->SetMaterial(0, LoadObject<UMaterialInterface>(nullptr, TEXT("/ProceduralVegetationEditor/Materials/FoliageAttachmentPointMaterial.FoliageAttachmentPointMaterial")));
+	SphereInstancer->SetNumCustomDataFloats(3);
+
+	const auto GetUniqueColor = [](float Alpha)
+	{
+		Alpha = FMath::Clamp(Alpha, 0.f, 1.f);
+		const float HueDeg = 240.f * Alpha;
+		FMath::SRandInit(Alpha * 255);
+		const float V = FMath::GetMappedRangeValueUnclamped(
+			FVector2f{0.0f, 1.0f},
+			FVector2f{0.25f, 1.0f},
+			FMath::SRand()
+		);
+		return FLinearColor(HueDeg, 0.9, V).HSVToLinearRGB();
+	};
+
+	TArray<TPair<FColor, FText>> LegendItems;
+	const TArray<FString> FoliageNames = FoliageFacade.GetFoliageNames();
+	for (int32 FoliageIndex = 0; FoliageIndex < FoliageNames.Num(); FoliageIndex++)
+	{
+		const FLinearColor Color = GetUniqueColor(static_cast<float>(FoliageIndex) / static_cast<float>(FoliageNames.Num()));
+		LegendItems.Add({Color.ToFColorSRGB(), FText::FromString(FSoftObjectPath(FoliageNames[FoliageIndex]).GetAssetName())});
+	}
+
+	FBox Bounds;
+	for (int32 PointIndex = 0; PointIndex < FoliageFacade.GetElementCount(); PointIndex++)
+	{
+		const PV::Facades::FFoliageEntryData FoliageData = FoliageFacade.GetFoliageEntry(PointIndex);
+
+		Bounds += static_cast<FVector>(FoliageData.PivotPoint);
+		const FLinearColor Color = GetUniqueColor(static_cast<float>(FoliageData.NameId) / static_cast<float>(FoliageNames.Num()));
+		LineBatchComponent->AddLine(
+			static_cast<FVector>(FoliageData.PivotPoint),
+			static_cast<FVector>(FoliageData.PivotPoint + FoliageData.UpVector * 5.0f),
+			Color,
+			SDPG_World,
+			EPointDrawSettings::None
+		);
+
+		const int32 InstanceIndex = SphereInstancer->AddInstance(
+			FTransform(
+				FRotator(),
+				static_cast<FVector>(FoliageData.PivotPoint),
+				FVector(0.015f)
+			)
+		);
+		SphereInstancer->SetCustomDataValue(InstanceIndex, 0, Color.R, false);
+		SphereInstancer->SetCustomDataValue(InstanceIndex, 1, Color.G, false);
+		SphereInstancer->SetCustomDataValue(InstanceIndex, 2, Color.B, false);
+	}
+
+	SphereInstancer->MarkRenderStateDirty();
+
+	InOutParams.ManagedResources.Add(SphereInstancer);
+	InOutParams.Scene->AddComponent(SphereInstancer, FTransform::Identity);
+	
+	InOutParams.ManagedResources.Add(LineBatchComponent);
+	InOutParams.Scene->AddComponent(LineBatchComponent, FTransform::Identity);
+
+	TSharedPtr<SPVEditorViewport> ViewportWidget = StaticCastSharedPtr<SPVEditorViewport>(InOutParams.EditorViewportClient->GetEditorViewportWidget());
+	ViewportWidget->PopulateLegendOverlayText(LegendItems);
+	
+	OutBounds = OutBounds + Bounds;
+}
+
 #undef LOCTEXT_NAMESPACE
diff --git a/Engine/Plugins/Experimental/ProceduralVegetationEditor/Source/ProceduralVegetationEditor/Private/Visualizations/PVDataVisualization.h b/Engine/Plugins/Experimental/ProceduralVegetationEditor/Source/ProceduralVegetationEditor/Private/Visualizations/PVDataVisualization.h
--- a/Engine/Plugins/Experimental/ProceduralVegetationEditor/Source/ProceduralVegetationEditor/Private/Visualizations/PVDataVisualization.h
+++ b/Engine/Plugins/Experimental/ProceduralVegetationEditor/Source/ProceduralVegetationEditor/Private/Visualizations/PVDataVisualization.h
@@ -3,9 +3,11 @@
 #pragma once
 
 #include "PCGDataVisualization.h"
-#include "DataVisualizations/PCGSpatialDataVisualization.h"
+
 #include "DataTypes/PVData.h"
 
+#include "DataVisualizations/PCGSpatialDataVisualization.h"
+
 class AActor;
 class UPCGData;
 struct FPCGContext;
@@ -32,6 +34,7 @@ protected:
 	void FoliageRenderer(const FManagedArrayCollection& InCollection, FPCGSceneSetupParams& InOutParams, FBoxSphereBounds& OutBounds);
 	void BonesRenderer(const FManagedArrayCollection& InCollection, FPCGSceneSetupParams& InOutParams, FBoxSphereBounds& OutBounds);
 	void FoliageGridRenderer(const FManagedArrayCollection& InCollection, FPCGSceneSetupParams& InOutParams, FBoxSphereBounds& OutBounds);
+	void FoliageAttachmentPointsRenderer(const FManagedArrayCollection& InCollection, FPCGSceneSetupParams& InOutParams, FBoxSphereBounds& OutBounds);
 };
 
 DECLARE_LOG_CATEGORY_EXTERN(LogProceduralVegetationDataVisualization, Log, All);
\ No newline at end of file
diff --git a/Engine/Plugins/Experimental/ProceduralVegetationEditor/Source/ProceduralVegetationEditor/Private/Visualizations/PVLineBatchComponent.cpp b/Engine/Plugins/Experimental/ProceduralVegetationEditor/Source/ProceduralVegetationEditor/Private/Visualizations/PVLineBatchComponent.cpp
--- a/Engine/Plugins/Experimental/ProceduralVegetationEditor/Source/ProceduralVegetationEditor/Private/Visualizations/PVLineBatchComponent.cpp
+++ b/Engine/Plugins/Experimental/ProceduralVegetationEditor/Source/ProceduralVegetationEditor/Private/Visualizations/PVLineBatchComponent.cpp
@@ -90,7 +90,7 @@ void UPVLineBatchComponent::Draw(const FSceneView* View, FPrimitiveDrawInterface
 			Line.EndPos,
 			Line.Color,
 			Line.DepthPriorityGroup,
-			3.0f,
+			2.0f,
 			0,
 			true
 		);
diff --git a/Engine/Plugins/Experimental/ProceduralVegetationEditor/Source/ProceduralVegetationEditor/Private/Visualizations/PVLineBatchComponent.h b/Engine/Plugins/Experimental/ProceduralVegetationEditor/Source/ProceduralVegetationEditor/Private/Visualizations/PVLineBatchComponent.h
--- a/Engine/Plugins/Experimental/ProceduralVegetationEditor/Source/ProceduralVegetationEditor/Private/Visualizations/PVLineBatchComponent.h
+++ b/Engine/Plugins/Experimental/ProceduralVegetationEditor/Source/ProceduralVegetationEditor/Private/Visualizations/PVLineBatchComponent.h
@@ -80,5 +80,5 @@ protected:
 	TArray<FPVLineInfo> Lines;
 
 private:
-	const float PointSize = 7.5f;
+	const float PointSize = 5.0f;
 };
