[PVE] Toggle visibility of overlay stats and persist foliage visibility over graph reexecutions

diff --git a/Engine/Plugins/Experimental/ProceduralVegetationEditor/Source/ProceduralVegetationEditor/Private/PVEditorCommands.cpp b/Engine/Plugins/Experimental/ProceduralVegetationEditor/Source/ProceduralVegetationEditor/Private/PVEditorCommands.cpp
--- a/Engine/Plugins/Experimental/ProceduralVegetationEditor/Source/ProceduralVegetationEditor/Private/PVEditorCommands.cpp
+++ b/Engine/Plugins/Experimental/ProceduralVegetationEditor/Source/ProceduralVegetationEditor/Private/PVEditorCommands.cpp
@@ -20,9 +20,9 @@ void FPVEditorCommands::RegisterCommands()
 	UI_COMMAND(Export, "Export", "Export outputs.", EUserInterfaceActionType::Button, FInputChord(EModifierKey::Control, EKeys::E));
 	UI_COMMAND(ShowMannequin, "Mannequin", "Show/Hide Mannequin", EUserInterfaceActionType::ToggleButton, FInputChord(EModifierKey::Control, EKeys::M));
 	UI_COMMAND(ShowScaleVisualization, "Scale Visualization", "Show/Hide Scale Visualization", EUserInterfaceActionType::ToggleButton, FInputChord(EModifierKey::Control, EKeys::V));
+	UI_COMMAND(ShowStats, "Stats Overlay", "Show/Hide Viewport Stats Overlay", EUserInterfaceActionType::ToggleButton, FInputChord());
 	UI_COMMAND(AutoFocusViewport, "Auto focus viewport", "Auto focus viewport on node selection and visualization mode switch", EUserInterfaceActionType::ToggleButton, FInputChord());
 	UI_COMMAND(LockNodeInspection, "LockNodeInspection", "Lock node inspection.", EUserInterfaceActionType::Button, FInputChord(EModifierKey::Control, EKeys::L));
-
 }
 
 #undef LOCTEXT_NAMESPACE
diff --git a/Engine/Plugins/Experimental/ProceduralVegetationEditor/Source/ProceduralVegetationEditor/Private/PVEditorCommands.h b/Engine/Plugins/Experimental/ProceduralVegetationEditor/Source/ProceduralVegetationEditor/Private/PVEditorCommands.h
--- a/Engine/Plugins/Experimental/ProceduralVegetationEditor/Source/ProceduralVegetationEditor/Private/PVEditorCommands.h
+++ b/Engine/Plugins/Experimental/ProceduralVegetationEditor/Source/ProceduralVegetationEditor/Private/PVEditorCommands.h
@@ -20,5 +20,6 @@ public:
 	TSharedPtr<FUICommandInfo> ShowMannequin;
 	TSharedPtr<FUICommandInfo> ShowScaleVisualization;
 	TSharedPtr<FUICommandInfo> AutoFocusViewport;
-	TSharedPtr<FUICommandInfo> LockNodeInspection;;
+	TSharedPtr<FUICommandInfo> ShowStats;
+	TSharedPtr<FUICommandInfo> LockNodeInspection;
 };
\ No newline at end of file
diff --git a/Engine/Plugins/Experimental/ProceduralVegetationEditor/Source/ProceduralVegetationEditor/Private/SPVEditorViewport.cpp b/Engine/Plugins/Experimental/ProceduralVegetationEditor/Source/ProceduralVegetationEditor/Private/SPVEditorViewport.cpp
--- a/Engine/Plugins/Experimental/ProceduralVegetationEditor/Source/ProceduralVegetationEditor/Private/SPVEditorViewport.cpp
+++ b/Engine/Plugins/Experimental/ProceduralVegetationEditor/Source/ProceduralVegetationEditor/Private/SPVEditorViewport.cpp
@@ -75,6 +75,7 @@ static FTableRowStyle GTransparentRowStyle = []
 SPVEditorViewport::SPVEditorViewport()
 {
 	FoliageComponents = MakeShared<TArray<TObjectPtr<UMeshComponent>>>();
+	FoliageVisibility = MakeShared<TMap<FString, bool>>();
 	
 	PreviewNodeBackgroundBrush = MakeShared<FSlateColorBrush>(FColor{70, 100, 200});
 	
@@ -171,6 +172,12 @@ void SPVEditorViewport::BindCommands()
 		FCanExecuteAction(),
 		FIsActionChecked::CreateSP(this, &SPVEditorViewport::IsShowScaleVisChecked));
 
+	CommandList->MapAction(
+			Commands.ShowStats,
+			FExecuteAction::CreateSP(this, &SPVEditorViewport::ToggleShowStats),
+			FCanExecuteAction(),
+			FIsActionChecked::CreateSP(this, &SPVEditorViewport::IsShowStatsChecked));
+	
 	CommandList->MapAction(
 		Commands.AutoFocusViewport,
 		FExecuteAction::CreateSP(this, &SPVEditorViewport::ToggleAutofocusViewport),
@@ -223,7 +230,7 @@ void SPVEditorViewport::PopulateViewportOverlays(TSharedRef<SOverlay> Overlay)
 			.BorderImage(FAppStyle::Get().GetBrush("FloatingBorder"))
 			.Visibility_Lambda([this]
 				{
-					return StatsOverlayText && !StatsOverlayText->GetText().IsEmpty()
+					return StatsOverlayText && !StatsOverlayText->GetText().IsEmpty() && GetDefault<UPVEditorSettings>()->bShowStats
 						? EVisibility::Visible
 						: EVisibility::Collapsed;
 				})
@@ -404,6 +411,18 @@ void SPVEditorViewport::SetScaleVisState(bool InEnable)
 	}
 }
 
+void SPVEditorViewport::ToggleShowStats()
+{
+	UPVEditorSettings* EditorSettings = GetMutableDefault<UPVEditorSettings>();
+	EditorSettings->bShowStats = !EditorSettings->bShowStats;
+	EditorSettings->SaveConfig();
+}
+
+bool SPVEditorViewport::IsShowStatsChecked() const
+{
+	return GetDefault<UPVEditorSettings>()->bShowStats;
+}
+
 void SPVEditorViewport::ToggleAutofocusViewport()
 {
 	UPVEditorSettings* EditorSettings = GetMutableDefault<UPVEditorSettings>();
@@ -451,6 +470,7 @@ UPVFoliageComponentsContext* SPVEditorViewport::CreateFoliageComponentsContext()
 {
 	UPVFoliageComponentsContext* Context = NewObject<UPVFoliageComponentsContext>();
 	Context->FoliageComponents = FoliageComponents;
+	Context->FoliageVisibility = FoliageVisibility;
 	return Context;
 }
 
@@ -478,6 +498,7 @@ void SPVEditorViewport::OnResetScene()
 	ResetVisualizationScene();
 
 	FoliageComponents->Empty();
+	LegendData.Empty();
 }
 
 FToolMenuEntry SPVEditorViewport::CreateSettingsToolbarMenu()
@@ -494,6 +515,10 @@ FToolMenuEntry SPVEditorViewport::CreateSettingsToolbarMenu()
 						const UPVFoliageComponentsContext* FoliageContext = Submenu->FindContext<UPVFoliageComponentsContext>();
 						
 						// Add preference toggles to this section.
+						FToolMenuSection& OverlaySection = Submenu->FindOrAddSection(TEXT("PVE_OverlaySettings"),
+							LOCTEXT("PVE_OverlaySettingsSectionLabel", "Overlay Settings"));
+						OverlaySection.AddMenuEntry(FPVEditorCommands::Get().ShowStats);
+						
 						FToolMenuSection& MannequinSection = Submenu->FindOrAddSection(TEXT("PVE_MannequinSettings"),
 							LOCTEXT("PVE_MannequinSettingsSectionLabel", "Mannequin Settings"));
 						MannequinSection.AddMenuEntry(FPVEditorCommands::Get().ShowMannequin);
@@ -539,9 +564,10 @@ FToolMenuEntry SPVEditorViewport::CreateSettingsToolbarMenu()
 																FText::FromString(ObjectName)),
 															FSlateIcon(),
 															FUIAction(
-																FExecuteAction::CreateLambda([InstancedStaticMeshComponent]()
+																FExecuteAction::CreateLambda([=]()
 																	{
 																		InstancedStaticMeshComponent->ToggleVisibility(true);
+																		Context->FoliageVisibility->FindOrAdd(ObjectName, true) = InstancedStaticMeshComponent->GetVisibleFlag();
 																	}
 																),
 																FCanExecuteAction(),
@@ -790,4 +816,10 @@ void SPVEditorViewport::PopulateLegendOverlayText(const TArrayView<TPair<FColor,
 	}
 }
 
+void SPVEditorViewport::AddInstancedFoliageComponent(TObjectPtr<UMeshComponent> Component)
+{
+	FoliageComponents->Add(Component);
+	Component->SetVisibility(FoliageVisibility->FindOrAdd(Component->GetName(), true), true);
+}
+
 #undef LOCTEXT_NAMESPACE
diff --git a/Engine/Plugins/Experimental/ProceduralVegetationEditor/Source/ProceduralVegetationEditor/Private/SPVEditorViewport.h b/Engine/Plugins/Experimental/ProceduralVegetationEditor/Source/ProceduralVegetationEditor/Private/SPVEditorViewport.h
--- a/Engine/Plugins/Experimental/ProceduralVegetationEditor/Source/ProceduralVegetationEditor/Private/SPVEditorViewport.h
+++ b/Engine/Plugins/Experimental/ProceduralVegetationEditor/Source/ProceduralVegetationEditor/Private/SPVEditorViewport.h
@@ -54,6 +54,7 @@ class UPVFoliageComponentsContext : public UObject
 
 public:
 	TSharedPtr<TArray<TObjectPtr<UMeshComponent>>> FoliageComponents = nullptr;
+	TSharedPtr<TMap<FString, bool>> FoliageVisibility = nullptr;
 };
 
 class SPVEditorViewport : public SPCGEditorViewport
@@ -72,7 +73,9 @@ public:
 	void PopulateStatsOverlayText(const TArrayView<FText> TextItems);
 	void PopulateLegendOverlayText(const TArrayView<TPair<FColor, FText>> LegendItems);
 
-	TSharedPtr<TArray<TObjectPtr<UMeshComponent>>> FoliageComponents;
+	void AddInstancedFoliageComponent(TObjectPtr<UMeshComponent> Component);
+
+	
 
 protected:
 	//~ Begin SEditorViewport Interface
@@ -98,6 +101,9 @@ private:
 	bool IsShowScaleVisChecked() const;
 	void SetScaleVisState(bool InEnable);
 
+	void ToggleShowStats();
+	bool IsShowStatsChecked() const;
+	
 	void ToggleAutofocusViewport();
 	bool IsAutoFocusViewportChecked() const;
 
@@ -123,6 +129,9 @@ private:
 	TObjectPtr<UStaticMeshComponent> MannequinComponent = nullptr;
 	TObjectPtr<UPVScaleVisualizationComponent> ScaleVisualizationComponent = nullptr;
 
+	TSharedPtr<TArray<TObjectPtr<UMeshComponent>>> FoliageComponents;
+	TSharedPtr<TMap<FString, bool>> FoliageVisibility;
+
 	bool bIsPreviewingLockedNode = false;
 	TSharedPtr<SRichTextBlock> OverlayText;
 	TSharedPtr<SRichTextBlock> StatsOverlayText;
diff --git a/Engine/Plugins/Experimental/ProceduralVegetationEditor/Source/ProceduralVegetationEditor/Private/Visualizations/PVDataVisualization.cpp b/Engine/Plugins/Experimental/ProceduralVegetationEditor/Source/ProceduralVegetationEditor/Private/Visualizations/PVDataVisualization.cpp
--- a/Engine/Plugins/Experimental/ProceduralVegetationEditor/Source/ProceduralVegetationEditor/Private/Visualizations/PVDataVisualization.cpp
+++ b/Engine/Plugins/Experimental/ProceduralVegetationEditor/Source/ProceduralVegetationEditor/Private/Visualizations/PVDataVisualization.cpp
@@ -200,7 +200,7 @@ void FPVDataVisualization::FoliageRenderer(const FManagedArrayCollection& InColl
 			InOutParams.EditorViewportClient->GetEditorViewportWidget());
 	for (const auto& [FoliageName , Component] : OutInstancedComponentMap)
 	{
-		ViewportWidget->FoliageComponents->Add(Component);
+		ViewportWidget->AddInstancedFoliageComponent(Component);
 		OutBounds = OutBounds + Component->CalcLocalBounds();
 	}
 }
diff --git a/Engine/Plugins/Experimental/ProceduralVegetationEditor/Source/ProceduralVegetationEditor/Public/PVEditorSettings.h b/Engine/Plugins/Experimental/ProceduralVegetationEditor/Source/ProceduralVegetationEditor/Public/PVEditorSettings.h
--- a/Engine/Plugins/Experimental/ProceduralVegetationEditor/Source/ProceduralVegetationEditor/Public/PVEditorSettings.h
+++ b/Engine/Plugins/Experimental/ProceduralVegetationEditor/Source/ProceduralVegetationEditor/Public/PVEditorSettings.h
@@ -39,6 +39,9 @@ public:
 	UPROPERTY(EditAnywhere, config, Category = Viewport)
 	bool bShowScaleVisualization = false;
 
+	UPROPERTY(EditAnywhere, config, Category = Export)
+	bool bShowStats = true;
+
 	UPROPERTY(EditAnywhere, config, Category = Export)
 	bool bAutoFocusViewport = false;
 
