[PVE] Object Avoidance

diff --git a/Engine/Plugins/Experimental/ProceduralVegetationEditor/Source/ProceduralVegetation/Private/Facades/PVBranchFacade.cpp b/Engine/Plugins/Experimental/ProceduralVegetationEditor/Source/ProceduralVegetation/Private/Facades/PVBranchFacade.cpp
--- a/Engine/Plugins/Experimental/ProceduralVegetationEditor/Source/ProceduralVegetation/Private/Facades/PVBranchFacade.cpp
+++ b/Engine/Plugins/Experimental/ProceduralVegetationEditor/Source/ProceduralVegetation/Private/Facades/PVBranchFacade.cpp
@@ -20,8 +20,7 @@ namespace PV::Facades
 		, TrunkURangeAttribute(InCollection, AttributeNames::TrunkURange, GroupNames::DetailsGroup)
 		, BranchParentNumbers(InCollection, AttributeNames::BranchParentNumber, GroupNames::BranchGroup)
 		, PlantNumbers(InCollection, AttributeNames::PlantNumber, GroupNames::BranchGroup)
-	{
-	}
+	{}
 
 	FBranchFacade::FBranchFacade(const FManagedArrayCollection& InCollection)
 		: Collection(nullptr)
@@ -38,8 +37,7 @@ namespace PV::Facades
 		, TrunkURangeAttribute(InCollection, AttributeNames::TrunkURange, GroupNames::DetailsGroup)
 		, BranchParentNumbers(InCollection, AttributeNames::BranchParentNumber, GroupNames::BranchGroup)
 		, PlantNumbers(InCollection, AttributeNames::PlantNumber, GroupNames::BranchGroup)
-	{
-	}
+	{}
 
 	bool FBranchFacade::IsValid() const
 	{
@@ -69,6 +67,12 @@ namespace PV::Facades
 		return EmptyPoints;
 	}
 
+	int32 FBranchFacade::GetRootPoint(const int32 Index) const
+	{
+		const TArray<int32>& Points = GetPoints(Index);
+		return Points.Num() ? Points[0] : INDEX_NONE;
+	}
+
 	void FBranchFacade::SetPoints(const int32 Index, const TArray<int32>& InPoints)
 	{
 		if (BranchPoints.IsValid() && BranchPoints.IsValidIndex(Index))
@@ -88,6 +92,22 @@ namespace PV::Facades
 		return EmptyChildren;
 	}
 
+	TArray<int32> FBranchFacade::GetImmediateChildren(const int32 Index) const
+	{
+		const int32 ParentBranchNumber = GetBranchNumber(Index);
+		return GetChildren(Index).FilterByPredicate(
+			[&, this](const int32 Child)
+				{
+					const int32 ChildIndex = GetBranchNumbers().Find(Child);
+					if (ChildIndex != INDEX_NONE && GetParentBranchNumber(ChildIndex) == ParentBranchNumber)
+					{
+						return true;
+					}
+					return false;
+				}
+		);
+	}
+
 	void FBranchFacade::SetChildren(const int32 Index, const TArray<int32>& InChildren)
 	{
 		if (Children.IsValid() && Children.IsValidIndex(Index))
@@ -354,8 +374,8 @@ namespace PV::Facades
 
 		OutSortedIndices.Sort([&](int32 Index1, int32 Index2)
 			{
-				int32 Index1Hierarchy = GetBranchHierarchyNumber(Index1);
-				int32 Index2Hierarchy = GetBranchHierarchyNumber(Index2);
+				int32 Index1Hierarchy = GetHierarchyGenerationNumber(Index1);
+				int32 Index2Hierarchy = GetHierarchyGenerationNumber(Index2);
 
 				if (Index1Hierarchy == Index2Hierarchy)
 				{
diff --git a/Engine/Plugins/Experimental/ProceduralVegetationEditor/Source/ProceduralVegetation/Private/Facades/PVFoliageFacade.cpp b/Engine/Plugins/Experimental/ProceduralVegetationEditor/Source/ProceduralVegetation/Private/Facades/PVFoliageFacade.cpp
--- a/Engine/Plugins/Experimental/ProceduralVegetationEditor/Source/ProceduralVegetation/Private/Facades/PVFoliageFacade.cpp
+++ b/Engine/Plugins/Experimental/ProceduralVegetationEditor/Source/ProceduralVegetation/Private/Facades/PVFoliageFacade.cpp
@@ -428,4 +428,9 @@ namespace PV::Facades
 	{
 		return PivotPointsAttribute.Get();
 	}
+
+	TManagedArray<FVector3f>& FFoliageFacade::ModifyPivotPositions()
+	{
+		return PivotPointsAttribute.Modify();
+	}
 }
diff --git a/Engine/Plugins/Experimental/ProceduralVegetationEditor/Source/ProceduralVegetation/Private/Nodes/PVObjectAvoidance.cpp b/Engine/Plugins/Experimental/ProceduralVegetationEditor/Source/ProceduralVegetation/Private/Nodes/PVObjectAvoidance.cpp
new file mode 100644
--- /dev/null
+++ b/Engine/Plugins/Experimental/ProceduralVegetationEditor/Source/ProceduralVegetation/Private/Nodes/PVObjectAvoidance.cpp
@@ -0,0 +1,287 @@
+// Copyright Epic Games, Inc. All Rights Reserved.
+
+#include "Nodes/PVObjectAvoidance.h"
+
+#include "MeshDescriptionAdapter.h"
+
+#include "Data/PCGDynamicMeshData.h"
+
+#include "DataTypes/PVGrowthData.h"
+
+#include "DynamicMesh/DynamicMeshAABBTree3.h"
+
+#include "Engine/StaticMesh.h"
+
+#include "Facades/PVBranchFacade.h"
+#include "Facades/PVBudVectorsFacade.h"
+#include "Facades/PVPointFacade.h"
+
+FPCGDataTypeIdentifier UPVObjectAvoidanceSettings::GetInputPinTypeIdentifier() const
+{
+	return FPCGDataTypeIdentifier{FPVDataTypeInfoGrowth::AsId()};
+}
+
+FPCGDataTypeIdentifier UPVObjectAvoidanceSettings::GetOutputPinTypeIdentifier() const
+{
+	return FPCGDataTypeIdentifier{FPVDataTypeInfoGrowth::AsId()};
+}
+
+FPCGElementPtr UPVObjectAvoidanceSettings::CreateElement() const
+{
+	return MakeShared<FPVObjectAvoidanceElement>();
+}
+
+bool FPVObjectAvoidanceElement::ExecuteInternal(FPCGContext* InContext) const
+{
+	TRACE_CPUPROFILER_EVENT_SCOPE(FPVColliderImporterElement::ExecuteInternal);
+
+	const static auto GetTriangleNormal = [](
+		FMeshDescriptionTriangleMeshAdapter& Mesh,
+		const int32 TID,
+		TOptional<FVector> BaryCoords = TOptional<FVector>()
+	)
+		{
+			if (Mesh.HasNormals() && BaryCoords.IsSet())
+			{
+				FVector3f NormalA, NormalB, NormalC;
+				Mesh.GetTriNormals(TID, NormalA, NormalB, NormalC);
+				return FVector(BaryCoords->X * NormalA + BaryCoords->Y * NormalB + BaryCoords->Z * NormalC).GetUnsafeNormal();
+			}
+			else
+			{
+				FVector VertexA, VertexB, VertexC;
+				Mesh.GetTriVertices(TID, VertexA, VertexB, VertexC);
+				return UE::Geometry::VectorUtil::Normal(VertexA, VertexB, VertexC);
+			}
+		};
+
+	check(InContext);
+
+	const UPVObjectAvoidanceSettings* InputSettings = InContext->GetInputSettings<UPVObjectAvoidanceSettings>();
+	check(InputSettings);
+
+	const TArray<FPCGTaggedData>& Inputs = InContext->InputData.GetInputsByPin(PCGPinConstants::DefaultInputLabel);
+	if (!Inputs.IsEmpty())
+	{
+		if (const UPVGrowthData* InputData = Cast<UPVGrowthData>(Inputs[0].Data))
+		{
+			FManagedArrayCollection SkeletonCollection = InputData->GetCollection();
+			for (const auto& [MeshPath, Transform] : InputSettings->ObjectAvoidanceSettings.Colliders)
+			{
+				const UStaticMesh* ColliderMesh = MeshPath.LoadSynchronous();
+				if (!ColliderMesh)
+				{
+					continue;
+				}
+
+#if WITH_EDITORONLY_DATA
+				const FMeshDescription* MeshDescription = ColliderMesh->GetMeshDescription(0);
+				FMeshDescriptionTriangleMeshAdapter Mesh(MeshDescription);
+				UE::Geometry::TMeshAABBTree3 Octree(&Mesh);
+
+				PV::Facades::FPointFacade PointFacade(SkeletonCollection);
+				PV::Facades::FBranchFacade BranchFacade(SkeletonCollection);
+				PV::Facades::FFoliageFacade FoliageFacade(SkeletonCollection);
+				PV::Facades::FBudVectorsFacade BudVectorFacade = PV::Facades::FBudVectorsFacade(SkeletonCollection);
+
+				TManagedArray<FVector3f>& PointPositions = PointFacade.ModifyPositions();
+				TManagedArray<TArray<FVector3f>>& BudDirections = BudVectorFacade.ModifyBudDirections();
+				TManagedArray<FVector3f>& FoliagePositions = FoliageFacade.ModifyPivotPositions();
+
+				const auto ApplyOffsetToBranchSuccessors = [&](const int32 BranchIndex, const FVector3f& Offset)
+					{
+						for (const int32 BranchPoint : BranchFacade.GetPoints(BranchIndex))
+						{
+							PointPositions[BranchPoint] += Offset;
+						}
+						for (const int32 FoliageEntryId : FoliageFacade.GetFoliageEntryIdsForBranch(BranchIndex))
+						{
+							FoliagePositions[FoliageEntryId] += Offset;
+						}
+
+						const TArray<int32> BranchChildren = BranchFacade.GetChildren(BranchIndex);
+						for (int32 BranchChild : BranchChildren)
+						{
+							const int32 ChildIndex = BranchFacade.GetBranchNumbers().Find(BranchChild);
+							if (ChildIndex != INDEX_NONE)
+							{
+								for (const int32 ChildBranchPoint : BranchFacade.GetPoints(ChildIndex))
+								{
+									PointPositions[ChildBranchPoint] += Offset;
+								}
+								for (const int32 ChildFoliageEntryId : FoliageFacade.GetFoliageEntryIdsForBranch(ChildIndex))
+								{
+									FoliagePositions[ChildFoliageEntryId] += Offset;
+								}
+							}
+						}
+					};
+
+
+				TArray<int32> BranchIndices;
+				BranchIndices.Reserve(BranchFacade.GetElementCount());
+				BranchFacade.GetSortedBranchIndicesByHierarchy(BranchIndices);
+
+				for (int32 BranchIndex : BranchIndices)
+				{
+					const TArray<int32> BranchPoints = BranchFacade.GetPoints(BranchIndex);
+					for (int32 BranchPointIndex = 1; BranchPointIndex < BranchPoints.Num(); ++BranchPointIndex)
+					{
+						const int32 PointIndex = BranchPoints[BranchPointIndex - 1];
+						const int32 NextPointIndex = BranchPoints[BranchPointIndex];
+
+						const FVector PointPosition = static_cast<FVector>(PointPositions[PointIndex]);
+						const FVector NextPointPosition = static_cast<FVector>(PointPositions[NextPointIndex]);
+						const FVector SegmentVector = NextPointPosition - PointPosition;
+
+						const FTransform InverseTransform = Transform.Inverse();
+
+						FRay3d Ray(PointPosition, SegmentVector);
+						FRay3d InverseRay(InverseTransform.TransformPosition(Ray.Origin), InverseTransform.TransformVectorNoScale(Ray.Direction));
+
+						const float PointLFR = PointFacade.GetLengthFromRoot(PointIndex);
+						const float NextPointLFR = PointFacade.GetLengthFromRoot(NextPointIndex);
+
+						const float PointScale = PointFacade.GetPointScale(PointIndex);
+						const FVector MaxLengthVector = Ray.Direction * (SegmentVector.Length() + PointScale);
+						const double MaxLength = InverseTransform.TransformVector(MaxLengthVector).Length();
+
+						int32 HitTriangle = INDEX_NONE;
+						FVector BaryCoords;
+						if (double RayT; Octree.FindNearestHitTriangle(InverseRay, RayT, HitTriangle, BaryCoords, {MaxLength}))
+						{
+							const FVector HitNormal = Transform.TransformVectorNoScale(GetTriangleNormal(Mesh, HitTriangle, BaryCoords));
+
+							if (FVector::DotProduct(Ray.Direction, HitNormal) > 0)
+							{
+								continue;
+							}
+
+							const FVector HitPoint = Transform.TransformPosition(InverseRay.PointAt(RayT));
+
+							const FVector InvertedNormal = -HitNormal;
+							const double DistanceToHit = FVector::Distance(PointPosition, HitPoint);
+							const FVector BiNormal = FVector::CrossProduct(Ray.Direction, InvertedNormal);
+							const FVector Tangent = FVector::CrossProduct(InvertedNormal, BiNormal);
+							const FVector AvoidanceOffset = Tangent * DistanceToHit * (FVector::DotProduct(Ray.Direction, InvertedNormal) + 1);
+							const FVector AvoidancePosition = HitPoint + AvoidanceOffset + HitNormal * PointScale;
+							const FVector AvoidanceDirection = (AvoidancePosition - PointPosition).GetUnsafeNormal();
+
+							const FVector OffsetPosition = PointPosition + AvoidanceDirection * FVector::Distance(PointPosition, NextPointPosition);
+							const FVector3f Offset = FVector3f(OffsetPosition - NextPointPosition);
+
+							const int32 StartIndex = FMath::Max(1, BranchPointIndex - InputSettings->ObjectAvoidanceSettings.Smoothness);
+							FVector3f PreviousSiblingPosition;
+							for (int32 SiblingPointIndex = StartIndex; SiblingPointIndex < BranchPointIndex; ++SiblingPointIndex)
+							{
+								const int32 SiblingPoint = BranchPoints[SiblingPointIndex];
+								const int32 NextSiblingPoint = BranchPoints[SiblingPointIndex + 1];
+
+								const float SiblingPointLFR = PointFacade.GetLengthFromRoot(SiblingPoint);
+								const float NextSiblingPointLFR = PointFacade.GetLengthFromRoot(NextSiblingPoint);
+
+								const FVector3f PointOffset = Offset / FMath::Pow(1.1f, BranchPointIndex - SiblingPointIndex);
+
+								if (SiblingPointIndex != StartIndex)
+								{
+									const FVector3f SiblingSegmentVector = PointPositions[SiblingPoint] - PreviousSiblingPosition;
+									const FQuat4f Rotation = FQuat4f::FindBetween(SiblingSegmentVector, SiblingSegmentVector + PointOffset);
+									for (const int32 FoliageEntryId : FoliageFacade.GetFoliageEntryIdsForBranch(BranchIndex))
+									{
+										const float FoliageLFR = FoliageFacade.GetLengthFromRoot(FoliageEntryId);
+										if (FoliageLFR > SiblingPointLFR && FoliageLFR < NextSiblingPointLFR)
+										{
+											FoliagePositions[FoliageEntryId] += PointOffset;
+											FoliageFacade.SetUpVector(
+												FoliageEntryId,
+												Rotation.RotateVector(FoliageFacade.GetUpVector(FoliageEntryId))
+											);
+											FoliageFacade.SetNormalVector(
+												FoliageEntryId,
+												Rotation.RotateVector(FoliageFacade.GetNormalVector(FoliageEntryId))
+											);
+										}
+									}
+									for (FVector3f& BudDirection : BudDirections[SiblingPoint])
+									{
+										BudDirection = Rotation.RotateVector(BudDirection);
+									}
+								}
+
+								PreviousSiblingPosition = PointPositions[SiblingPoint];
+								PointPositions[SiblingPoint] += PointOffset;
+
+								const TArray<int32> BranchChildren = BranchFacade.GetImmediateChildren(BranchIndex);
+								for (int32 ChildBranch : BranchChildren)
+								{
+									const int32 ChildBranchIndex = BranchFacade.GetBranchNumbers().Find(ChildBranch);
+									if (ChildBranchIndex != INDEX_NONE)
+									{
+										const int32 RootPoint = BranchFacade.GetRootPoint(ChildBranchIndex);
+										const float ChildRootPointLFR = PointFacade.GetLengthFromRoot(RootPoint);
+										if (ChildRootPointLFR >= SiblingPointLFR && ChildRootPointLFR < NextSiblingPointLFR)
+										{
+											ApplyOffsetToBranchSuccessors(ChildBranchIndex, PointOffset);
+										}
+									}
+								}
+							}
+
+
+							for (int32 SiblingPointIndex = BranchPointIndex; SiblingPointIndex < BranchPoints.Num(); ++SiblingPointIndex)
+							{
+								PointPositions[BranchPoints[SiblingPointIndex]] += Offset;
+							}
+
+							const FQuat4f Rotation = FQuat4f::FindBetween(FVector3f(SegmentVector), FVector3f(SegmentVector) + Offset);
+							for (const int32 FoliageEntryId : FoliageFacade.GetFoliageEntryIdsForBranch(BranchIndex))
+							{
+								const float FoliageLFR = FoliageFacade.GetLengthFromRoot(FoliageEntryId);
+								if (FoliageLFR > PointLFR)
+								{
+									FoliagePositions[FoliageEntryId] += Offset;
+									if (FoliageLFR < NextPointLFR)
+									{
+										FoliageFacade.SetUpVector(
+											FoliageEntryId,
+											Rotation.RotateVector(FoliageFacade.GetUpVector(FoliageEntryId))
+										);
+										FoliageFacade.SetNormalVector(
+											FoliageEntryId,
+											Rotation.RotateVector(FoliageFacade.GetNormalVector(FoliageEntryId))
+										);
+									}
+								}
+							}
+							for (FVector3f& BudDirection : BudDirections[NextPointIndex])
+							{
+								BudDirection = Rotation.RotateVector(BudDirection);
+							}
+
+							const TArray<int32> BranchChildren = BranchFacade.GetImmediateChildren(BranchIndex);
+							for (int32 ChildBranch : BranchChildren)
+							{
+								const int32 ChildBranchIndex = BranchFacade.GetBranchNumbers().Find(ChildBranch);
+								if (ChildBranchIndex != INDEX_NONE)
+								{
+									const int32 RootPoint = BranchFacade.GetRootPoint(ChildBranchIndex);
+									if (PointFacade.GetLengthFromRoot(RootPoint) > PointLFR)
+									{
+										ApplyOffsetToBranchSuccessors(ChildBranchIndex, Offset);
+									}
+								}
+							}
+						}
+					}
+				}
+#endif
+			}
+
+			UPVGrowthData* OutManagedArrayCollectionData = FPCGContext::NewObject_AnyThread<UPVGrowthData>(InContext);
+			OutManagedArrayCollectionData->Initialize(MoveTemp(SkeletonCollection));
+			InContext->OutputData.TaggedData.Emplace(OutManagedArrayCollectionData);
+		}
+	}
+
+	return true;
+}
diff --git a/Engine/Plugins/Experimental/ProceduralVegetationEditor/Source/ProceduralVegetation/ProceduralVegetation.Build.cs b/Engine/Plugins/Experimental/ProceduralVegetationEditor/Source/ProceduralVegetation/ProceduralVegetation.Build.cs
--- a/Engine/Plugins/Experimental/ProceduralVegetationEditor/Source/ProceduralVegetation/ProceduralVegetation.Build.cs
+++ b/Engine/Plugins/Experimental/ProceduralVegetationEditor/Source/ProceduralVegetation/ProceduralVegetation.Build.cs
@@ -7,7 +7,7 @@ namespace UnrealBuildTool.Rules
 		public ProceduralVegetation(ReadOnlyTargetRules Target) : base(Target)
 		{
 			PCHUsage = PCHUsageMode.UseExplicitOrSharedPCHs;
-			
+
 			PublicDependencyModuleNames.AddRange(
 				new string[]
 				{
@@ -17,18 +17,21 @@ namespace UnrealBuildTool.Rules
 			PrivateDependencyModuleNames.AddRange(
 				new string[]
 				{
+					"Chaos",
 					"Core",
 					"CoreUObject",
-					"Chaos",
 					"Engine",
-					"RenderCore",
 					"GeometryCore",
 					"GeometryFramework",
 					"Json",
 					"JsonUtilities",
+					"MeshDescription",
+					"ModelingComponents",
 					"PCG",
+					"RenderCore",
+					"StaticMeshDescription"
 				}
 			);
 		}
 	}
-}
+}
\ No newline at end of file
diff --git a/Engine/Plugins/Experimental/ProceduralVegetationEditor/Source/ProceduralVegetation/Public/Facades/PVBranchFacade.h b/Engine/Plugins/Experimental/ProceduralVegetationEditor/Source/ProceduralVegetation/Public/Facades/PVBranchFacade.h
--- a/Engine/Plugins/Experimental/ProceduralVegetationEditor/Source/ProceduralVegetation/Public/Facades/PVBranchFacade.h
+++ b/Engine/Plugins/Experimental/ProceduralVegetationEditor/Source/ProceduralVegetation/Public/Facades/PVBranchFacade.h
@@ -27,10 +27,14 @@ namespace PV::Facades
 
 		const TArray<int32>& GetPoints(const int32 Index) const;
 
+		int32 GetRootPoint(const int32 Index) const;
+
 		void SetPoints(const int32 Index, const TArray<int32>& InPoints);
 
 		const TArray<int32>& GetChildren(const int32 Index) const;
 
+		TArray<int32> GetImmediateChildren(const int32 Index) const;
+
 		void SetChildren(const int32 Index, const TArray<int32>& InChildren);
 
 		const TArray<int32>& GetParents(const int32 Index) const;
diff --git a/Engine/Plugins/Experimental/ProceduralVegetationEditor/Source/ProceduralVegetation/Public/Facades/PVFoliageFacade.h b/Engine/Plugins/Experimental/ProceduralVegetationEditor/Source/ProceduralVegetation/Public/Facades/PVFoliageFacade.h
--- a/Engine/Plugins/Experimental/ProceduralVegetationEditor/Source/ProceduralVegetation/Public/Facades/PVFoliageFacade.h
+++ b/Engine/Plugins/Experimental/ProceduralVegetationEditor/Source/ProceduralVegetation/Public/Facades/PVFoliageFacade.h
@@ -36,7 +36,7 @@ namespace PV::Facades
 
 		bool IsValid() const;
 
-		int32 NumFoliageEntries() const { return NameIdsAttribute.Num(); };
+		int32 NumFoliageEntries() const { return NameIdsAttribute.Num(); }
 
 		FFoliageEntryData GetFoliageEntry(const int32 Index) const;
 
@@ -95,6 +95,8 @@ namespace PV::Facades
 		FTransform GetFoliageTransform(const int32 Id) const;
 
 		const TManagedArray<FVector3f>& GetPivotPositions() const;
+		
+		TManagedArray<FVector3f>& ModifyPivotPositions();
 
 	protected:
 		void DefineSchema(const int32 InitialSize = 0);
diff --git a/Engine/Plugins/Experimental/ProceduralVegetationEditor/Source/ProceduralVegetation/Public/Nodes/PVObjectAvoidance.h b/Engine/Plugins/Experimental/ProceduralVegetationEditor/Source/ProceduralVegetation/Public/Nodes/PVObjectAvoidance.h
new file mode 100644
--- /dev/null
+++ b/Engine/Plugins/Experimental/ProceduralVegetationEditor/Source/ProceduralVegetation/Public/Nodes/PVObjectAvoidance.h
@@ -0,0 +1,76 @@
+// Copyright Epic Games, Inc. All Rights Reserved.
+
+#pragma once
+
+#include "Nodes/PVBaseSettings.h"
+
+#include "PVObjectAvoidance.generated.h"
+
+USTRUCT()
+struct FPVColliderParams
+{
+	GENERATED_BODY()
+
+	UPROPERTY(EditAnywhere, Category = "Collider")
+	TSoftObjectPtr<UStaticMesh> Mesh;
+
+	UPROPERTY(EditAnywhere, Category = "Collider")
+	FTransform Transform;
+};
+
+USTRUCT()
+struct FPVObjectAvoidanceParams
+{
+	GENERATED_BODY()
+
+	UPROPERTY(EditAnywhere, Category="Object Avoidance", meta=(ClampMin=0, ClampMax=10, UIMin=0, UIMax=10))
+	int32 Smoothness = 0;
+
+	UPROPERTY(EditAnywhere, Category = "Colliders")
+	TArray<FPVColliderParams> Colliders;
+};
+
+UCLASS(BlueprintType, ClassGroup = (Procedural))
+class UPVObjectAvoidanceSettings : public UPVBaseSettings
+{
+	GENERATED_BODY()
+
+public:
+#if WITH_EDITOR
+	virtual FName GetDefaultNodeName() const override
+	{
+		return FName(TEXT("Object Avoidance"));
+	}
+
+	virtual FText GetNodeTooltipText() const override
+	{
+		return FText::FromString("Object Avoidance");
+	}
+
+	virtual FLinearColor GetNodeTitleColor() const override
+	{
+		return FLinearColor::Yellow;
+	}
+
+	virtual TArray<EPVRenderType> GetDefaultRenderType() const override
+	{
+		return TArray{EPVRenderType::PointData};
+	}
+#endif
+
+protected:
+	virtual FPCGDataTypeIdentifier GetInputPinTypeIdentifier() const override;
+	virtual FPCGDataTypeIdentifier GetOutputPinTypeIdentifier() const override;
+
+	virtual FPCGElementPtr CreateElement() const override;
+
+public:
+	UPROPERTY(EditAnywhere, Category = "Object Avoidance", meta = (ShowOnlyInnerProperties))
+	FPVObjectAvoidanceParams ObjectAvoidanceSettings;
+};
+
+class FPVObjectAvoidanceElement : public FPVBaseElement
+{
+protected:
+	virtual bool ExecuteInternal(FPCGContext* InContext) const override;
+};
diff --git a/Engine/Plugins/Experimental/ProceduralVegetationEditor/Source/ProceduralVegetationEditor/Private/Visualizations/PVDataVisualization.cpp b/Engine/Plugins/Experimental/ProceduralVegetationEditor/Source/ProceduralVegetationEditor/Private/Visualizations/PVDataVisualization.cpp
--- a/Engine/Plugins/Experimental/ProceduralVegetationEditor/Source/ProceduralVegetationEditor/Private/Visualizations/PVDataVisualization.cpp
+++ b/Engine/Plugins/Experimental/ProceduralVegetationEditor/Source/ProceduralVegetationEditor/Private/Visualizations/PVDataVisualization.cpp
@@ -34,6 +34,7 @@
 #include "Materials/Material.h"
 
 #include "Nodes/PVBaseSettings.h"
+#include "Nodes/PVObjectAvoidance.h"
 
 #include "Rendering/SkeletalMeshRenderData.h"
 
@@ -111,6 +112,19 @@ FPCGSetupSceneFunc FPVDataVisualization::GetViewportSetupFunc(const UPCGSettings
 			}
 		}
 
+		if (const UPVObjectAvoidanceSettings* ObjectAvoidanceSettings = Cast<UPVObjectAvoidanceSettings>(WeakSettings))
+		{
+			for (const auto& [Mesh, Transform] : ObjectAvoidanceSettings->ObjectAvoidanceSettings.Colliders)
+			{
+				TObjectPtr<UStaticMeshComponent> StaticMeshComponent = NewObject<UStaticMeshComponent>(GetTransientPackage(), NAME_None, RF_Transient);
+				StaticMeshComponent->SetStaticMesh(Mesh.LoadSynchronous());
+				StaticMeshComponent->UpdateBounds();
+
+				InOutParams.ManagedResources.Add(StaticMeshComponent);
+				InOutParams.Scene->AddComponent(StaticMeshComponent, Transform);
+			}
+		}
+
 		if (PV::Utilities::DebugModeEnabled())
 		{
 			if (WeakSettings.IsValid() && WeakSettings->bDebug)
diff --git a/Engine/Plugins/Experimental/ProceduralVegetationEditor/Source/ProceduralVegetationEditor/Private/Visualizations/PVLineBatchComponent.cpp b/Engine/Plugins/Experimental/ProceduralVegetationEditor/Source/ProceduralVegetationEditor/Private/Visualizations/PVLineBatchComponent.cpp
--- a/Engine/Plugins/Experimental/ProceduralVegetationEditor/Source/ProceduralVegetationEditor/Private/Visualizations/PVLineBatchComponent.cpp
+++ b/Engine/Plugins/Experimental/ProceduralVegetationEditor/Source/ProceduralVegetationEditor/Private/Visualizations/PVLineBatchComponent.cpp
@@ -97,12 +97,12 @@ void UPVLineBatchComponent::Draw(const FSceneView* View, FPrimitiveDrawInterface
 
 		if (Line.PointDrawSettings == EPointDrawSettings::Start || Line.PointDrawSettings == EPointDrawSettings::Both)
 		{
-			PDI->DrawPoint(Line.StartPos, FLinearColor::White, PointSize, SDPG_Foreground);
+			PDI->DrawPoint(Line.StartPos, FLinearColor::White, PointSize, Line.DepthPriorityGroup);
 		}
 
 		if (Line.PointDrawSettings == EPointDrawSettings::End || Line.PointDrawSettings == EPointDrawSettings::Both)
 		{
-			PDI->DrawPoint(Line.EndPos, FLinearColor::White, PointSize, SDPG_Foreground);
+			PDI->DrawPoint(Line.EndPos, FLinearColor::White, PointSize, Line.DepthPriorityGroup);
 		}
 	}
 }
diff --git a/Engine/Plugins/Experimental/ProceduralVegetationEditor/Source/ProceduralVegetationEditor/Private/Visualizations/PVSkeletonVisualizerComponent.cpp b/Engine/Plugins/Experimental/ProceduralVegetationEditor/Source/ProceduralVegetationEditor/Private/Visualizations/PVSkeletonVisualizerComponent.cpp
--- a/Engine/Plugins/Experimental/ProceduralVegetationEditor/Source/ProceduralVegetationEditor/Private/Visualizations/PVSkeletonVisualizerComponent.cpp
+++ b/Engine/Plugins/Experimental/ProceduralVegetationEditor/Source/ProceduralVegetationEditor/Private/Visualizations/PVSkeletonVisualizerComponent.cpp
@@ -56,7 +56,7 @@ void UPVSkeletonVisualizerComponent::SetCollection(const FManagedArrayCollection
 		{
 			FVector CurrentPosition = FVector(PointsPositions[BranchPoints[BranchPointIndex]]);
 
-			AddLine(PreviousPosition, CurrentPosition, FLinearColor::White, ESceneDepthPriorityGroup::SDPG_MAX, EPointDrawSettings::Start);
+			AddLine(PreviousPosition, CurrentPosition, FLinearColor::White, ESceneDepthPriorityGroup::SDPG_World, EPointDrawSettings::Start);
 
 			PreviousPosition = CurrentPosition;
 		}
