Extend & Clean up PVE Facades APIs

* Add Getters and Setters for attributes that needs setting by the PVE importer
* Remove const& on return values which are faster to copy
* Make use of const on parameters consistent
* Convert const TArray<>& to TArray<> to allow for use of MoveTemp
* Return const& instead of copy where it makes sense to do so

diff --git a/Engine/Plugins/Experimental/ProceduralVegetationEditor/Source/ProceduralVegetation/Private/Facades/PVBranchFacade.cpp b/Engine/Plugins/Experimental/ProceduralVegetationEditor/Source/ProceduralVegetation/Private/Facades/PVBranchFacade.cpp
--- a/Engine/Plugins/Experimental/ProceduralVegetationEditor/Source/ProceduralVegetation/Private/Facades/PVBranchFacade.cpp
+++ b/Engine/Plugins/Experimental/ProceduralVegetationEditor/Source/ProceduralVegetation/Private/Facades/PVBranchFacade.cpp
@@ -56,7 +56,7 @@ namespace PV::Facades
 		return Parents.Num();
 	}
 
-	const TArray<int32>& FBranchFacade::GetPoints(const int32 Index) const
+	const TArray<int32>& FBranchFacade::GetPoints(int32 Index) const
 	{
 		if (BranchPoints.IsValid() && BranchPoints.IsValidIndex(Index))
 		{
@@ -67,21 +67,21 @@ namespace PV::Facades
 		return EmptyPoints;
 	}
 
-	int32 FBranchFacade::GetRootPoint(const int32 Index) const
+	int32 FBranchFacade::GetRootPoint(int32 Index) const
 	{
 		const TArray<int32>& Points = GetPoints(Index);
 		return Points.Num() ? Points[0] : INDEX_NONE;
 	}
 
-	void FBranchFacade::SetPoints(const int32 Index, const TArray<int32>& InPoints)
+	void FBranchFacade::SetPoints(int32 Index, TArray<int32> InPoints)
 	{
 		if (BranchPoints.IsValid() && BranchPoints.IsValidIndex(Index))
 		{
-			BranchPoints.ModifyAt(Index, InPoints);
+			BranchPoints.Modify()[Index] = MoveTemp(InPoints);
 		}
 	}
 
-	const TArray<int32>& FBranchFacade::GetChildren(const int32 Index) const
+	const TArray<int32>& FBranchFacade::GetChildren(int32 Index) const
 	{
 		if (Children.IsValid() && Children.IsValidIndex(Index))
 		{
@@ -92,7 +92,7 @@ namespace PV::Facades
 		return EmptyChildren;
 	}
 
-	TArray<int32> FBranchFacade::GetImmediateChildren(const int32 Index) const
+	TArray<int32> FBranchFacade::GetImmediateChildren(int32 Index) const
 	{
 		const int32 ParentBranchNumber = GetBranchNumber(Index);
 		return GetChildren(Index).FilterByPredicate(
@@ -108,15 +108,15 @@ namespace PV::Facades
 		);
 	}
 
-	void FBranchFacade::SetChildren(const int32 Index, const TArray<int32>& InChildren)
+	void FBranchFacade::SetChildren(int32 Index, TArray<int32> InChildren)
 	{
 		if (Children.IsValid() && Children.IsValidIndex(Index))
 		{
-			Children.ModifyAt(Index, InChildren);
+			Children.Modify()[Index] = MoveTemp(InChildren);
 		}
 	}
 
-	const TArray<int32>& FBranchFacade::GetParents(const int32 Index) const
+	const TArray<int32>& FBranchFacade::GetParents(int32 Index) const
 	{
 		if (Parents.IsValid() && Parents.IsValidIndex(Index))
 		{
@@ -127,7 +127,7 @@ namespace PV::Facades
 		return EmptyParents;
 	}
 
-	int32 FBranchFacade::GetParentIndex(const int32 BranchIndex) const
+	int32 FBranchFacade::GetParentIndex(int32 BranchIndex) const
 	{
 		if (BranchParentNumbers.IsValid() && BranchParentNumbers.IsValidIndex(BranchIndex))
 		{
@@ -137,7 +137,7 @@ namespace PV::Facades
 		return INDEX_NONE;
 	}
 
-	int32 FBranchFacade::GetParentBranchNumber(const int32 BranchIndex) const
+	int32 FBranchFacade::GetParentBranchNumber(int32 BranchIndex) const
 	{
 		if (BranchParentNumbers.IsValid() && BranchParentNumbers.IsValidIndex(BranchIndex))
 		{
@@ -147,15 +147,23 @@ namespace PV::Facades
 		return INDEX_NONE;
 	}
 
-	void FBranchFacade::SetParents(const int32 Index, const TArray<int32>& InParents)
+	void FBranchFacade::SetParentBranchNumber(int32 BranchIndex, int32 InParentBranchNumber)
+	{
+		if (BranchParentNumbers.IsValid() && BranchParentNumbers.IsValidIndex(BranchIndex))
+		{
+			BranchParentNumbers.ModifyAt(BranchIndex, InParentBranchNumber);
+		}
+	}
+
+	void FBranchFacade::SetParents(int32 Index, TArray<int32> InParents)
 	{
 		if (Parents.IsValid() && Parents.IsValidIndex(Index))
 		{
-			Parents.ModifyAt(Index, InParents);
+			Parents.Modify()[Index] = MoveTemp(InParents);
 		}
 	}
 
-	int32 FBranchFacade::GetBranchNumber(const int32 Index) const
+	int32 FBranchFacade::GetBranchNumber(int32 Index) const
 	{
 		if (BranchNumbers.IsValid() && BranchNumbers.IsValidIndex(Index))
 		{
@@ -165,7 +173,32 @@ namespace PV::Facades
 		return INDEX_NONE;
 	}
 
-	int32 FBranchFacade::GetBranchSourceBudNumber(const int32 Index) const
+	void FBranchFacade::SetBranchNumber(int32 Index, int32 InBranchNumber)
+	{
+		if (BranchNumbers.IsValid() && BranchNumbers.IsValidIndex(Index))
+		{
+			BranchNumbers.ModifyAt(Index, InBranchNumber);
+		}
+	}
+
+	int32 FBranchFacade::GetBranchPlantNumber(int32 Index) const
+	{
+		if (PlantNumbers.IsValid() && PlantNumbers.IsValidIndex(Index))
+		{
+			return PlantNumbers[Index];
+		}
+		return INDEX_NONE;
+	}
+
+	void FBranchFacade::SetBranchPlantNumber(int32 Index, int32 InPlantNumber)
+	{
+		if (PlantNumbers.IsValid() && PlantNumbers.IsValidIndex(Index))
+		{
+			PlantNumbers.ModifyAt(Index, InPlantNumber);
+		}
+	}
+
+	int32 FBranchFacade::GetBranchSourceBudNumber(int32 Index) const
 	{
 		if (BranchSourceBudNumber.IsValid() && BranchSourceBudNumber.IsValidIndex(Index))
 		{
@@ -175,7 +208,15 @@ namespace PV::Facades
 		return INDEX_NONE;
 	}
 
-	TArray<int32> FBranchFacade::GetParentBranchIndices(int BranchIndex) const
+	void FBranchFacade::SetBranchSourceBudNumber(int32 Index, int32 InSourceBudNumber)
+	{
+		if (BranchSourceBudNumber.IsValid() && BranchSourceBudNumber.IsValidIndex(Index))
+		{
+			BranchSourceBudNumber.ModifyAt(Index, InSourceBudNumber);
+		}
+	}
+
+	TArray<int32> FBranchFacade::GetParentBranchIndices(int32 BranchIndex) const
 	{
 		TArray<int32> ParentIndices;
 		if (IsValid() && Parents.Num() > 0)
@@ -190,7 +231,7 @@ namespace PV::Facades
 		return ParentIndices;
 	}
 
-	int32 FBranchFacade::GetParentBranchIndex(const int32 BranchIndex) const
+	int32 FBranchFacade::GetParentBranchIndex(int32 BranchIndex) const
 	{
 		if (IsValid() && Parents.IsValidIndex(BranchIndex))
 		{
@@ -205,7 +246,7 @@ namespace PV::Facades
 		return INDEX_NONE;
 	}
 
-	int32 FBranchFacade::GetHierarchyGenerationNumber(const int32 Index) const
+	int32 FBranchFacade::GetHierarchyGenerationNumber(int32 Index) const
 	{
 		// The generation attributes contain incorrect information at the moment 
 		// thus the number of parents can be used to derive the correct number
@@ -222,7 +263,7 @@ namespace PV::Facades
 		return BranchNumbers.Get();
 	}
 
-	int32 FBranchFacade::GetBranchUVMaterial(const int32 Index) const
+	int32 FBranchFacade::GetBranchUVMaterial(int32 Index) const
 	{
 		if (BranchUVMaterial.IsValid() && BranchUVMaterial.IsValidIndex(Index))
 		{
@@ -232,7 +273,7 @@ namespace PV::Facades
 		return INDEX_NONE;
 	}
 
-	void FBranchFacade::SetBranchUVMaterial(const int32 Index, const int32 InMaterial)
+	void FBranchFacade::SetBranchUVMaterial(int32 Index, int32 InMaterial)
 	{
 		if (!BranchUVMaterial.IsValid())
 		{
@@ -244,7 +285,7 @@ namespace PV::Facades
 		}
 	}
 
-	int32 FBranchFacade::GetBranchHierarchyNumber(const int32 Index) const
+	int32 FBranchFacade::GetBranchHierarchyNumber(int32 Index) const
 	{
 		if (BranchHierarchyNumber.IsValid() && BranchHierarchyNumber.IsValidIndex(Index))
 		{
@@ -254,7 +295,15 @@ namespace PV::Facades
 		return INDEX_NONE;
 	}
 
-	int32 FBranchFacade::GetBranchSimulationGroupIndex(const int32 Index) const
+	void FBranchFacade::SetBranchHierarchyNumber(int32 Index, int32 InBranchHierarchyNumber)
+	{
+		if (BranchHierarchyNumber.IsValid() && BranchHierarchyNumber.IsValidIndex(Index))
+		{
+			BranchHierarchyNumber.ModifyAt(Index, InBranchHierarchyNumber);
+		}
+	}
+
+	int32 FBranchFacade::GetBranchSimulationGroupIndex(int32 Index) const
 	{
 		if (BranchSimulationGroupIndex.IsValid() && BranchSimulationGroupIndex.IsValidIndex(Index))
 		{
@@ -264,7 +313,7 @@ namespace PV::Facades
 		return INDEX_NONE;
 	}
 
-	void FBranchFacade::SetBranchSimulationGroupIndex(const int32 Index, const int32 InSimulationGroupIndex)
+	void FBranchFacade::SetBranchSimulationGroupIndex(int32 Index, int32 InSimulationGroupIndex)
 	{
 		if (!BranchSimulationGroupIndex.IsValid())
 		{
@@ -304,7 +353,7 @@ namespace PV::Facades
 		return FString();
 	}
 
-	void FBranchFacade::SetTrunkURange(const TArray<FVector2f>& InURange)
+	void FBranchFacade::SetTrunkURange(TArray<FVector2f> InURange)
 	{
 		check(!IsConst());
 
@@ -319,7 +368,7 @@ namespace PV::Facades
 			TrunkURangeAttribute.Add();
 		}
 
-		TrunkURangeAttribute.Modify()[0] = InURange;
+		TrunkURangeAttribute.Modify()[0] = MoveTemp(InURange);
 	}
 
 	const TArray<FVector2f>& FBranchFacade::GetTrunkURange() const
@@ -333,7 +382,7 @@ namespace PV::Facades
 		return TrunkURange;
 	}
 
-	void FBranchFacade::CopyEntry(const int32 FromIndex, const int32 ToIndex)
+	void FBranchFacade::CopyEntry(int32 FromIndex, int32 ToIndex)
 	{
 		if (IsValid() && BranchNumbers.IsValidIndex(FromIndex) && BranchNumbers.IsValidIndex(ToIndex))
 		{
@@ -357,7 +406,7 @@ namespace PV::Facades
 		}
 	}
 
-	void FBranchFacade::RemoveEntries(const int32 NumEntries, const int32 StartIndex)
+	void FBranchFacade::RemoveEntries(int32 NumEntries, int32 StartIndex)
 	{
 		if (IsValid() && BranchNumbers.IsValidIndex(StartIndex) && StartIndex + NumEntries <= BranchNumbers.Num())
 		{
@@ -386,9 +435,9 @@ namespace PV::Facades
 			});
 	}
 
-	int32 FBranchFacade::GetBranchIndexFromPointIndex(const int32 PointIndex) const
+	int32 FBranchFacade::GetBranchIndexFromPointIndex(int32 PointIndex) const
 	{
-		for (int BranchIndex = 0; BranchIndex < GetElementCount(); BranchIndex++)
+		for (int32 BranchIndex = 0; BranchIndex < GetElementCount(); BranchIndex++)
 		{
 			TArray<int32> PointInexes = GetPoints(BranchIndex);
 			if (PointInexes.Contains(PointIndex))
diff --git a/Engine/Plugins/Experimental/ProceduralVegetationEditor/Source/ProceduralVegetation/Private/Facades/PVFoliageFacade.cpp b/Engine/Plugins/Experimental/ProceduralVegetationEditor/Source/ProceduralVegetation/Private/Facades/PVFoliageFacade.cpp
--- a/Engine/Plugins/Experimental/ProceduralVegetationEditor/Source/ProceduralVegetation/Private/Facades/PVFoliageFacade.cpp
+++ b/Engine/Plugins/Experimental/ProceduralVegetationEditor/Source/ProceduralVegetation/Private/Facades/PVFoliageFacade.cpp
@@ -9,7 +9,7 @@
 
 namespace PV::Facades
 {
-	FFoliageFacade::FFoliageFacade(FManagedArrayCollection& InCollection, const int32 InitialSize)
+	FFoliageFacade::FFoliageFacade(FManagedArrayCollection& InCollection, int32 InitialSize)
 		: ConstCollection(InCollection)
 		, Collection(&InCollection)
 		, NamesAttribute(InCollection, AttributeNames::FoliageName, GroupNames::FoliageNamesGroup)
@@ -44,7 +44,7 @@ namespace PV::Facades
 	{
 	}
 
-	void FFoliageFacade::DefineSchema(const int32 InitialSize)
+	void FFoliageFacade::DefineSchema(int32 InitialSize)
 	{
 		check(!IsConst());
 
@@ -87,7 +87,7 @@ namespace PV::Facades
 			FoliageIdsAttribute.IsValid() && ParentBoneIdsAttribute.IsValid();
 	}
 
-	FFoliageEntryData FFoliageFacade::GetFoliageEntry(const int32 Index) const
+	FFoliageEntryData FFoliageFacade::GetFoliageEntry(int32 Index) const
 	{
 		FFoliageEntryData ReturnData = {};
 		if (IsValid() && Index > -1)
@@ -128,7 +128,7 @@ namespace PV::Facades
 		return ReturnData;
 	}
 
-	const float& FFoliageFacade::GetLengthFromRoot(const int32 Index) const
+	float FFoliageFacade::GetLengthFromRoot(int32 Index) const
 	{
 		if (LengthFromRootAttribute.IsValid() && LengthFromRootAttribute.IsValidIndex(Index))
 		{
@@ -139,7 +139,7 @@ namespace PV::Facades
 		return DefaultLengthFromRoot;
 	}
 
-	void FFoliageFacade::SetLengthFromRoot(const int32 Index, const float Input)
+	void FFoliageFacade::SetLengthFromRoot(int32 Index, float Input)
 	{
 		check(!IsConst());
 		if (IsValid() && LengthFromRootAttribute.IsValidIndex(Index))
@@ -148,7 +148,7 @@ namespace PV::Facades
 		}
 	}
 
-	int32 FFoliageFacade::GetParentBoneID(const int32 Index) const
+	int32 FFoliageFacade::GetParentBoneID(int32 Index) const
 	{
 		if (ParentBoneIdsAttribute.IsValid() && ParentBoneIdsAttribute.IsValidIndex(Index))
 		{
@@ -158,7 +158,7 @@ namespace PV::Facades
 		return INDEX_NONE;
 	}
 
-	void FFoliageFacade::SetParentBoneID(const int32 Index, const int32 Input)
+	void FFoliageFacade::SetParentBoneID(int32 Index, int32 Input)
 	{
 		check(!IsConst());
 		if (IsValid() && ParentBoneIdsAttribute.IsValidIndex(Index))
@@ -167,7 +167,7 @@ namespace PV::Facades
 		}
 	}
 
-	const TArray<int32>& FFoliageFacade::GetFoliageEntryIdsForBranch(const int32 Index) const
+	const TArray<int32>& FFoliageFacade::GetFoliageEntryIdsForBranch(int32 Index) const
 	{
 		if (FoliageIdsAttribute.IsValid() && FoliageIdsAttribute.IsValidIndex(Index))
 		{
@@ -198,7 +198,7 @@ namespace PV::Facades
 		return INDEX_NONE;
 	}
 
-	void FFoliageFacade::SetFoliageEntry(const int32 Index, const FFoliageEntryData& InputData)
+	void FFoliageFacade::SetFoliageEntry(int32 Index, const FFoliageEntryData& InputData)
 	{
 		check(!IsConst());
 		if (IsValid() && Index >= 0)
@@ -216,16 +216,25 @@ namespace PV::Facades
 		}
 	}
 
-	void FFoliageFacade::SetFoliageIdsArray(const int32 Index, const TArray<int32>& InputIds)
+	void FFoliageFacade::SetFoliageIdsArray(int32 Index, TArray<int32> InputIds)
 	{
 		check(!IsConst());
 		if (IsValid() && FoliageIdsAttribute.IsValidIndex(Index))
 		{
-			FoliageIdsAttribute.Modify()[Index] = InputIds;
+			FoliageIdsAttribute.Modify()[Index] = MoveTemp(InputIds);
 		}
 	}
 
-	void FFoliageFacade::SetFoliageBranchId(const int32 Index, const int32 InputId)
+	int32 FFoliageFacade::GetFoliageBranchId(int32 Index) const
+	{
+		if (BranchIdsAttribute.IsValid() && BranchIdsAttribute.IsValidIndex(Index))
+		{
+			return BranchIdsAttribute[Index];
+		}
+		return INDEX_NONE;
+	}
+
+	void FFoliageFacade::SetFoliageBranchId(int32 Index, int32 InputId)
 	{
 		check(!IsConst());
 		if (IsValid() && BranchIdsAttribute.IsValidIndex(Index))
@@ -234,31 +243,29 @@ namespace PV::Facades
 		}
 	}
 
-	FString FFoliageFacade::GetFoliageName(const int32 Index) const
+	const FString& FFoliageFacade::GetFoliageName(int32 Index) const
 	{
 		if (NamesAttribute.IsValid() && NamesAttribute.IsValidIndex(Index))
 		{
 			return NamesAttribute[Index];
 		}
 
-		return FString();
+		const static FString DefaultFoliageName = FString();
+		return DefaultFoliageName;
 	}
 
-	TArray<FString> FFoliageFacade::GetFoliageNames() const
+	const TArray<FString>& FFoliageFacade::GetFoliageNames() const
 	{
-		TArray<FString> FoliageNames;
 		if (NamesAttribute.IsValid())
 		{
-			for (int32 i = 0; i < NamesAttribute.Num(); i++)
-			{
-				FoliageNames.AddUnique(NamesAttribute[i]);
-			}
+			return NamesAttribute.Get().GetConstArray();
 		}
 
-		return FoliageNames;
+		const static TArray<FString> DefaultFoliageNames;
+		return DefaultFoliageNames;
 	}
 
-	const FVector3f& FFoliageFacade::GetPivotPoint(const int32 Index) const
+	const FVector3f& FFoliageFacade::GetPivotPoint(int32 Index) const
 	{
 		if (PivotPointsAttribute.IsValid() && PivotPointsAttribute.IsValidIndex(Index))
 		{
@@ -291,7 +298,6 @@ namespace PV::Facades
 		return DefaultNormalVector;
 	}
 
-
 	void FFoliageFacade::SetPivotPoint(const int32 Index, const FVector3f& Input)
 	{
 		check(!IsConst());
@@ -301,7 +307,7 @@ namespace PV::Facades
 		}
 	}
 
-	void FFoliageFacade::SetUpVector(const int32 Index, const FVector3f& Input)
+	void FFoliageFacade::SetUpVector(int32 Index, const FVector3f& Input)
 	{
 		check(!IsConst());
 		if (IsValid() && UpVectorsAttribute.IsValidIndex(Index))
@@ -310,7 +316,7 @@ namespace PV::Facades
 		}
 	}
 
-	void FFoliageFacade::SetNormalVector(const int32 Index, const FVector3f& Input)
+	void FFoliageFacade::SetNormalVector(int32 Index, const FVector3f& Input)
 	{
 		check(!IsConst());
 		if (IsValid() && NormalVectorsAttribute.IsValidIndex(Index))
@@ -319,7 +325,7 @@ namespace PV::Facades
 		}
 	}
 
-	void FFoliageFacade::SetScale(const int32 Index, const float& Input)
+	void FFoliageFacade::SetScale(int32 Index, float Input)
 	{
 		check(!IsConst());
 		if (IsValid() && ScalesAttribute.IsValidIndex(Index))
@@ -328,7 +334,7 @@ namespace PV::Facades
 		}
 	}
 
-	void FFoliageFacade::SetFoliageName(const int32 Index, const FString& Input)
+	void FFoliageFacade::SetFoliageName(int32 Index, const FString& Input)
 	{
 		check(!IsConst());
 		if (IsValid() && NamesAttribute.IsValidIndex(Index))
@@ -348,18 +354,18 @@ namespace PV::Facades
 		}
 	}
 
-	void FFoliageFacade::SetFoliagePath(const FString InPath)
+	void FFoliageFacade::SetFoliagePath(FString InPath)
 	{
 		check(!IsConst());
 
-		int32 NumElements = Collection->NumElements(GroupNames::DetailsGroup);
+		const int32 NumElements = Collection->NumElements(GroupNames::DetailsGroup);
 
 		if (NumElements == 0)
 		{
 			Collection->AddElements(1, GroupNames::DetailsGroup);
 		}
 
-		FoliagePathAttribute.Modify()[0] = InPath;
+		FoliagePathAttribute.Modify()[0] = MoveTemp(InPath);
 	}
 
 	FString FFoliageFacade::GetFoliagePath() const
@@ -377,7 +383,7 @@ namespace PV::Facades
 		return NameIdsAttribute.Num();
 	}
 
-	void FFoliageFacade::CopyEntry(const int32 FromIndex, const int32 ToIndex)
+	void FFoliageFacade::CopyEntry(int32 FromIndex, int32 ToIndex)
 	{
 		if (IsValid() && NameIdsAttribute.IsValidIndex(FromIndex) && NameIdsAttribute.IsValidIndex(ToIndex))
 		{
@@ -392,7 +398,7 @@ namespace PV::Facades
 		}
 	}
 
-	void FFoliageFacade::RemoveEntries(const int32 NumEntries, const int32 StartIndex)
+	void FFoliageFacade::RemoveEntries(int32 NumEntries, int32 StartIndex)
 	{
 		if (IsValid() && NameIdsAttribute.IsValidIndex(StartIndex) && StartIndex + NumEntries <= NameIdsAttribute.Num())
 		{
@@ -400,7 +406,7 @@ namespace PV::Facades
 		}
 	}
 
-	FTransform FFoliageFacade::GetFoliageTransform(const int32 Id) const
+	FTransform FFoliageFacade::GetFoliageTransform(int32 Id) const
 	{
 		const FFoliageEntryData Data = GetFoliageEntry(Id);
 
diff --git a/Engine/Plugins/Experimental/ProceduralVegetationEditor/Source/ProceduralVegetation/Private/Facades/PVPointFacade.cpp b/Engine/Plugins/Experimental/ProceduralVegetationEditor/Source/ProceduralVegetation/Private/Facades/PVPointFacade.cpp
--- a/Engine/Plugins/Experimental/ProceduralVegetationEditor/Source/ProceduralVegetation/Private/Facades/PVPointFacade.cpp
+++ b/Engine/Plugins/Experimental/ProceduralVegetationEditor/Source/ProceduralVegetation/Private/Facades/PVPointFacade.cpp
@@ -76,7 +76,7 @@ namespace PV::Facades
 		return Positions.Num();
 	}
 
-	const FVector3f& FPointFacade::GetPosition(const int32 Index) const
+	const FVector3f& FPointFacade::GetPosition(int32 Index) const
 	{
 		if (Positions.IsValid() && Positions.IsValidIndex(Index))
 		{
@@ -87,82 +87,85 @@ namespace PV::Facades
 		return DefaultPosition;
 	}
 
-	const float& FPointFacade::GetLengthFromRoot(const int32 Index) const
+	float FPointFacade::GetLengthFromRoot(int32 Index) const
 	{
 		if (LengthFromRoot.IsValid() && LengthFromRoot.IsValidIndex(Index))
 		{
 			return LengthFromRoot[Index];
 		}
-		static constexpr float DefaultLengthFromRoot = 0.0;
-		return DefaultLengthFromRoot;
+
+		return 0;
 	}
 
-	const float& FPointFacade::GetLengthFromSeed(const int32 Index) const
+	float FPointFacade::GetLengthFromSeed(int32 Index) const
 	{
 		if (LengthFromSeed.IsValid() && LengthFromSeed.IsValidIndex(Index))
 		{
 			return LengthFromSeed[Index];
 		}
-		static constexpr float DefaultLengthFromSeed = 0.0;
-		return DefaultLengthFromSeed;
+		
+		return 0;
 	}
 
-	const float& FPointFacade::GetPointScaleGradient(const int32 Index) const
+	float FPointFacade::GetPointScaleGradient(int32 Index) const
 	{
 		if (PointScaleGradient.IsValid() && PointScaleGradient.IsValidIndex(Index))
 		{
 			return PointScaleGradient[Index];
 		}
 
-		static constexpr float DefaultPScale = 0.0;
-		return DefaultPScale;
+		return 0;
 	}
 
-	const float& FPointFacade::GetHullGradient(const int32 Index) const
+	float FPointFacade::GetHullGradient(int32 Index) const
 	{
 		if (HullGradient.IsValid() && HullGradient.IsValidIndex(Index))
 		{
 			return HullGradient[Index];
 		}
 
-		static constexpr float DefaultValue = 0.0;
-		return DefaultValue;
+		return 0;
 	}
 
-	const float& FPointFacade::GetMainTrunkGradient(const int32 Index) const
+	float FPointFacade::GetMainTrunkGradient(int32 Index) const
 	{
 		if (MainTrunkGradient.IsValid() && MainTrunkGradient.IsValidIndex(Index))
 		{
 			return MainTrunkGradient[Index];
 		}
 
-		static constexpr float DefaultValue = 0.0;
-		return DefaultValue;
+		return 0;
 	}
 
-	const float& FPointFacade::GetGroundGradient(const int32 Index) const
+	float FPointFacade::GetGroundGradient(int32 Index) const
 	{
 		if (GroundGradient.IsValid() && GroundGradient.IsValidIndex(Index))
 		{
 			return GroundGradient[Index];
 		}
 
-		static constexpr float DefaultValue = 0.0;
-		return DefaultValue;
+		return 0;
 	}
 
-	const float& FPointFacade::GetPointScale(const int32 Index) const
+	float FPointFacade::GetPointScale(int32 Index) const
 	{
 		if (PointScale.IsValid() && PointScale.IsValidIndex(Index))
 		{
 			return PointScale[Index];
 		}
 
-		static constexpr float DefaultPScale = 0.0;
-		return DefaultPScale;
+		return 0;
+	}
+
+	void FPointFacade::SetPointScale(int32 Index, float InPointScale)
+	{
+		if (PointScale.IsValid() && PointScale.IsValidIndex(Index))
+		{
+			PointScale.ModifyAt(Index, InPointScale);
+		}
 	}
 
-	const int32 FPointFacade::GetBudNumber(const int32 Index) const
+	int32 FPointFacade::GetBudNumber(int32 Index) const
 	{
 		if (BudNumber.IsValid() && BudNumber.IsValidIndex(Index))
 		{
@@ -171,8 +174,16 @@ namespace PV::Facades
 
 		return INDEX_NONE;
 	}
+
+	void FPointFacade::SetBudNumber(int32 Index, int32 InBudNumber)
+	{
+		if (BudNumber.IsValid() && BudNumber.IsValidIndex(Index))
+		{
+			BudNumber.ModifyAt(Index, InBudNumber);
+		}
+	}
 	
-	const TArray<float>& FPointFacade::GetBudLightDetected(const int32 Index) const
+	const TArray<float>& FPointFacade::GetBudLightDetected(int32 Index) const
 	{
 		if (BudLightDetected.IsValid() && BudLightDetected.IsValidIndex(Index))
 		{
@@ -183,7 +194,15 @@ namespace PV::Facades
 		return EmptyArray;
 	}
 
-	const TArray<int>& FPointFacade::GetBudDevelopment(const int32 Index) const
+	void FPointFacade::SetBudLightDetected(int32 Index, TArray<float> InBudLightDetected)
+	{
+		if (BudLightDetected.IsValid() && BudLightDetected.IsValidIndex(Index))
+		{
+			BudLightDetected.Modify()[Index] = MoveTemp(InBudLightDetected);
+		}
+	}
+
+	const TArray<int>& FPointFacade::GetBudDevelopment(int32 Index) const
 	{
 		if (BudDevelopment.IsValid() && BudDevelopment.IsValidIndex(Index))
 		{
@@ -193,8 +212,49 @@ namespace PV::Facades
 		static const TArray<int> EmptyArray;
 		return EmptyArray;
 	}
+	
+	void FPointFacade::SetBudDevelopment(int32 Index, TArray<int> InBudDevelopment)
+	{
+		if (BudDevelopment.IsValid() && BudDevelopment.IsValidIndex(Index))
+		{
+			BudDevelopment.Modify()[Index] = MoveTemp(InBudDevelopment);
+		}
+	}
 
-	const TArray<FVector3f>& FPointFacade::GetBudDirection(const int32 Index) const
+	int32 FPointFacade::GetBudGeneratation(int32 Index) const
+	{
+		const TArray<int>& PointBudDevelopment = GetBudDevelopment(Index);
+		if (PointBudDevelopment.IsValidIndex(0))
+		{
+			return PointBudDevelopment[0];
+		}
+
+		return INDEX_NONE;
+	}
+
+	int32 FPointFacade::GetBudAge(int32 Index) const
+	{
+		const TArray<int>& PointBudDevelopment = GetBudDevelopment(Index);
+		if (PointBudDevelopment.IsValidIndex(1))
+		{
+			return PointBudDevelopment[1];
+		}
+
+		return INDEX_NONE;
+	}
+
+	int32 FPointFacade::GetBudBranchAge(int32 Index) const
+	{
+		const TArray<int>& PointBudDevelopment = GetBudDevelopment(Index);
+		if (PointBudDevelopment.IsValidIndex(2))
+		{
+			return PointBudDevelopment[2];
+		}
+
+		return INDEX_NONE;
+	}
+
+	const TArray<FVector3f>& FPointFacade::GetBudDirection(int32 Index) const
 	{
 		if (BudDirections.IsValid() && BudDirections.IsValidIndex(Index))
 		{
@@ -205,7 +265,15 @@ namespace PV::Facades
 		return EmptyArray;
 	}
 
-	const TArray<float>& FPointFacade::GetBudHormoneLevels(const int32 Index) const
+	void FPointFacade::SetBudDirections(int32 Index, TArray<FVector3f> InBudDirections)
+	{
+		if (BudDirections.IsValid() && BudDirections.IsValidIndex(Index))
+		{
+			BudDirections.Modify()[Index] = MoveTemp(InBudDirections);
+		}
+	}
+
+	const TArray<float>& FPointFacade::GetBudHormoneLevels(int32 Index) const
 	{
 		if (BudHormoneLevels.IsValid() && BudHormoneLevels.IsValidIndex(Index))
 		{
@@ -216,26 +284,25 @@ namespace PV::Facades
 		return EmptyArray;
 	}
 
-	void FPointFacade::SetBudHormoneLevels(const int32 Index, const TArray<float>& InBudHormoneLevels)
+	void FPointFacade::SetBudHormoneLevels(int32 Index, TArray<float> InBudHormoneLevels)
 	{
 		if (BudHormoneLevels.IsValid() && BudHormoneLevels.IsValidIndex(Index))
 		{
-			BudHormoneLevels.ModifyAt(Index, InBudHormoneLevels);
+			BudHormoneLevels.Modify()[Index] = MoveTemp(InBudHormoneLevels);
 		}
 	}
 
-	const float& FPointFacade::GetBranchGradient(const int32 Index) const
+	float FPointFacade::GetBranchGradient(int32 Index) const
 	{
 		if (BranchGradients.IsValid() && BranchGradients.IsValidIndex(Index))
 		{
 			return BranchGradients[Index];
 		}
 
-		static constexpr float DefaultBranchGradient = 0.0;
-		return DefaultBranchGradient;
+		return 0;
 	}
 
-	void FPointFacade::SetBranchGradient(const int32 Index, const float InBranchGradient)
+	void FPointFacade::SetBranchGradient(int32 Index, float InBranchGradient)
 	{
 		if (BranchGradients.IsValid() && BranchGradients.IsValidIndex(Index))
 		{
@@ -243,48 +310,53 @@ namespace PV::Facades
 		}
 	}
 
-	const float& FPointFacade::GetPlantGradient(const int32 Index) const
+	float FPointFacade::GetPlantGradient(int32 Index) const
 	{
 		if (PlantGradients.IsValid() && PlantGradients.IsValidIndex(Index))
 		{
 			return PlantGradients[Index];
 		}
 
-		static constexpr float DefaultPlantGradient = 0.0;
-		return DefaultPlantGradient;
+		return 0;
 	}
 
-	const float& FPointFacade::GetNjordPixelIndex(const int32 Index) const
+	void FPointFacade::SetPlantGradient(int32 Index, float InPlantGradient)
+	{
+		if (PlantGradients.IsValid() && PlantGradients.IsValidIndex(Index))
+		{
+			PlantGradients.ModifyAt(Index, InPlantGradient);
+		}
+	}
+
+	float FPointFacade::GetNjordPixelIndex(int32 Index) const
 	{
 		if (NjordPixelIndex.IsValid() && NjordPixelIndex.IsValidIndex(Index))
 		{
 			return NjordPixelIndex[Index];
 		}
 
-		static constexpr float IndexNone = -1.0;
-		return IndexNone;
+		return float(INDEX_NONE);
 	}
 
-	void FPointFacade::SetPlantGradient(const int32 Index, const float InPlantGradient)
+	void FPointFacade::SetNjordPixelIndex(int32 Index, float InNjordPixelIndex)
 	{
-		if (PlantGradients.IsValid() && PlantGradients.IsValidIndex(Index))
+		if (NjordPixelIndex.IsValid() && NjordPixelIndex.IsValidIndex(Index))
 		{
-			PlantGradients.ModifyAt(Index, InPlantGradient);
+			NjordPixelIndex.ModifyAt(Index, InNjordPixelIndex);
 		}
 	}
 
-	const float& FPointFacade::GetTextureCoordV(const int32 Index) const
+	float FPointFacade::GetTextureCoordV(int32 Index) const
 	{
 		if (TextureCoordV.IsValid() && TextureCoordV.IsValidIndex(Index))
 		{
 			return TextureCoordV[Index];
 		}
 
-		static constexpr float IndexNone = -1.0;
-		return IndexNone;
+		return float(INDEX_NONE);
 	}
 
-	void FPointFacade::SeTextureCoordV(const int32 Index, const float InTextureCoordV)
+	void FPointFacade::SetTextureCoordV(int32 Index, float InTextureCoordV)
 	{
 		if (!TextureCoordV.IsValid())
 		{
@@ -297,18 +369,17 @@ namespace PV::Facades
 		}
 	}
 
-	const float& FPointFacade::GetTextureCoordUOffset(const int32 Index) const
+	float FPointFacade::GetTextureCoordUOffset(int32 Index) const
 	{
 		if (TextureCoordUOffset.IsValid() && TextureCoordUOffset.IsValidIndex(Index))
 		{
 			return TextureCoordUOffset[Index];
 		}
 
-		static constexpr float IndexNone = -1.0;
-		return IndexNone;
+		return float(INDEX_NONE);
 	}
 
-	void FPointFacade::SetTextureCoordUOffset(const int32 Index, const float InTextureCoordUOffset)
+	void FPointFacade::SetTextureCoordUOffset(int32 Index, float InTextureCoordUOffset)
 	{
 		if (!TextureCoordUOffset.IsValid())
 		{
@@ -321,7 +392,7 @@ namespace PV::Facades
 		}
 	}
 
-	const FVector2f& FPointFacade::GetURange(const int32 Index) const
+	const FVector2f& FPointFacade::GetURange(int32 Index) const
 	{
 		if (URange.IsValid() && URange.IsValidIndex(Index))
 		{
@@ -332,7 +403,7 @@ namespace PV::Facades
 		return Range;
 	}
 
-	void FPointFacade::SetURange(const int32 Index, const FVector2f InURange)
+	void FPointFacade::SetURange(int32 Index, const FVector2f InURange)
 	{
 		if (!URange.IsValid())
 		{
@@ -385,7 +456,7 @@ namespace PV::Facades
 		return LengthFromRoot.Get();
 	}
 
-	void FPointFacade::CopyEntry(const int32 FromIndex, const int32 ToIndex)
+	void FPointFacade::CopyEntry(int32 FromIndex, int32 ToIndex)
 	{
 		if (IsValid() && Positions.IsValidIndex(FromIndex) && Positions.IsValidIndex(ToIndex))
 		{
@@ -421,7 +492,7 @@ namespace PV::Facades
 		}
 	}
 
-	void FPointFacade::RemoveEntries(const int32 NumEntries, const int32 StartIndex)
+	void FPointFacade::RemoveEntries(int32 NumEntries, int32 StartIndex)
 	{
 		if (IsValid() && Positions.IsValidIndex(StartIndex) && StartIndex + NumEntries <= Positions.Num())
 		{
diff --git a/Engine/Plugins/Experimental/ProceduralVegetationEditor/Source/ProceduralVegetation/Private/Implementations/PVMaterialSettings.cpp b/Engine/Plugins/Experimental/ProceduralVegetationEditor/Source/ProceduralVegetation/Private/Implementations/PVMaterialSettings.cpp
--- a/Engine/Plugins/Experimental/ProceduralVegetationEditor/Source/ProceduralVegetation/Private/Implementations/PVMaterialSettings.cpp
+++ b/Engine/Plugins/Experimental/ProceduralVegetationEditor/Source/ProceduralVegetation/Private/Implementations/PVMaterialSettings.cpp
@@ -325,7 +325,7 @@ void FPVMaterialSettings::SetOutputUVs(FManagedArrayCollection& Collection, TArr
 			UV.Y *= UVRatio;
 			UV.Y *= UVYTextureScale;
 
-			PointFacade.SeTextureCoordV(BranchPoints[j], UV.Y);
+			PointFacade.SetTextureCoordV(BranchPoints[j], UV.Y);
 			PointFacade.SetTextureCoordUOffset(BranchPoints[j], UV.X);
 			PointFacade.SetURange( BranchPoints[j], FVector2f(URange.GetLowerBound().GetValue(), URange.GetUpperBound().GetValue()));
 		}
diff --git a/Engine/Plugins/Experimental/ProceduralVegetationEditor/Source/ProceduralVegetation/Public/Facades/IShrinkableFacade.h b/Engine/Plugins/Experimental/ProceduralVegetationEditor/Source/ProceduralVegetation/Public/Facades/IShrinkableFacade.h
--- a/Engine/Plugins/Experimental/ProceduralVegetationEditor/Source/ProceduralVegetation/Public/Facades/IShrinkableFacade.h
+++ b/Engine/Plugins/Experimental/ProceduralVegetationEditor/Source/ProceduralVegetation/Public/Facades/IShrinkableFacade.h
@@ -17,11 +17,9 @@ namespace PV::Facades
 
 		virtual int32 GetElementCount() const = 0;
 
-		virtual void CopyEntry(const int32 FromIndex, const int32 ToIndex) = 0;
+		virtual void CopyEntry(int32 FromIndex, int32 ToIndex) = 0;
 
-		virtual void RemoveEntries(const int32 NumEntries, const int32 StartIndex) = 0;
-
-		// virtual void UpdateIndices(const TMap<int32, int32>& OldIDsToNewIDs) = 0;
+		virtual void RemoveEntries(int32 NumEntries, int32 StartIndex) = 0;
 
 		static TMap<int32, int32> RemoveEntries(IShrinkable& FacadeOut, TArray<bool>& EntriesToRemove)
 		{
diff --git a/Engine/Plugins/Experimental/ProceduralVegetationEditor/Source/ProceduralVegetation/Public/Facades/PVBranchFacade.h b/Engine/Plugins/Experimental/ProceduralVegetationEditor/Source/ProceduralVegetation/Public/Facades/PVBranchFacade.h
--- a/Engine/Plugins/Experimental/ProceduralVegetationEditor/Source/ProceduralVegetation/Public/Facades/PVBranchFacade.h
+++ b/Engine/Plugins/Experimental/ProceduralVegetationEditor/Source/ProceduralVegetation/Public/Facades/PVBranchFacade.h
@@ -25,44 +25,56 @@ namespace PV::Facades
 
 		virtual int32 GetElementCount() const override;
 
-		const TArray<int32>& GetPoints(const int32 Index) const;
+		const TArray<int32>& GetPoints(int32 Index) const;
 
-		int32 GetRootPoint(const int32 Index) const;
+		int32 GetRootPoint(int32 Index) const;
 
-		void SetPoints(const int32 Index, const TArray<int32>& InPoints);
+		void SetPoints(int32 Index, TArray<int32> InPoints);
 
-		const TArray<int32>& GetChildren(const int32 Index) const;
+		const TArray<int32>& GetChildren(int32 Index) const;
 
-		TArray<int32> GetImmediateChildren(const int32 Index) const;
+		TArray<int32> GetImmediateChildren(int32 Index) const;
 
-		void SetChildren(const int32 Index, const TArray<int32>& InChildren);
+		void SetChildren(int32 Index, TArray<int32> InChildren);
 
-		const TArray<int32>& GetParents(const int32 Index) const;
+		const TArray<int32>& GetParents(int32 Index) const;
 
-		int32 GetParentIndex(const int32 BranchIndex) const;
+		int32 GetParentIndex(int32 BranchIndex) const;
 
-		int32 GetParentBranchNumber(const int32 BranchIndex) const;
+		int32 GetParentBranchNumber(int32 BranchIndex) const;
+		
+		void SetParentBranchNumber(int32 BranchIndex, int32 InParentBranchNumber);
+
+		void SetParents(int32 Index, TArray<int32> InParents);
+
+		int32 GetBranchNumber(int32 Index) const;
 
-		void SetParents(const int32 Index, const TArray<int32>& InParents);
+		void SetBranchNumber(int32 Index, int32 InBranchNumber);
 
-		int32 GetBranchNumber(const int32 Index) const;
+		int32 GetBranchPlantNumber(int32 Index) const;
 
-		int32 GetBranchSourceBudNumber(const int32 Index) const;
+		void SetBranchPlantNumber(int32 Index, int32 InPlantNumber);
+
+		int32 GetBranchSourceBudNumber(int32 Index) const;
+		
+		void SetBranchSourceBudNumber(int32 Index, int32 InSourceBudNumber);
 
 		TArray<int32> GetParentBranchIndices(int BranchIndex) const;
 		
-		int32 GetParentBranchIndex(const int32 BranchIndex) const;
+		int32 GetParentBranchIndex(int32 BranchIndex) const;
 
-		int32 GetHierarchyGenerationNumber(const int32 Index) const;
+		int32 GetHierarchyGenerationNumber(int32 Index) const;
 
 		const TManagedArray<int32>& GetBranchNumbers() const;
 
-		int32 GetBranchUVMaterial(const int32 Index) const;
+		int32 GetBranchUVMaterial(int32 Index) const;
 
-		void SetBranchUVMaterial(const int32 Index, const int32 InMaterial);
+		void SetBranchUVMaterial(int32 Index, int32 InMaterial);
 		
 		int32 GetBranchHierarchyNumber(int32 Index) const;
 
+		void SetBranchHierarchyNumber(int32 Index, int32 InBranchHierarchyNumber);
+
 		int32 GetBranchSimulationGroupIndex(int32 Index) const;
 		
 		void SetBranchSimulationGroupIndex(int32 Index, int32 InLogicalDepth);
@@ -71,13 +83,13 @@ namespace PV::Facades
 
 		FString GetTrunkMaterialPath() const;
 
-		void SetTrunkURange(const TArray<FVector2f>& InURange);
+		void SetTrunkURange(TArray<FVector2f> InURange);
 
 		const TArray<FVector2f>& GetTrunkURange() const;
 
-		virtual void CopyEntry(const int32 FromIndex, const int32 ToIndex) override;
+		virtual void CopyEntry(int32 FromIndex, int32 ToIndex) override;
 
-		virtual void RemoveEntries(const int32 NumEntries, const int32 StartIndex) override;
+		virtual void RemoveEntries(int32 NumEntries, int32 StartIndex) override;
 
 		void GetSortedBranchIndicesByHierarchy(TArray<int32>& OutSortedIndices) const;
 		
diff --git a/Engine/Plugins/Experimental/ProceduralVegetationEditor/Source/ProceduralVegetation/Public/Facades/PVFoliageFacade.h b/Engine/Plugins/Experimental/ProceduralVegetationEditor/Source/ProceduralVegetation/Public/Facades/PVFoliageFacade.h
--- a/Engine/Plugins/Experimental/ProceduralVegetationEditor/Source/ProceduralVegetation/Public/Facades/PVFoliageFacade.h
+++ b/Engine/Plugins/Experimental/ProceduralVegetationEditor/Source/ProceduralVegetation/Public/Facades/PVFoliageFacade.h
@@ -29,7 +29,7 @@ namespace PV::Facades
 	class PROCEDURALVEGETATION_API FFoliageFacade final : public IShrinkable
 	{
 	public:
-		FFoliageFacade(FManagedArrayCollection& InCollection, const int32 InitialSize = 0);
+		FFoliageFacade(FManagedArrayCollection& InCollection, int32 InitialSize = 0);
 		FFoliageFacade(const FManagedArrayCollection& InCollection);
 
 		bool IsConst() const { return Collection == nullptr; }
@@ -40,9 +40,9 @@ namespace PV::Facades
 
 		FFoliageEntryData GetFoliageEntry(const int32 Index) const;
 
-		const float& GetLengthFromRoot(const int32 Index) const;
+		float GetLengthFromRoot(int32 Index) const;
 
-		void SetLengthFromRoot(const int32 Index, const float Input);
+		void SetLengthFromRoot(int32 Index, float Input);
 
 		int32 GetParentBoneID(int32 Index) const;
 
@@ -52,54 +52,56 @@ namespace PV::Facades
 
 		int32 AddFoliageEntry(const FFoliageEntryData& InputData);
 
-		void SetFoliageEntry(const int32 Index, const FFoliageEntryData& InputData);
+		void SetFoliageEntry(int32 Index, const FFoliageEntryData& InputData);
 
-		void SetFoliageIdsArray(const int32 Index, const TArray<int32>& InputIds);
+		void SetFoliageIdsArray(int32 Index, TArray<int32> InputIds);
 
-		void SetFoliageBranchId(const int32 Index, const int32 InputId);
+		int32 GetFoliageBranchId(int32 Index) const;
+
+		void SetFoliageBranchId(int32 Index, int32 InputId);
 
 		int32 NumFoliageNames() const { return NamesAttribute.Num(); };
 
-		FString GetFoliageName(const int32 Index) const;
+		const FString& GetFoliageName(int32 Index) const;
 
-		TArray<FString> GetFoliageNames() const;
+		const TArray<FString>& GetFoliageNames() const;
 
-		const FVector3f& GetPivotPoint(const int32 Index) const;
+		const FVector3f& GetPivotPoint(int32 Index) const;
 
-		const FVector3f& GetUpVector(const int32 Index) const;
+		const FVector3f& GetUpVector(int32 Index) const;
 
-		const FVector3f& GetNormalVector(const int32 Index) const;
+		const FVector3f& GetNormalVector(int32 Index) const;
 
-		void SetPivotPoint(const int32 Index, const FVector3f& Input);
+		void SetPivotPoint(int32 Index, const FVector3f& Input);
 
-		void SetUpVector(const int32 Index, const FVector3f& Input);
+		void SetUpVector(int32 Index, const FVector3f& Input);
 
-		void SetNormalVector(const int32 Index, const FVector3f& Input);
+		void SetNormalVector(int32 Index, const FVector3f& Input);
 
-		void SetScale(const int32 Index, const float& Input);
+		void SetScale(int32 Index, float Input);
 
-		void SetFoliageName(const int32 Index, const FString& Input);
+		void SetFoliageName(int32 Index, const FString& Input);
 
 		void SetFoliageNames(const TArray<FString>& InputNames);
 
-		void SetFoliagePath(const FString InPath);
+		void SetFoliagePath(FString InPath);
 
 		FString GetFoliagePath() const;
 
 		virtual int32 GetElementCount() const override;
 
-		virtual void CopyEntry(const int32 FromIndex, const int32 ToIndex) override;
+		virtual void CopyEntry(int32 FromIndex, int32 ToIndex) override;
 
-		virtual void RemoveEntries(const int32 NumEntries, const int32 StartIndex) override;
+		virtual void RemoveEntries(int32 NumEntries, int32 StartIndex) override;
 		
-		FTransform GetFoliageTransform(const int32 Id) const;
+		FTransform GetFoliageTransform(int32 Id) const;
 
 		const TManagedArray<FVector3f>& GetPivotPositions() const;
 		
 		TManagedArray<FVector3f>& ModifyPivotPositions();
 
 	protected:
-		void DefineSchema(const int32 InitialSize = 0);
+		void DefineSchema(int32 InitialSize = 0);
 
 		const FManagedArrayCollection& ConstCollection;
 		FManagedArrayCollection* Collection = nullptr;
diff --git a/Engine/Plugins/Experimental/ProceduralVegetationEditor/Source/ProceduralVegetation/Public/Facades/PVPointFacade.h b/Engine/Plugins/Experimental/ProceduralVegetationEditor/Source/ProceduralVegetation/Public/Facades/PVPointFacade.h
--- a/Engine/Plugins/Experimental/ProceduralVegetationEditor/Source/ProceduralVegetation/Public/Facades/PVPointFacade.h
+++ b/Engine/Plugins/Experimental/ProceduralVegetationEditor/Source/ProceduralVegetation/Public/Facades/PVPointFacade.h
@@ -22,49 +22,67 @@ namespace PV::Facades
 
 		virtual int32 GetElementCount() const override;
 
-		const FVector3f& GetPosition(const int32 Index) const;
+		const FVector3f& GetPosition(int32 Index) const;
 
-		const float& GetLengthFromRoot(const int32 Index) const;
+		float GetLengthFromRoot(int32 Index) const;
 
-		const float& GetLengthFromSeed(const int32 Index) const;
+		float GetLengthFromSeed(int32 Index) const;
 
-		const float& GetPointScaleGradient(const int32 Index) const;
+		float GetPointScaleGradient(int32 Index) const;
 
-		const float& GetHullGradient(const int32 Index) const;
+		float GetHullGradient(int32 Index) const;
 
-		const float& GetMainTrunkGradient(const int32 Index) const;
+		float GetMainTrunkGradient(int32 Index) const;
 
-		const float& GetGroundGradient(const int32 Index) const;
+		float GetGroundGradient(int32 Index) const;
 		
-		const float& GetPointScale(const int32 Index) const;
+		float GetPointScale(int32 Index) const;
 
-		const int32 GetBudNumber(const int32 Index) const;
+		void SetPointScale(int32 Index, float InPointScale);
 
-		const TArray<float>& GetBudLightDetected(const int32 Index) const;
+		int32 GetBudNumber(int32 Index) const;
 
-		const TArray<int>& GetBudDevelopment(const int32 Index) const;
+		void SetBudNumber(int32 Index, int32 InBudNumber);
+
+		const TArray<float>& GetBudLightDetected(int32 Index) const;
+
+		void SetBudLightDetected(int32 Index, TArray<float> InBudLightDetected);
+
+		const TArray<int>& GetBudDevelopment(int32 Index) const;
 		
-		const TArray<FVector3f>& GetBudDirection(const int32 Index) const;
+		void SetBudDevelopment(int32 Index, TArray<int> InBudDevelopment);
+
+		int32 GetBudGeneratation(int32 Index) const;
+
+		int32 GetBudAge(int32 Index) const;
+
+		int32 GetBudBranchAge(int32 Index) const;
+
+		const TArray<FVector3f>& GetBudDirection(int32 Index) const;
+
+		void SetBudDirections(int32 Index, TArray<FVector3f> InBudDirections);
+
+		const TArray<float>& GetBudHormoneLevels(int32 Index) const;
 
-		const TArray<float>& GetBudHormoneLevels(const int32 Index) const;
+		void SetBudHormoneLevels(int32 Index, TArray<float> InBudHormoneLevels);
 
-		void SetBudHormoneLevels(const int32 Index, const TArray<float>& InBudHormoneLevels);
+		float GetBranchGradient(int32 Index) const;
 
-		const float& GetBranchGradient(const int32 Index) const;
+		void SetBranchGradient(int32 Index, float InBranchGradient);
 
-		void SetBranchGradient(const int32 Index, const float InBranchGradient);
+		float GetPlantGradient(int32 Index) const;
 
-		const float& GetPlantGradient(const int32 Index) const;
+		void SetPlantGradient(int32 Index, float InPlantGradient);
 
-		const float& GetNjordPixelIndex(const int32 Index) const;
+		float GetNjordPixelIndex(int32 Index) const;
 
-		void SetPlantGradient(const int32 Index, const float InPlantGradient);
+		void SetNjordPixelIndex(int32 Index, float InNjordPixelIndex);
 
-		const float& GetTextureCoordV(int32 Index) const;
+		float GetTextureCoordV(int32 Index) const;
 		
-		void SeTextureCoordV(int32 Index, float InTextureCoordV);
+		void SetTextureCoordV(int32 Index, float InTextureCoordV);
 		
-		const float& GetTextureCoordUOffset(int32 Index) const;
+		float GetTextureCoordUOffset(int32 Index) const;
 		
 		void SetTextureCoordUOffset(int32 Index, float InTextureCoordUOffset);
 		
@@ -88,9 +106,9 @@ namespace PV::Facades
 
 		const TManagedArray<float>& GetLengthFromRootsArray() const;
 
-		virtual void CopyEntry(const int32 FromIndex, const int32 ToIndex) override;
+		virtual void CopyEntry(int32 FromIndex, int32 ToIndex) override;
 
-		virtual void RemoveEntries(const int32 NumEntries, const int32 StartIndex) override;
+		virtual void RemoveEntries(int32 NumEntries, int32 StartIndex) override;
 
 	private:
 		TManagedArrayAccessor<FVector3f> Positions;
