Add PVE debug visualization capabilities
* Add capability to visualize array of vectors
* Add "Custom" attribute visualization
* Slight optimization of line and instanced mesh drawing
* Fix bug where line batcher component would not refresh its render state

diff --git a/Engine/Plugins/Experimental/ProceduralVegetationEditor/Source/ProceduralVegetation/Public/DataTypes/PVData.h b/Engine/Plugins/Experimental/ProceduralVegetationEditor/Source/ProceduralVegetation/Public/DataTypes/PVData.h
--- a/Engine/Plugins/Experimental/ProceduralVegetationEditor/Source/ProceduralVegetation/Public/DataTypes/PVData.h
+++ b/Engine/Plugins/Experimental/ProceduralVegetationEditor/Source/ProceduralVegetation/Public/DataTypes/PVData.h
@@ -58,13 +58,31 @@ struct FPVVisualizationSettings
 	UPROPERTY(EditAnywhere, Category = DebugSettings)
 	FName AttributeToFilter;
 
+	UPROPERTY(EditAnywhere, Category = DebugSettings, meta = (EditCondition = "DebugType==EPVDebugType::Custom"))
+	FName CustomGroupToFilter = TEXT("Custom");
+
+	UPROPERTY(EditAnywhere, Category = DebugSettings, meta = (EditCondition = "DebugType==EPVDebugType::Custom"))
+	FName CustomPivotPositionAttributeName = TEXT("Position");
+
+	UPROPERTY(EditAnywhere, Category = DebugSettings, meta = (EditCondition = "DebugType==EPVDebugType::Custom"))
+	FName CustomPivotScaleAttributeName = TEXT("Scale");
+
 	UPROPERTY(EditAnywhere, Category = DebugSettings)
 	EPVDebugValueVisualizationMode VisualizationMode = EPVDebugValueVisualizationMode::Text;
 
 	UPROPERTY(EditAnywhere, Category = DebugSettings)
 	FColor Color = FColor::Black;
+
+	UPROPERTY(EditAnywhere, Category = DebugSettings)
+	float DirectionStrength = 8.f;
+
+	UPROPERTY(EditAnywhere, Category = DebugSettings)
+	float TextScale = 2.5f;
+
+	UPROPERTY(EditAnywhere, Category = DebugSettings)
+	float PointScale = 0.005f;
 	
-	FName GetDebugTypeString() const
+	FName GetGroupToFilter() const
 	{
 		switch (DebugType)
 		{
@@ -78,7 +96,7 @@ struct FPVVisualizationSettings
 			return "Foliage";
 
 		case EPVDebugType::Custom:
-			return "Custom";
+			return CustomGroupToFilter;
 		}
 		return NAME_None;
 	}
diff --git a/Engine/Plugins/Experimental/ProceduralVegetationEditor/Source/ProceduralVegetationEditor/Private/Visualizations/DebugVisualization/PVDataAttributeValue.cpp b/Engine/Plugins/Experimental/ProceduralVegetationEditor/Source/ProceduralVegetationEditor/Private/Visualizations/DebugVisualization/PVDataAttributeValue.cpp
--- a/Engine/Plugins/Experimental/ProceduralVegetationEditor/Source/ProceduralVegetationEditor/Private/Visualizations/DebugVisualization/PVDataAttributeValue.cpp
+++ b/Engine/Plugins/Experimental/ProceduralVegetationEditor/Source/ProceduralVegetationEditor/Private/Visualizations/DebugVisualization/PVDataAttributeValue.cpp
@@ -49,8 +49,12 @@ TArray<FText> FPVAttributeScalarValue::ToText()
 
 TArray<FVector> FPVAttributeScalarValue::ToVectors()
 {
-	static TArray<FVector> VecArray;
-	return VecArray;
+	return TArray<FVector>();
+}
+
+TArray<TArray<FVector>> FPVAttributeScalarValue::ToVectorArrays()
+{
+	return TArray<TArray<FVector>>();
 }
 
 FPVAttributeVectorValue::FPVAttributeVectorValue(const FName InGroupName, const FName InAttributeName,
@@ -61,11 +65,11 @@ FPVAttributeVectorValue::FPVAttributeVectorValue(const FName InGroupName, const
 	{
 		if (InAttributeType == EManagedArrayType::FVectorType)
 		{
-			const TManagedArray<FVector3f>& FloatAttributesArray = InCollection.GetAttribute<FVector3f>(InAttributeName, InGroupName);
+			const TManagedArray<FVector3f>& VectorAttributesArray = InCollection.GetAttribute<FVector3f>(InAttributeName, InGroupName);
 
-			for (int i = 0; i < FloatAttributesArray.Num(); i++)
+			for (int i = 0; i < VectorAttributesArray.Num(); i++)
 			{
-				ValueArray.Add(FVector(FloatAttributesArray[i].X, FloatAttributesArray[i].Y, FloatAttributesArray[i].Z));
+				ValueArray.Add(FVector(VectorAttributesArray[i].X, VectorAttributesArray[i].Y, VectorAttributesArray[i].Z));
 			}
 		}
 		else
@@ -89,12 +93,84 @@ TArray<FText> FPVAttributeVectorValue::ToText()
 
 TArray<FVector> FPVAttributeVectorValue::ToVectors()
 {
-	TArray<FVector> Vectors;
-	
+	return ValueArray;
+}
+
+TArray<TArray<FVector>> FPVAttributeVectorValue::ToVectorArrays()
+{
+	TArray<TArray<FVector>> VectorArrays;
+	VectorArrays.SetNum(ValueArray.Num());
+
 	for (int i = 0; i < ValueArray.Num(); i++)
 	{
-		Vectors.Add(ValueArray[i]);
+		VectorArrays[i] = { ValueArray[i] };
+	}
+
+	return VectorArrays;
+}
+
+FPVAttributeVectorArrayValue::FPVAttributeVectorArrayValue(const FName InGroupName, const FName InAttributeName,
+	EManagedArrayType InAttributeType, const FManagedArrayCollection& InCollection)
+	: FPVAttributeValueInterface(InGroupName, InAttributeName, InAttributeType, InCollection)
+{
+	if (InCollection.HasAttribute(InAttributeName, InGroupName))
+	{
+		check(InAttributeType == EManagedArrayType::FFVectorArrayType);
+
+		const TManagedArray<TArray<FVector3f>>& VectorArrayAttributesArray = InCollection.GetAttribute<TArray<FVector3f>>(InAttributeName, InGroupName);
+
+		for (int32 i = 0; i < VectorArrayAttributesArray.Num(); i++)
+		{
+			const TArray<FVector3f> VectorArray = VectorArrayAttributesArray[i];
+			TArray<FVector>& Value = ValueArray.AddDefaulted_GetRef();
+			Value.SetNumUninitialized(VectorArray.Num());
+			for (int32 j = 0; j < VectorArray.Num(); ++j)
+			{
+				Value[j] = FVector(VectorArray[j].X, VectorArray[j].Y, VectorArray[j].Z);
+			}
+		}
+	}
+}
+
+TArray<FText> FPVAttributeVectorArrayValue::ToText()
+{
+	TArray<FText> TextArray;
+
+	for (const TArray<FVector>& Value : ValueArray)
+	{
+		FString OutTextString = TEXT("{ ");
+		for (int32 i = 0; i < Value.Num(); ++i)
+		{
+			OutTextString += FString::Printf(TEXT("{ %s }"), *Value[i].ToString());
+			if (i != Value.Num() - 1)
+			{
+				OutTextString += ", ";
+			}
+		}
+		OutTextString += " }";
+		
+		TextArray.Add(FText::FromString(OutTextString));
+	}
+
+	return TextArray;
+}
+
+TArray<FVector> FPVAttributeVectorArrayValue::ToVectors()
+{
+	TArray<FVector> Vectors;
+
+	for (const TArray<FVector>& Value : ValueArray)
+	{
+		for (const FVector& Vector : Value)
+		{
+			Vectors.Add(Vector);
+		}
 	}
 
 	return Vectors;
 }
+
+TArray<TArray<FVector>> FPVAttributeVectorArrayValue::ToVectorArrays()
+{
+	return ValueArray;
+}
\ No newline at end of file
diff --git a/Engine/Plugins/Experimental/ProceduralVegetationEditor/Source/ProceduralVegetationEditor/Private/Visualizations/DebugVisualization/PVDataAttributeValue.h b/Engine/Plugins/Experimental/ProceduralVegetationEditor/Source/ProceduralVegetationEditor/Private/Visualizations/DebugVisualization/PVDataAttributeValue.h
--- a/Engine/Plugins/Experimental/ProceduralVegetationEditor/Source/ProceduralVegetationEditor/Private/Visualizations/DebugVisualization/PVDataAttributeValue.h
+++ b/Engine/Plugins/Experimental/ProceduralVegetationEditor/Source/ProceduralVegetationEditor/Private/Visualizations/DebugVisualization/PVDataAttributeValue.h
@@ -12,6 +12,7 @@ public:
 	FPVAttributeScalarValue(const FName InGroupName, const FName InAttributeName, EManagedArrayType InAttributeType, const FManagedArrayCollection& InCollection);
 	virtual TArray<FText> ToText() override;
 	virtual TArray<FVector> ToVectors() override;
+	virtual TArray<TArray<FVector>> ToVectorArrays() override;
 
 private:
 	TArray<float> ValueArray;
@@ -24,7 +25,21 @@ public:
 	virtual TArray<FText> ToText() override;
 	virtual bool IsScalar() override { return false;}
 	virtual TArray<FVector> ToVectors() override;
+	virtual TArray<TArray<FVector>> ToVectorArrays() override;
 
 private:
 	TArray<FVector> ValueArray;
 };
+
+class FPVAttributeVectorArrayValue : public FPVAttributeValueInterface
+{
+public:
+	FPVAttributeVectorArrayValue(const FName InGroupName, const FName InAttributeName, EManagedArrayType InAttributeType, const FManagedArrayCollection& InCollection);
+	virtual TArray<FText> ToText() override;
+	virtual bool IsScalar() override { return false; }
+	virtual TArray<FVector> ToVectors() override;
+	virtual TArray<TArray<FVector>> ToVectorArrays() override;
+
+private:
+	TArray<TArray<FVector>> ValueArray;
+};
diff --git a/Engine/Plugins/Experimental/ProceduralVegetationEditor/Source/ProceduralVegetationEditor/Private/Visualizations/DebugVisualization/PVDataAttributeValueInterface.cpp b/Engine/Plugins/Experimental/ProceduralVegetationEditor/Source/ProceduralVegetationEditor/Private/Visualizations/DebugVisualization/PVDataAttributeValueInterface.cpp
--- a/Engine/Plugins/Experimental/ProceduralVegetationEditor/Source/ProceduralVegetationEditor/Private/Visualizations/DebugVisualization/PVDataAttributeValueInterface.cpp
+++ b/Engine/Plugins/Experimental/ProceduralVegetationEditor/Source/ProceduralVegetationEditor/Private/Visualizations/DebugVisualization/PVDataAttributeValueInterface.cpp
@@ -14,7 +14,9 @@ FPVAttributeValuePtr FPVAttributeValueInterface::Create(const FName InGroupName,
 	case EManagedArrayType::FFloatType:
 			return std::make_shared<FPVAttributeScalarValue>(InGroupName, InAttributeName, InAttributeType,InCollection);
 	case EManagedArrayType::FVectorType:
-			return std::make_shared<FPVAttributeVectorValue>(InGroupName, InAttributeName, InAttributeType,InCollection);
+		return std::make_shared<FPVAttributeVectorValue>(InGroupName, InAttributeName, InAttributeType, InCollection);
+	case EManagedArrayType::FFVectorArrayType:
+		return std::make_shared<FPVAttributeVectorArrayValue>(InGroupName, InAttributeName, InAttributeType, InCollection);
 	default:
 		return nullptr;
 	}
diff --git a/Engine/Plugins/Experimental/ProceduralVegetationEditor/Source/ProceduralVegetationEditor/Private/Visualizations/DebugVisualization/PVDataAttributeValueInterface.h b/Engine/Plugins/Experimental/ProceduralVegetationEditor/Source/ProceduralVegetationEditor/Private/Visualizations/DebugVisualization/PVDataAttributeValueInterface.h
--- a/Engine/Plugins/Experimental/ProceduralVegetationEditor/Source/ProceduralVegetationEditor/Private/Visualizations/DebugVisualization/PVDataAttributeValueInterface.h
+++ b/Engine/Plugins/Experimental/ProceduralVegetationEditor/Source/ProceduralVegetationEditor/Private/Visualizations/DebugVisualization/PVDataAttributeValueInterface.h
@@ -23,6 +23,7 @@ public:
 	virtual TArray<FText> ToText() = 0;
 	virtual TArray<FVector> ToVectors() = 0;
 	virtual bool IsScalar() { return true; }
+	virtual TArray<TArray<FVector>> ToVectorArrays() = 0;
 
 	static FPVAttributeValuePtr Create(const FName InGroupName, const FName InAttributeName, EManagedArrayType InAttributeType,const FManagedArrayCollection& InCollection);
 };
diff --git a/Engine/Plugins/Experimental/ProceduralVegetationEditor/Source/ProceduralVegetationEditor/Private/Visualizations/DebugVisualization/PVDebugVisualizationBase.cpp b/Engine/Plugins/Experimental/ProceduralVegetationEditor/Source/ProceduralVegetationEditor/Private/Visualizations/DebugVisualization/PVDebugVisualizationBase.cpp
--- a/Engine/Plugins/Experimental/ProceduralVegetationEditor/Source/ProceduralVegetationEditor/Private/Visualizations/DebugVisualization/PVDebugVisualizationBase.cpp
+++ b/Engine/Plugins/Experimental/ProceduralVegetationEditor/Source/ProceduralVegetationEditor/Private/Visualizations/DebugVisualization/PVDebugVisualizationBase.cpp
@@ -21,10 +21,6 @@
 
 DEFINE_LOG_CATEGORY(LogPVDebugVisualization);
 
-const float FPVDebugVisualizationBase::DefaultDirectionStrength = 20;
-const float FPVDebugVisualizationBase::DefaultTextScale = 2.5;
-const float FPVDebugVisualizationBase::DefaultPointScale = 0.005f;
-
 void FPVDebugVisualizationBase::Draw(FVisualizerDrawContext& InContext)
 {
 	DrawPivotPoints(InContext);
@@ -33,139 +29,163 @@ void FPVDebugVisualizationBase::Draw(FVisualizerDrawContext& InContext)
 
 void FPVDebugVisualizationBase::DrawPivotPoints(const FVisualizerDrawContext& InContext)
 {
-	auto Collection = InContext.Collection;
-	auto Settings = InContext.VisualizationSettings;
+	const auto& Settings = InContext.VisualizationSettings;
 
 	if (Settings.bShowAnchorPoints)
 	{
-		const TArray<FVector3f>& Positions = GetPivotPositions(Collection);
+		const TArray<FVector3f> Positions = GetPivotPositions(InContext.Collection);
 	
 		for (int32 i = 0; i < Positions.Num(); i++)
 		{
 			const FVector Pos(Positions[i]);
-		
-			AddDrawAsPoint(Pos);
+			AddDrawAsPoint(Pos, Settings.PointScale);
 		}	
 	}
 }
 
 void FPVDebugVisualizationBase::DrawAttributes(FVisualizerDrawContext& InContext)
 {
-	auto Collection = InContext.Collection;
+	const auto& Collection = InContext.Collection;
 	
-	auto VisualizationSettings = InContext.VisualizationSettings;
+	const auto& VisualizationSettings = InContext.VisualizationSettings;
 	const FName AttributeName = VisualizationSettings.AttributeToFilter;
-	const FName GroupName = VisualizationSettings.GetDebugTypeString();
+	const FName GroupName = VisualizationSettings.GetGroupToFilter();
 
-	if (!GroupName.IsNone() && !AttributeName.IsNone() && Collection.HasGroup(GroupName))
+	if (GroupName.IsNone() || AttributeName.IsNone())
 	{
-		if (Collection.HasAttribute(AttributeName, GroupName))
-		{
-			int NumElements = Collection.NumElements(GroupName);
-			const FManagedArrayCollection::EArrayType AttributeType = Collection.GetAttributeType(AttributeName, GroupName);
+		return;
+	}
 
-			FPVAttributeValuePtr AttributeValue = FPVAttributeValueInterface::Create(GroupName, AttributeName, AttributeType, Collection);
+	if (!Collection.HasAttribute(AttributeName, GroupName) || !Collection.HasGroup(GroupName))
+	{
+		return;
+	}
+	
+	const FManagedArrayCollection::EArrayType AttributeType = Collection.GetAttributeType(AttributeName, GroupName);
+	const FPVAttributeValuePtr AttributeValue = FPVAttributeValueInterface::Create(GroupName, AttributeName, AttributeType, Collection);
 
-			if (AttributeValue)
-			{
-				if (AttributeValue->IsScalar() || VisualizationSettings.VisualizationMode == EPVDebugValueVisualizationMode::Text)
-				{
-					auto TextValueArray = AttributeValue->ToText();
+	if (!AttributeValue)
+	{
+		return;
+	}
+
+	const int32 NumElements = Collection.NumElements(GroupName);
+
+	if (AttributeValue->IsScalar() || VisualizationSettings.VisualizationMode == EPVDebugValueVisualizationMode::Text)
+	{
+		const TArray<FText> TextValueArray = AttributeValue->ToText();
+		check(TextValueArray.Num() == NumElements);
 		
-					for (int32 i = 0; i < NumElements; i++)
-					{
-						FVector3f OutPos;
-						float OutScale;
-						GetPivot(Collection, i, OutPos, OutScale);
-						
-						check(TextValueArray.IsValidIndex(i));
-
-						DrawAsText(InContext.SceneSetupParams, TextValueArray[i], OutPos, OutScale, VisualizationSettings.Color);
-					}
-				}
-				else if (VisualizationSettings.VisualizationMode == EPVDebugValueVisualizationMode::Direction)
-				{
-					auto VectorsAttribute = AttributeValue->ToVectors();
+		for (int32 i = 0; i < NumElements; i++)
+		{
+			FVector3f OutPos;
+			float OutScale;
+			GetPivot(Collection, i, OutPos, OutScale);
 
-					for (int32 i = 0; i < NumElements; i++)
-					{
-						FVector3f OutPos;
-						float OutScale;
-						GetPivot(Collection, i, OutPos, OutScale);
+			DrawAsText(InContext.SceneSetupParams, TextValueArray[i], OutPos, OutScale, VisualizationSettings.TextScale, VisualizationSettings.Color);
+		}
+	}
+	else if (VisualizationSettings.VisualizationMode == EPVDebugValueVisualizationMode::Direction)
+	{
+		TObjectPtr<UPVLineBatchComponent> LineBatchComponent = GetOrCreateLineComponent(InContext.SceneSetupParams);
 
-						FVector StartPos(OutPos);
-						FVector EndPos = StartPos + VectorsAttribute[i].GetSafeNormal() * DefaultDirectionStrength;
-						check(VectorsAttribute.IsValidIndex(i));
+		if (AttributeType == FManagedArrayCollection::EArrayType::FFVectorArrayType)
+		{
+			const TArray<TArray<FVector>> VectorArrays = AttributeValue->ToVectorArrays();
+			check(VectorArrays.Num() == NumElements);
 
-						DrawAsDirection(InContext.SceneSetupParams, StartPos, EndPos, VisualizationSettings.Color);
-					}
-				}
-				else if (VisualizationSettings.VisualizationMode == EPVDebugValueVisualizationMode::Point)
+			for (int32 i = 0; i < NumElements; i++)
+			{
+				FVector3f OutPos;
+				float OutScale;
+				GetPivot(Collection, i, OutPos, OutScale);
+
+				for (int32 j = 0; j < VectorArrays[i].Num(); ++j)
 				{
-					auto Vectors = AttributeValue->ToVectors();
+					const FVector& DirectionVector = VectorArrays[i][j];
+					FVector StartPos(OutPos);
+					FVector EndPos = StartPos + DirectionVector * VisualizationSettings.DirectionStrength;
 
-					for (int32 i = 0; i < NumElements; i++)
-					{
-						check(Vectors.IsValidIndex(i));
+					FColor DrawColor = FLinearColor::LerpUsingHSV(FLinearColor::Black, VisualizationSettings.Color.ReinterpretAsLinear(), j / (float)VectorArrays[i].Num()).ToFColor(true);
+					DrawAsDirection(LineBatchComponent, StartPos, EndPos, DrawColor);
+					DrawAsText(InContext.SceneSetupParams, FText::FromString(FString::Printf(TEXT("%d"), j)), FVector3f(EndPos), 0.f, VisualizationSettings.TextScale, DrawColor);
+				}
+			}
+		}
+		else
+		{
+			const TArray<FVector> Vectors = AttributeValue->ToVectors();
+			check(Vectors.Num() == NumElements);
+
+			for (int32 i = 0; i < NumElements; i++)
+			{
+				FVector3f OutPos;
+				float OutScale;
+				GetPivot(Collection, i, OutPos, OutScale);
 
-						AddDrawAsPoint(Vectors[i]);
-					}
-				}	
+				const FVector StartPos(OutPos);
+				const FVector EndPos = StartPos + Vectors[i].GetSafeNormal() * VisualizationSettings.DirectionStrength;
+
+				DrawAsDirection(LineBatchComponent, StartPos, EndPos, VisualizationSettings.Color);
 			}
 		}
 	}
+	else if (VisualizationSettings.VisualizationMode == EPVDebugValueVisualizationMode::Point)
+	{
+		const TArray<FVector> Vectors = AttributeValue->ToVectors();
+		check(Vectors.Num() == NumElements);
+
+		for (int32 i = 0; i < NumElements; i++)
+		{
+			AddDrawAsPoint(Vectors[i], VisualizationSettings.PointScale);
+		}
+	}	
 
-	DrawAsPoint(InContext.SceneSetupParams, PointsToDraw);
+	if (PointsToDraw.Num() > 0)
+	{
+		DrawAsPoint(GetOrCreatePointComponent(InContext.SceneSetupParams), PointsToDraw);
+	}
 }
 
-void FPVDebugVisualizationBase::AddDrawAsPoint(const FVector& InPos)
+void FPVDebugVisualizationBase::AddDrawAsPoint(const FVector& InPos, float PointScale)
 {
 	auto Transform = FTransform(InPos);
-	Transform.SetScale3D(FVector(DefaultPointScale, DefaultPointScale, DefaultPointScale));
+	Transform.SetScale3D(FVector(PointScale, PointScale, PointScale));
 
 	PointsToDraw.Add(Transform);
 }
 
-void FPVDebugVisualizationBase::DrawAsText(FPCGSceneSetupParams& InOutParams, const FText& InTextToDraw, const FVector3f& InPos, const float& InScale, const FLinearColor& InColor)
+void FPVDebugVisualizationBase::DrawAsText(FPCGSceneSetupParams& InOutParams, const FText& InTextToDraw, const FVector3f& InPos, float InScale, float TextSize, const FLinearColor& InColor)
 {
 	const FVector3f& Pos = InPos + (FVector3f::RightVector * -1 * InScale);
 
 	TObjectPtr<UTextRenderComponent> Component = NewObject<UTextRenderComponent>(GetTransientPackage(), NAME_None, RF_Transient);
 	Component->SetText(InTextToDraw);
 	Component->SetTextRenderColor(InColor.ToFColor(false));
-	Component->SetWorldSize(DefaultTextScale);
+	Component->SetWorldSize(TextSize);
 	Component->SetGenerateOverlapEvents(false);
-						
+	Component->SetCollisionEnabled(ECollisionEnabled::NoCollision);
+	
 	InOutParams.ManagedResources.Add(Component);
 	InOutParams.Scene->AddComponent(Component, FTransform(FVector(Pos.X, Pos.Y, Pos.Z)));
 }
 
-void FPVDebugVisualizationBase::DrawAsDirection(FPCGSceneSetupParams& InOutParams, const FVector& InStartPos, const FVector& InEndPos,
+void FPVDebugVisualizationBase::DrawAsDirection(UPVLineBatchComponent* LineBatchComponent, const FVector& InStartPos, const FVector& InEndPos,
 	const FLinearColor& InColor)
 {
-	UPVLineBatchComponent* LineComponent = GetOrCreateLineComponent(InOutParams);
-
-	check(LineComponent)
-	
-	if (LineComponent)
-	{
-		LineComponent->AddLine(InStartPos, InEndPos, InColor);
-	}
+	check(LineBatchComponent)
+	LineBatchComponent->AddLine(InStartPos, InEndPos, InColor);
 }
 
-void FPVDebugVisualizationBase::DrawAsPoint(FPCGSceneSetupParams& InOutParams, const TArray<FTransform>& InTransforms, const FLinearColor& InColor)
+void FPVDebugVisualizationBase::DrawAsPoint(UInstancedStaticMeshComponent* InstancedStaticMeshComponent, const TArray<FTransform>& InTransforms, const FLinearColor& InColor)
 {
 	if (InTransforms.Num() == 0)
+	{
 		return;
+	}
 	
-	UInstancedStaticMeshComponent* Component = GetOrCreatePointComponent(InOutParams);
-
-	check(Component)
-	
-	if (Component)
-	{
-		Component->AddInstances(InTransforms, false);
-	}	
+	check(InstancedStaticMeshComponent != nullptr);
+	InstancedStaticMeshComponent->AddInstances(InTransforms, false);
 }
 
 UPVLineBatchComponent* FPVDebugVisualizationBase::GetOrCreateLineComponent(FPCGSceneSetupParams& InOutParams)
@@ -195,6 +215,7 @@ UInstancedStaticMeshComponent* FPVDebugVisualizationBase::GetOrCreatePointCompon
 		
 		InstancedStaticMeshComponent = NewObject<UInstancedStaticMeshComponent>(GetTransientPackage(), NAME_None);
 		InstancedStaticMeshComponent->SetStaticMesh(PointSphere);
+		InstancedStaticMeshComponent->SetCollisionEnabled(ECollisionEnabled::NoCollision);
 
 		AddToScene(InOutParams, InstancedStaticMeshComponent);
 	}
diff --git a/Engine/Plugins/Experimental/ProceduralVegetationEditor/Source/ProceduralVegetationEditor/Private/Visualizations/DebugVisualization/PVDebugVisualizationBase.h b/Engine/Plugins/Experimental/ProceduralVegetationEditor/Source/ProceduralVegetationEditor/Private/Visualizations/DebugVisualization/PVDebugVisualizationBase.h
--- a/Engine/Plugins/Experimental/ProceduralVegetationEditor/Source/ProceduralVegetationEditor/Private/Visualizations/DebugVisualization/PVDebugVisualizationBase.h
+++ b/Engine/Plugins/Experimental/ProceduralVegetationEditor/Source/ProceduralVegetationEditor/Private/Visualizations/DebugVisualization/PVDebugVisualizationBase.h
@@ -24,24 +24,21 @@ public:
 protected:
 	void DrawPivotPoints(const FVisualizerDrawContext& InContext);
 	void DrawAttributes(FVisualizerDrawContext& InContext);
-	void AddDrawAsPoint(const FVector& InPos);
+	void AddDrawAsPoint(const FVector& InPos, float PointScale);
 	virtual TArray<FVector3f> GetPivotPositions(const FManagedArrayCollection& InCollection) = 0;
 	virtual void GetPivot(const FManagedArrayCollection& InCollection, const int InIndex, FVector3f& OutPos, float& OutScale) = 0;
 	
-	static void DrawAsText(FPCGSceneSetupParams& InOutParams, const FText& InTextToDraw, const FVector3f& InPos, const float& InScale, const FLinearColor& InColor = FLinearColor::White);
-	static void DrawAsDirection( FPCGSceneSetupParams& InOutParams, const FVector& InStartPos, const FVector& InEndPos, const FLinearColor& InColor = FLinearColor::White);
-	static void DrawAsPoint(FPCGSceneSetupParams& InOutParams,const TArray<FTransform>& InTransforms, const FLinearColor& InColor = FLinearColor::White);
+	static void DrawAsText(FPCGSceneSetupParams& InOutParams, const FText& InTextToDraw, const FVector3f& InPos, float InScale, float TextSize, const FLinearColor& InColor = FLinearColor::White);
+	static void DrawAsDirection(UPVLineBatchComponent* LineBatchComponent, const FVector& InStartPos, const FVector& InEndPos, const FLinearColor& InColor = FLinearColor::White);
+	static void DrawAsPoint(UInstancedStaticMeshComponent* InstancedStaticMeshComponent, const TArray<FTransform>& InTransforms, const FLinearColor& InColor = FLinearColor::White);
 	
 	static UPVLineBatchComponent* GetOrCreateLineComponent(FPCGSceneSetupParams& InOutParams);
 	static UInstancedStaticMeshComponent* GetOrCreatePointComponent(FPCGSceneSetupParams& InOutParams);
 
 	static void AddToScene(FPCGSceneSetupParams& InOutParams, UPrimitiveComponent* InComponent);
 	
+private:
 	TArray<FTransform> PointsToDraw;
-
-	static const float DefaultDirectionStrength;
-	static const float DefaultTextScale;
-	static const float DefaultPointScale;
 };
 typedef TSharedPtr<FPVDebugVisualizationBase> FPVDebugVisualizationPtr;
 
diff --git a/Engine/Plugins/Experimental/ProceduralVegetationEditor/Source/ProceduralVegetationEditor/Private/Visualizations/DebugVisualization/PVDebugVisualizations.cpp b/Engine/Plugins/Experimental/ProceduralVegetationEditor/Source/ProceduralVegetationEditor/Private/Visualizations/DebugVisualization/PVDebugVisualizations.cpp
--- a/Engine/Plugins/Experimental/ProceduralVegetationEditor/Source/ProceduralVegetationEditor/Private/Visualizations/DebugVisualization/PVDebugVisualizations.cpp
+++ b/Engine/Plugins/Experimental/ProceduralVegetationEditor/Source/ProceduralVegetationEditor/Private/Visualizations/DebugVisualization/PVDebugVisualizations.cpp
@@ -5,46 +5,52 @@
 #include "Facades/PVBranchFacade.h"
 #include "Facades/PVFoliageFacade.h"
 #include "Facades/PVPointFacade.h"
+#include "Facades/PVAttributesNames.h"
 
-TArray<FVector3f> FPVPointDebugVisualization::GetPivotPositions(const FManagedArrayCollection& InCollection)
+namespace PV::DebugVisualization
 {
-	const PV::Facades::FPointFacade PointFacade(InCollection);
-	const TManagedArray<FVector3f>& PointPositions = PointFacade.GetPositions();
+	const TArray<FVector3f>& GetVectorAttribute(const FManagedArrayCollection& InCollection, FName AttributeName, FName GroupName)
+	{
+		const TManagedArrayAccessor<FVector3f> Attribute(InCollection, AttributeName, GroupName);
+		if (Attribute.IsValid())
+		{
+			return Attribute.Get().GetConstArray();
+		}
+		
+		static TArray<FVector3f> EmptyArray;
+		return EmptyArray;
+	}
 
-	TArray<FVector3f> Positions;
-	for (int32 i = 0; i < PointPositions.Num(); i++)
+	const TArray<float>& GetFloatAttribute(const FManagedArrayCollection& InCollection, FName AttributeName, FName GroupName)
 	{
-		const FVector3f& Pos = PointPositions[i];
+		const TManagedArrayAccessor<float> Attribute(InCollection, AttributeName, GroupName);
+		if (Attribute.IsValid())
+		{
+			return Attribute.Get().GetConstArray();
+		}
 
-		Positions.Add(Pos);
+		static TArray<float> EmptyArray;
+		return EmptyArray;
 	}
+};
 
-	return Positions;
+TArray<FVector3f> FPVPointDebugVisualization::GetPivotPositions(const FManagedArrayCollection& InCollection)
+{
+	return PV::DebugVisualization::GetVectorAttribute(InCollection, PV::AttributeNames::PointPosition, PV::GroupNames::PointGroup);
 }
 
 void FPVPointDebugVisualization::GetPivot(const FManagedArrayCollection& InCollection, const int InIndex, FVector3f& OutPos, float& OutScale)
 {
-	const PV::Facades::FPointFacade PointFacade(InCollection);
+	const TArray<FVector3f>& PointPositions = PV::DebugVisualization::GetVectorAttribute(InCollection, PV::AttributeNames::PointPosition, PV::GroupNames::PointGroup);
+	const TArray<float>& PointScales = PV::DebugVisualization::GetFloatAttribute(InCollection, PV::AttributeNames::PointScale, PV::GroupNames::PointGroup);
 
-	OutPos = PointFacade.GetPosition(InIndex);
-	OutScale = PointFacade.GetPointScale(InIndex);
+	OutPos = PointPositions.IsValidIndex(InIndex) ? PointPositions[InIndex] : FVector3f::ZeroVector;
+	OutScale = PointScales.IsValidIndex(InIndex) ? PointScales[InIndex] : 1.f;
 }
 
 TArray<FVector3f> FPVFoliageDebugVisualization::GetPivotPositions(const FManagedArrayCollection& InCollection)
 {
-	const PV::Facades::FFoliageFacade FoliageFacade(InCollection);
-	const TManagedArray<FVector3f>& PointPositions =  FoliageFacade.GetPivotPositions();
-	
-	TArray<FVector3f> Positions;
-	for (int32 i = 0; i < PointPositions.Num(); i++)
-	{
-		const FVector3f& Pos = PointPositions[i];
-
-		Positions.Add(Pos);
-	}
-
-	return Positions;
-
+	return PV::DebugVisualization::GetVectorAttribute(InCollection, PV::AttributeNames::FoliagePivotPoint, PV::GroupNames::FoliageGroup);
 }
 
 void FPVFoliageDebugVisualization::GetPivot(const FManagedArrayCollection& InCollection, const int InIndex, FVector3f& OutPos, float& OutScale)
@@ -52,19 +58,22 @@ void FPVFoliageDebugVisualization::GetPivot(const FManagedArrayCollection& InCol
 	const PV::Facades::FFoliageFacade FoliageFacade(InCollection);
 	const PV::Facades::FBranchFacade BranchFacade(InCollection);
 	const PV::Facades::FPointFacade PointFacade(InCollection);
-	
-	PV::Facades::FFoliageEntryData Data =  FoliageFacade.GetFoliageEntry(InIndex);
-	OutPos = Data.PivotPoint;
 
-	const TArray<int32>& BranchPoints = BranchFacade.GetPoints(Data.BranchId);
+	OutPos = FVector3f::ZeroVector;
+	OutScale = 1.f;
 
-	if (BranchPoints.Num() > 0)
+	if (FoliageFacade.IsValid())
 	{
-		OutScale = PointFacade.GetPointScale(BranchPoints[0]);
+		OutPos = FoliageFacade.GetPivotPoint(InIndex);
 	}
-	else
+
+	if (FoliageFacade.IsValid() && BranchFacade.IsValid() && PointFacade.IsValid())
 	{
-		OutScale = 1;
+		const TArray<int32>& BranchPoints = BranchFacade.GetPoints(FoliageFacade.GetFoliageBranchId(InIndex));
+		if (BranchPoints.Num() > 0)
+		{
+			OutScale = PointFacade.GetPointScale(BranchPoints[0]);
+		}
 	}
 }
 
@@ -107,4 +116,18 @@ void FPVBranchDebugVisualization::GetPivot(const FManagedArrayCollection& InColl
 			OutScale = PointFacade.GetPointScale(RootPointIndex);
 		}
 	}
+}
+
+TArray<FVector3f> FPVCustomDebugVisualization::GetPivotPositions(const FManagedArrayCollection& InCollection)
+{
+	return PV::DebugVisualization::GetVectorAttribute(InCollection, PivotPositionAttributeName, PivotGroupName);
+}
+
+void FPVCustomDebugVisualization::GetPivot(const FManagedArrayCollection& InCollection, const int InIndex, FVector3f& OutPos, float& OutScale)
+{
+	const TArray<FVector3f>& PointPositions = PV::DebugVisualization::GetVectorAttribute(InCollection, PivotPositionAttributeName, PivotGroupName);
+	const TArray<float>& PointScales = PV::DebugVisualization::GetFloatAttribute(InCollection, PivotScaleAttributeName, PivotGroupName);
+
+	OutPos = PointPositions.IsValidIndex(InIndex) ? PointPositions[InIndex] : FVector3f::ZeroVector;
+	OutScale = PointScales.IsValidIndex(InIndex) ? PointScales[InIndex] : 1.f;
 }
\ No newline at end of file
diff --git a/Engine/Plugins/Experimental/ProceduralVegetationEditor/Source/ProceduralVegetationEditor/Private/Visualizations/DebugVisualization/PVDebugVisualizations.h b/Engine/Plugins/Experimental/ProceduralVegetationEditor/Source/ProceduralVegetationEditor/Private/Visualizations/DebugVisualization/PVDebugVisualizations.h
--- a/Engine/Plugins/Experimental/ProceduralVegetationEditor/Source/ProceduralVegetationEditor/Private/Visualizations/DebugVisualization/PVDebugVisualizations.h
+++ b/Engine/Plugins/Experimental/ProceduralVegetationEditor/Source/ProceduralVegetationEditor/Private/Visualizations/DebugVisualization/PVDebugVisualizations.h
@@ -35,3 +35,23 @@ protected:
 	virtual TArray<FVector3f> GetPivotPositions(const FManagedArrayCollection& InCollection) override;
 	virtual void GetPivot(const FManagedArrayCollection& InCollection, const int InIndex, FVector3f& OutPos, float& OutScale) override;
 };
+
+class FPVCustomDebugVisualization final : public FPVDebugVisualizationBase
+{
+private:
+	FName PivotPositionAttributeName = NAME_None;
+	FName PivotScaleAttributeName = NAME_None;
+	FName PivotGroupName = NAME_None;
+
+public:
+	FPVCustomDebugVisualization(FName InPivotPositionAttributeName, FName InPivotScaleAttributeName, FName InPivotGroupName) 
+		: PivotPositionAttributeName(InPivotPositionAttributeName)
+		, PivotScaleAttributeName(InPivotScaleAttributeName)
+		, PivotGroupName(InPivotGroupName)
+	{
+	}
+
+protected:
+	virtual TArray<FVector3f> GetPivotPositions(const FManagedArrayCollection& InCollection) override;
+	virtual void GetPivot(const FManagedArrayCollection& InCollection, const int InIndex, FVector3f& OutPos, float& OutScale) override;
+};
diff --git a/Engine/Plugins/Experimental/ProceduralVegetationEditor/Source/ProceduralVegetationEditor/Private/Visualizations/DebugVisualization/PVDebugVisualizer.cpp b/Engine/Plugins/Experimental/ProceduralVegetationEditor/Source/ProceduralVegetationEditor/Private/Visualizations/DebugVisualization/PVDebugVisualizer.cpp
--- a/Engine/Plugins/Experimental/ProceduralVegetationEditor/Source/ProceduralVegetationEditor/Private/Visualizations/DebugVisualization/PVDebugVisualizer.cpp
+++ b/Engine/Plugins/Experimental/ProceduralVegetationEditor/Source/ProceduralVegetationEditor/Private/Visualizations/DebugVisualization/PVDebugVisualizer.cpp
@@ -18,24 +18,24 @@ void FPVDebugVisualizer::Draw(const UPVData* InData, FPCGSceneSetupParams& InOut
 {
 	check(InData);
 	
-	TArray<FPVVisualizationSettings> VisualizationSettings = InData->GetDebugSettings().VisualizationSettings;
+	const TArray<FPVVisualizationSettings>& VisualizationSettings = InData->GetDebugSettings().VisualizationSettings;
 
 	for (int32 i = 0; i < VisualizationSettings.Num(); i++)
 	{
-		auto Settings = VisualizationSettings[i];
+		const auto& Settings = VisualizationSettings[i];
 
-		FVisualizerDrawContext Context(InData->GetCollection(), Settings,InOutParams);
+		FVisualizerDrawContext Context(InData->GetCollection(), Settings, InOutParams);
 
-		if (FPVDebugVisualizationPtr VisualizationPtr = CreateVisualizer(Settings.DebugType))
+		if (FPVDebugVisualizationPtr VisualizationPtr = CreateVisualizer(Settings))
 		{
 			VisualizationPtr->Draw(Context);
 		}	
 	}
 }
 
-FPVDebugVisualizationPtr FPVDebugVisualizer::CreateVisualizer(EPVDebugType InDebugType)
+FPVDebugVisualizationPtr FPVDebugVisualizer::CreateVisualizer(const FPVVisualizationSettings& InSettings)
 {
-	switch (InDebugType)
+	switch (InSettings.DebugType)
 	{
 		case EPVDebugType::Point:
 			return MakeShared<FPVPointDebugVisualization>();
@@ -46,6 +46,9 @@ FPVDebugVisualizationPtr FPVDebugVisualizer::CreateVisualizer(EPVDebugType InDeb
 		case EPVDebugType::Branches:
 			return MakeShared<FPVBranchDebugVisualization>();
 
+		case EPVDebugType::Custom:
+			return MakeShared<FPVCustomDebugVisualization>(InSettings.CustomPivotPositionAttributeName, InSettings.CustomPivotScaleAttributeName, InSettings.CustomGroupToFilter);
+
 		default:
 			return nullptr;
 	}
diff --git a/Engine/Plugins/Experimental/ProceduralVegetationEditor/Source/ProceduralVegetationEditor/Private/Visualizations/DebugVisualization/PVDebugVisualizer.h b/Engine/Plugins/Experimental/ProceduralVegetationEditor/Source/ProceduralVegetationEditor/Private/Visualizations/DebugVisualization/PVDebugVisualizer.h
--- a/Engine/Plugins/Experimental/ProceduralVegetationEditor/Source/ProceduralVegetationEditor/Private/Visualizations/DebugVisualization/PVDebugVisualizer.h
+++ b/Engine/Plugins/Experimental/ProceduralVegetationEditor/Source/ProceduralVegetationEditor/Private/Visualizations/DebugVisualization/PVDebugVisualizer.h
@@ -10,7 +10,6 @@ struct FPVVisualizationSettings;
 class FPVDebugVisualizationBase;
 struct FPCGSceneSetupParams;
 class UPVData;
-enum class EPVDebugType;
 
 struct FVisualizerDrawContext
 {
@@ -23,7 +22,7 @@ struct FVisualizerDrawContext
 class FPVDebugVisualizer
 {
 public:
-	static void Draw(const UPVData* InData,FPCGSceneSetupParams& InOutParams);
+	static void Draw(const UPVData* InData, FPCGSceneSetupParams& InOutParams);
 private:
-	static FPVDebugVisualizationPtr CreateVisualizer(EPVDebugType InDebugType);
+	static FPVDebugVisualizationPtr CreateVisualizer(const FPVVisualizationSettings& InSettings);
 };
\ No newline at end of file
diff --git a/Engine/Plugins/Experimental/ProceduralVegetationEditor/Source/ProceduralVegetationEditor/Private/Visualizations/PVLineBatchComponent.cpp b/Engine/Plugins/Experimental/ProceduralVegetationEditor/Source/ProceduralVegetationEditor/Private/Visualizations/PVLineBatchComponent.cpp
--- a/Engine/Plugins/Experimental/ProceduralVegetationEditor/Source/ProceduralVegetationEditor/Private/Visualizations/PVLineBatchComponent.cpp
+++ b/Engine/Plugins/Experimental/ProceduralVegetationEditor/Source/ProceduralVegetationEditor/Private/Visualizations/PVLineBatchComponent.cpp
@@ -6,7 +6,8 @@
 
 FPVLineSceneProxy::FPVLineSceneProxy(const UPVLineBatchComponent* InComponent)
 	: FPrimitiveSceneProxy(InComponent)
-	, LineComponent(InComponent)
+	, Lines(InComponent->Lines)
+	, PointSize(InComponent->PointSize)
 {}
 
 FPrimitiveViewRelevance FPVLineSceneProxy::GetViewRelevance(const FSceneView* View) const
@@ -25,6 +26,32 @@ void FPVLineSceneProxy::GetDynamicMeshElements(
 	FMeshElementCollector& Collector
 ) const
 {
+	const static auto DrawLines = [](FPrimitiveDrawInterface* PDI, const TArray<FPVLineInfo>& Lines, const float PointSize)
+	{
+		for (const FPVLineInfo& Line : Lines)
+		{
+			PDI->DrawLine(
+				Line.StartPos,
+				Line.EndPos,
+				Line.Color,
+				Line.DepthPriorityGroup,
+				2.0f,
+				0,
+				true
+			);
+
+			if (Line.PointDrawSettings == EPointDrawSettings::Start || Line.PointDrawSettings == EPointDrawSettings::Both)
+			{
+				PDI->DrawPoint(Line.StartPos, FLinearColor::White, PointSize, Line.DepthPriorityGroup);
+			}
+
+			if (Line.PointDrawSettings == EPointDrawSettings::End || Line.PointDrawSettings == EPointDrawSettings::Both)
+			{
+				PDI->DrawPoint(Line.EndPos, FLinearColor::White, PointSize, Line.DepthPriorityGroup);
+			}
+		}
+	};
+
 	for (int32 ViewIndex = 0; ViewIndex < Views.Num(); ViewIndex++)
 	{
 		if (VisibilityMap & (1 << ViewIndex))
@@ -37,7 +64,7 @@ void FPVLineSceneProxy::GetDynamicMeshElements(
 			FPrimitiveDrawInterface* PDI = Collector.GetPDI(ViewIndex);
 #endif
 
-			LineComponent->Draw(View, PDI);
+			DrawLines(PDI, Lines, PointSize);
 		}
 	}
 }
@@ -71,6 +98,7 @@ UPVLineBatchComponent::UPVLineBatchComponent()
 void UPVLineBatchComponent::InitBounds()
 {
 	BBox = FBox(ForceInit);
+	MarkRenderStateDirty();
 }
 
 void UPVLineBatchComponent::AddLine(const FVector& InStartPos, const FVector& InEndPos, const FLinearColor& InColor,const ESceneDepthPriorityGroup InDepthPriorityGroup, const EPointDrawSettings InPointDrawSettings)
@@ -79,38 +107,15 @@ void UPVLineBatchComponent::AddLine(const FVector& InStartPos, const FVector& In
 
 	BBox += InStartPos;
 	BBox += InEndPos;
-}
 
-void UPVLineBatchComponent::Draw(const FSceneView* View, FPrimitiveDrawInterface* PDI) const
-{
-	for (const FPVLineInfo& Line : Lines)
-	{
-		PDI->DrawLine(
-			Line.StartPos,
-			Line.EndPos,
-			Line.Color,
-			Line.DepthPriorityGroup,
-			2.0f,
-			0,
-			true
-		);
-
-		if (Line.PointDrawSettings == EPointDrawSettings::Start || Line.PointDrawSettings == EPointDrawSettings::Both)
-		{
-			PDI->DrawPoint(Line.StartPos, FLinearColor::White, PointSize, Line.DepthPriorityGroup);
-		}
-
-		if (Line.PointDrawSettings == EPointDrawSettings::End || Line.PointDrawSettings == EPointDrawSettings::Both)
-		{
-			PDI->DrawPoint(Line.EndPos, FLinearColor::White, PointSize, Line.DepthPriorityGroup);
-		}
-	}
+	MarkRenderStateDirty();
 }
 
 void UPVLineBatchComponent::Flush()
 {
 	Lines.Empty();
 	BBox = FBox(ForceInit);
+	MarkRenderStateDirty();
 }
 
 FPrimitiveSceneProxy* UPVLineBatchComponent::CreateSceneProxy()
diff --git a/Engine/Plugins/Experimental/ProceduralVegetationEditor/Source/ProceduralVegetationEditor/Private/Visualizations/PVLineBatchComponent.h b/Engine/Plugins/Experimental/ProceduralVegetationEditor/Source/ProceduralVegetationEditor/Private/Visualizations/PVLineBatchComponent.h
--- a/Engine/Plugins/Experimental/ProceduralVegetationEditor/Source/ProceduralVegetationEditor/Private/Visualizations/PVLineBatchComponent.h
+++ b/Engine/Plugins/Experimental/ProceduralVegetationEditor/Source/ProceduralVegetationEditor/Private/Visualizations/PVLineBatchComponent.h
@@ -12,8 +12,29 @@
 struct FManagedArrayCollection;
 class UPVLineBatchComponent;
 
+enum class EPointDrawSettings : uint8
+{
+	None,
+	Start,
+	End,
+	Both
+};
+
+struct FPVLineInfo
+{
+	FVector StartPos;
+	FVector EndPos;
+	FLinearColor Color;
+	ESceneDepthPriorityGroup DepthPriorityGroup = SDPG_World;
+	EPointDrawSettings PointDrawSettings = EPointDrawSettings::None;
+};
+
 class FPVLineSceneProxy final : public FPrimitiveSceneProxy
 {
+private:
+	TArray<FPVLineInfo> Lines;
+	float PointSize;
+
 public:
 	FPVLineSceneProxy(const UPVLineBatchComponent* InComponent);
 
@@ -27,32 +48,7 @@ public:
 
 	virtual SIZE_T GetTypeHash() const override;
 	virtual uint32 GetMemoryFootprint(void) const override;
-
 	virtual bool CanBeOccluded() const override;
-
-private:
-	const TObjectPtr<const UPVLineBatchComponent> LineComponent;
-};
-
-UENUM()
-enum class EPointDrawSettings : uint8
-{
-	None,
-	Start,
-	End,
-	Both
-};
-
-USTRUCT()
-struct FPVLineInfo
-{
-	GENERATED_BODY()
-
-	FVector StartPos;
-	FVector EndPos;
-	FLinearColor Color;
-	ESceneDepthPriorityGroup DepthPriorityGroup = SDPG_World;
-	EPointDrawSettings PointDrawSettings = EPointDrawSettings::None;
 };
 
 UCLASS(Blueprintable, BlueprintType, meta=(BlueprintSpawnableComponent), MinimalAPI)
@@ -60,6 +56,9 @@ class UPVLineBatchComponent : public UPrimitiveComponent
 {
 	GENERATED_BODY()
 
+private:
+	friend class FPVLineSceneProxy;
+
 public:
 	UPVLineBatchComponent();
 
@@ -68,8 +67,6 @@ public:
 				const ESceneDepthPriorityGroup InDepthPriorityGroup = ESceneDepthPriorityGroup::SDPG_World,
 				const EPointDrawSettings InPointDrawSettings = EPointDrawSettings::None);
 
-	void Draw(const FSceneView* View, FPrimitiveDrawInterface* PDI) const;
-
 	void Flush();
 	
 	virtual FPrimitiveSceneProxy* CreateSceneProxy() override;
@@ -78,7 +75,7 @@ public:
 protected:
 	FBox BBox;
 	TArray<FPVLineInfo> Lines;
-
+	
 private:
 	const float PointSize = 5.0f;
 };
