Refactor PVE collection debug visualization
- Make attribute debug drawing more generic to allow for drawing of more attribute types
- Add debug draw mode for sphere

#Synced-CL 48398234

diff --git a/Engine/Plugins/Experimental/ProceduralVegetationEditor/Source/ProceduralVegetation/Public/DataTypes/PVData.h b/Engine/Plugins/Experimental/ProceduralVegetationEditor/Source/ProceduralVegetation/Public/DataTypes/PVData.h
--- a/Engine/Plugins/Experimental/ProceduralVegetationEditor/Source/ProceduralVegetation/Public/DataTypes/PVData.h
+++ b/Engine/Plugins/Experimental/ProceduralVegetationEditor/Source/ProceduralVegetation/Public/DataTypes/PVData.h
@@ -41,7 +41,8 @@ enum class EPVDebugValueVisualizationMode
 {
 	Text,
 	Direction,
-	Point
+	Point,
+	Sphere
 };
 
 USTRUCT()
@@ -55,6 +56,9 @@ struct FPVVisualizationSettings
 	UPROPERTY(EditAnywhere, Category = DebugSettings)
 	bool bShowAnchorPoints = true;
 
+	UPROPERTY(EditAnywhere, Category = DebugSettings, meta = (EditCondition = "VisualizationMode==EPVDebugValueVisualizationMode::Point || bShowAnchorPoints"))
+	bool bDrawPointAsMesh = true;
+
 	UPROPERTY(EditAnywhere, Category = DebugSettings)
 	FName AttributeToFilter;
 
@@ -73,14 +77,14 @@ struct FPVVisualizationSettings
 	UPROPERTY(EditAnywhere, Category = DebugSettings)
 	FColor Color = FColor::Black;
 
-	UPROPERTY(EditAnywhere, Category = DebugSettings)
-	float DirectionStrength = 8.f;
+	UPROPERTY(EditAnywhere, Category = DebugSettings, meta = (EditCondition = "bShowAnchorPoints"))
+	float AnchorPointsScale = 1.f;
 
-	UPROPERTY(EditAnywhere, Category = DebugSettings)
+	UPROPERTY(EditAnywhere, Category = DebugSettings, meta = (EditCondition = "VisualizationMode==EPVDebugValueVisualizationMode::Text"))
 	float TextScale = 2.5f;
 
 	UPROPERTY(EditAnywhere, Category = DebugSettings)
-	float PointScale = 0.005f;
+	float GizmoScale = 1.f;
 	
 	FName GetGroupToFilter() const
 	{
diff --git a/Engine/Plugins/Experimental/ProceduralVegetationEditor/Source/ProceduralVegetationEditor/Private/Visualizations/DebugVisualization/PVDataAttributeValue.cpp b/Engine/Plugins/Experimental/ProceduralVegetationEditor/Source/ProceduralVegetationEditor/Private/Visualizations/DebugVisualization/PVDataAttributeValue.cpp
deleted file mode 100644
diff --git a/Engine/Plugins/Experimental/ProceduralVegetationEditor/Source/ProceduralVegetationEditor/Private/Visualizations/DebugVisualization/PVDataAttributeValue.h b/Engine/Plugins/Experimental/ProceduralVegetationEditor/Source/ProceduralVegetationEditor/Private/Visualizations/DebugVisualization/PVDataAttributeValue.h
deleted file mode 100644
diff --git a/Engine/Plugins/Experimental/ProceduralVegetationEditor/Source/ProceduralVegetationEditor/Private/Visualizations/DebugVisualization/PVDataAttributeValueInterface.cpp b/Engine/Plugins/Experimental/ProceduralVegetationEditor/Source/ProceduralVegetationEditor/Private/Visualizations/DebugVisualization/PVDataAttributeValueInterface.cpp
deleted file mode 100644
diff --git a/Engine/Plugins/Experimental/ProceduralVegetationEditor/Source/ProceduralVegetationEditor/Private/Visualizations/DebugVisualization/PVDataAttributeValueInterface.h b/Engine/Plugins/Experimental/ProceduralVegetationEditor/Source/ProceduralVegetationEditor/Private/Visualizations/DebugVisualization/PVDataAttributeValueInterface.h
deleted file mode 100644
diff --git a/Engine/Plugins/Experimental/ProceduralVegetationEditor/Source/ProceduralVegetationEditor/Private/Visualizations/DebugVisualization/PVDebugVisualizationBase.cpp b/Engine/Plugins/Experimental/ProceduralVegetationEditor/Source/ProceduralVegetationEditor/Private/Visualizations/DebugVisualization/PVDebugVisualizationBase.cpp
--- a/Engine/Plugins/Experimental/ProceduralVegetationEditor/Source/ProceduralVegetationEditor/Private/Visualizations/DebugVisualization/PVDebugVisualizationBase.cpp
+++ b/Engine/Plugins/Experimental/ProceduralVegetationEditor/Source/ProceduralVegetationEditor/Private/Visualizations/DebugVisualization/PVDebugVisualizationBase.cpp
@@ -3,7 +3,6 @@
 #include "PVDebugVisualizationBase.h"
 
 #include "AdvancedPreviewScene.h"
-#include "PVDataAttributeValueInterface.h"
 #include "PVDebugVisualizer.h"
 
 #include "Components/InstancedStaticMeshComponent.h"
@@ -21,6 +20,411 @@
 
 DEFINE_LOG_CATEGORY(LogPVDebugVisualization);
 
+namespace PVE::AttributeToString
+{
+	FString AttributeValueToString(const float& Value) { return FString::SanitizeFloat(Value, 2); }
+	FString AttributeValueToString(const double& Value) { return FString::SanitizeFloat(Value, 2); }
+	FString AttributeValueToString(const uint8& Value) { return FString::FromInt((int32)Value); }
+	FString AttributeValueToString(const int32& Value) { return FString::FromInt(Value); }
+	FString AttributeValueToString(const FString& Value) { return Value; }
+	FString AttributeValueToString(const FName& Value) { return Value.ToString(); }
+	FString AttributeValueToString(const FLinearColor& Value) { return Value.ToString(); }
+	FString AttributeValueToString(const FVector3f& Value) { return Value.ToString(); }
+	FString AttributeValueToString(const FVector3d& Value) { return Value.ToString(); }
+	FString AttributeValueToString(const FVector4f& Value) { return Value.ToString(); }
+	FString AttributeValueToString(const FVector4d& Value) { return Value.ToString(); }
+	FString AttributeValueToString(const bool& Value) { return Value ? FString("true") : FString("false"); }
+	FString AttributeValueToString(const FTransform3f& Value) { return Value.ToString(); }
+	FString AttributeValueToString(const FTransform& Value) { return Value.ToString(); }
+	FString AttributeValueToString(const FBox& Value) { return Value.ToString(); } 
+	FString AttributeValueToString(const FVector2f& Value) { return Value.ToString(); }
+	FString AttributeValueToString(const FIntVector& Value) { return Value.ToString(); }
+	FString AttributeValueToString(const FIntVector2& Value) { return Value.ToString(); }
+	FString AttributeValueToString(const FIntVector4& Value) { return Value.ToString(); }
+	FString AttributeValueToString(const FSoftObjectPath& Value) { return Value.ToString(); }
+
+	template<typename T>
+	FString AttributeValueToString(const TArray<T>& Array)
+	{
+		FString Out;
+
+		for (int32 Idx = 0; Idx < Array.Num(); ++Idx)
+		{
+			Out += AttributeValueToString(Array[Idx]);
+
+			if (Idx != Array.Num() - 1)
+			{
+				Out += "; ";
+			}
+		}
+
+		return Out;
+	}
+
+	template<typename T>
+	FString AttributeValueToString(const TSet<T>& Value)
+	{
+		return AttributeValueToString(Value.Array());
+	}
+
+	template<typename T>
+	TArray<FString> AttributeToString(const FManagedArrayCollection& InCollection, const FName& InAttributeName, const FName& InGroupName)
+	{
+		const TManagedArray<T>& AttributesArray = InCollection.GetAttribute<T>(InAttributeName, InGroupName);
+
+		TArray<FString> OutStringArray;
+		OutStringArray.Reserve(AttributesArray.Num());
+
+		for (int32 i = 0; i < AttributesArray.Num(); ++i)
+		{
+			FString ValueAsString = AttributeValueToString(AttributesArray[i]);
+			
+			// Clip really long strings
+			constexpr int32 MaxStringLength = 1000;
+			if (ValueAsString.Len() > MaxStringLength)
+			{
+				ValueAsString.LeftInline(MaxStringLength);
+				ValueAsString += "...";
+			}
+			
+			OutStringArray.Add(ValueAsString);
+		}
+
+		return OutStringArray;
+	}
+
+	TArray<FString> ConvertAttributeToString(const FManagedArrayCollection& InCollection, const FName& InAttributeName, const FName& InGroupName)
+	{
+		const FManagedArrayCollection::EArrayType ArrayType = InCollection.GetAttributeType(InAttributeName, InGroupName);
+
+		TArray<FString> AttributeAsString;
+
+		switch (ArrayType)
+		{
+		case FManagedArrayCollection::EArrayType::FFloatType:
+			AttributeAsString = AttributeToString<float>(InCollection, InAttributeName, InGroupName);
+			break;
+		case FManagedArrayCollection::EArrayType::FDoubleType:
+			AttributeAsString = AttributeToString<double>(InCollection, InAttributeName, InGroupName);
+			break;
+		case FManagedArrayCollection::EArrayType::FUInt8Type:
+			AttributeAsString = AttributeToString<uint8>(InCollection, InAttributeName, InGroupName);
+			break;
+		case FManagedArrayCollection::EArrayType::FInt32Type:
+			AttributeAsString = AttributeToString<int32>(InCollection, InAttributeName, InGroupName);
+			break;
+		case FManagedArrayCollection::EArrayType::FBoolType:
+			AttributeAsString = AttributeToString<bool>(InCollection, InAttributeName, InGroupName);
+			break;
+		case FManagedArrayCollection::EArrayType::FStringType:
+			AttributeAsString = AttributeToString<FString>(InCollection, InAttributeName, InGroupName);
+			break;
+		case FManagedArrayCollection::EArrayType::FNameType:
+			AttributeAsString = AttributeToString<FName>(InCollection, InAttributeName, InGroupName);
+			break;
+		case FManagedArrayCollection::EArrayType::FSoftObjectPathType:
+			AttributeAsString = AttributeToString<FSoftObjectPath>(InCollection, InAttributeName, InGroupName);
+			break;
+		case FManagedArrayCollection::EArrayType::FLinearColorType:
+			AttributeAsString = AttributeToString<FLinearColor>(InCollection, InAttributeName, InGroupName);
+			break;
+		case FManagedArrayCollection::EArrayType::FVectorType:
+			AttributeAsString = AttributeToString<FVector3f>(InCollection, InAttributeName, InGroupName);
+			break;
+		case FManagedArrayCollection::EArrayType::FVector2DType:
+			AttributeAsString = AttributeToString<FVector2f>(InCollection, InAttributeName, InGroupName);
+			break;
+		case FManagedArrayCollection::EArrayType::FVector3dType:
+			AttributeAsString = AttributeToString<FVector3d>(InCollection, InAttributeName, InGroupName);
+			break;
+		case FManagedArrayCollection::EArrayType::FVector4fType:
+			AttributeAsString = AttributeToString<FVector4f>(InCollection, InAttributeName, InGroupName);
+			break;
+		case FManagedArrayCollection::EArrayType::FIntVectorType:
+			AttributeAsString = AttributeToString<FIntVector>(InCollection, InAttributeName, InGroupName);
+			break;
+		case FManagedArrayCollection::EArrayType::FIntVector2Type:
+			AttributeAsString = AttributeToString<FIntVector2>(InCollection, InAttributeName, InGroupName);
+			break;
+		case FManagedArrayCollection::EArrayType::FFVectorArrayType:
+			AttributeAsString = AttributeToString<TArray<FVector3f>>(InCollection, InAttributeName, InGroupName);
+			break;
+		case FManagedArrayCollection::EArrayType::FIntVector3ArrayType:
+			AttributeAsString = AttributeToString<TArray<FIntVector3>>(InCollection, InAttributeName, InGroupName);
+			break;
+		case FManagedArrayCollection::EArrayType::FIntVector2ArrayType:
+			AttributeAsString = AttributeToString<TArray<FIntVector2>>(InCollection, InAttributeName, InGroupName);
+			break;
+		case FManagedArrayCollection::EArrayType::FVector4fArrayType:
+			AttributeAsString = AttributeToString<TArray<FVector4f>>(InCollection, InAttributeName, InGroupName);
+			break;
+		case FManagedArrayCollection::EArrayType::FTransformType:
+			AttributeAsString = AttributeToString<FTransform>(InCollection, InAttributeName, InGroupName);
+			break;
+		case FManagedArrayCollection::EArrayType::FVector2DArrayType:
+			AttributeAsString = AttributeToString<TArray<FVector2f>>(InCollection, InAttributeName, InGroupName);
+			break;
+		case FManagedArrayCollection::EArrayType::FIntArrayType:
+			AttributeAsString = AttributeToString<TSet<int32>>(InCollection, InAttributeName, InGroupName);
+			break;
+		case FManagedArrayCollection::EArrayType::FInt32ArrayType:
+			AttributeAsString = AttributeToString<TArray<int32>>(InCollection, InAttributeName, InGroupName);
+			break;
+		case FManagedArrayCollection::EArrayType::FFloatArrayType:
+			AttributeAsString = AttributeToString<TArray<float>>(InCollection, InAttributeName, InGroupName);
+			break;
+		case FManagedArrayCollection::EArrayType::FBoxType:
+			AttributeAsString = AttributeToString<FBox>(InCollection, InAttributeName, InGroupName);
+			break;
+		case FManagedArrayCollection::EArrayType::FTransform3fType:
+			AttributeAsString = AttributeToString<FTransform3f>(InCollection, InAttributeName, InGroupName);
+			break;
+		case FManagedArrayCollection::EArrayType::FIntVector4Type:
+			AttributeAsString = AttributeToString<FIntVector4>(InCollection, InAttributeName, InGroupName);
+			break;
+		default:
+			break;
+		}
+
+		return AttributeAsString;
+	}
+};
+
+namespace PVE::AttributeToScalar
+{
+	double AttributeValueToScalar(const float& Value) { return Value; }
+	double AttributeValueToScalar(const double& Value) { return Value; }
+	double AttributeValueToScalar(const uint8& Value) { return double(Value); }
+	double AttributeValueToScalar(const int32& Value) { return double(Value); }
+	double AttributeValueToScalar(const bool& Value) { return Value ? 1 : 0; }
+	double AttributeValueToScalar(const FString& Value) { return FCString::IsNumeric(*Value) ? FCString::Atof(*Value) : 0; }
+	double AttributeValueToScalar(const FName& Value) { return AttributeValueToScalar(Value.ToString()); }
+	double AttributeValueToScalar(const FVector2f& Value) { return Value.Size(); }
+	double AttributeValueToScalar(const FVector3f& Value) { return Value.Size(); }
+	double AttributeValueToScalar(const FVector3d& Value) { return Value.Size(); }
+	double AttributeValueToScalar(const FVector4f& Value) { return Value.Size(); }
+	double AttributeValueToScalar(const FIntVector& Value) { return Value.Size(); }
+
+	template<typename T>
+	TArray<double> AttributeToScalar(const FManagedArrayCollection& InCollection, const FName& InAttributeName, const FName& InGroupName)
+	{
+		const TManagedArray<T>& AttributesArray = InCollection.GetAttribute<T>(InAttributeName, InGroupName);
+
+		TArray<double> OutScalarArray;
+		OutScalarArray.Reserve(AttributesArray.Num());
+
+		for (int32 i = 0; i < AttributesArray.Num(); ++i)
+		{
+			const double ValueAsScalar = AttributeValueToScalar(AttributesArray[i]);
+			OutScalarArray.Add(ValueAsScalar);
+		}
+
+		return OutScalarArray;
+	}
+
+	TArray<double> ConvertAttributeToScalar(const FManagedArrayCollection& InCollection, const FName& InAttributeName, const FName& InGroupName)
+	{
+		const FManagedArrayCollection::EArrayType ArrayType = InCollection.GetAttributeType(InAttributeName, InGroupName);
+
+		TArray<double> AttributeAsScalar;
+
+		switch (ArrayType)
+		{
+		case FManagedArrayCollection::EArrayType::FFloatType:
+			AttributeAsScalar = AttributeToScalar<float>(InCollection, InAttributeName, InGroupName);
+			break;
+		case FManagedArrayCollection::EArrayType::FDoubleType:
+			AttributeAsScalar = AttributeToScalar<double>(InCollection, InAttributeName, InGroupName);
+			break;
+		case FManagedArrayCollection::EArrayType::FUInt8Type:
+			AttributeAsScalar = AttributeToScalar<uint8>(InCollection, InAttributeName, InGroupName);
+			break;
+		case FManagedArrayCollection::EArrayType::FInt32Type:
+			AttributeAsScalar = AttributeToScalar<int32>(InCollection, InAttributeName, InGroupName);
+			break;
+		case FManagedArrayCollection::EArrayType::FBoolType:
+			AttributeAsScalar = AttributeToScalar<bool>(InCollection, InAttributeName, InGroupName);
+			break;
+		case FManagedArrayCollection::EArrayType::FStringType:
+			AttributeAsScalar = AttributeToScalar<FString>(InCollection, InAttributeName, InGroupName);
+			break;
+		case FManagedArrayCollection::EArrayType::FNameType:
+			AttributeAsScalar = AttributeToScalar<FName>(InCollection, InAttributeName, InGroupName);
+			break;
+		case FManagedArrayCollection::EArrayType::FVector2DType:
+			AttributeAsScalar = AttributeToScalar<FVector2f>(InCollection, InAttributeName, InGroupName);
+			break;
+		case FManagedArrayCollection::EArrayType::FVectorType:
+			AttributeAsScalar = AttributeToScalar<FVector3f>(InCollection, InAttributeName, InGroupName);
+			break;
+		case FManagedArrayCollection::EArrayType::FVector3dType:
+			AttributeAsScalar = AttributeToScalar<FVector3d>(InCollection, InAttributeName, InGroupName);
+			break;
+		case FManagedArrayCollection::EArrayType::FVector4fType:
+			AttributeAsScalar = AttributeToScalar<FVector4f>(InCollection, InAttributeName, InGroupName);
+			break;
+		case FManagedArrayCollection::EArrayType::FIntVectorType:
+			AttributeAsScalar = AttributeToScalar<FIntVector>(InCollection, InAttributeName, InGroupName);
+			break;
+		default:
+			break;
+		}
+
+		return AttributeAsScalar;
+	}
+};
+
+namespace PVE::AttributeToVector3
+{
+	typedef TArray<FVector, TInlineAllocator<1>> FVectorValue;
+
+	FVectorValue AttributeValueToVector3(const float& Value) { return { FVector(Value) }; }
+	FVectorValue AttributeValueToVector3(const double& Value) { return { FVector(Value) }; }
+	FVectorValue AttributeValueToVector3(const uint8& Value) { return { FVector(Value) }; }
+	FVectorValue AttributeValueToVector3(const int32& Value) { return { FVector(Value) }; }
+	FVectorValue AttributeValueToVector3(const FString& Value) { FVector OutVec(EForceInit::ForceInitToZero); OutVec.InitFromString(Value); return { OutVec }; }
+	FVectorValue AttributeValueToVector3(const FName& Value) { return AttributeValueToVector3(Value.ToString()); }
+	FVectorValue AttributeValueToVector3(const FLinearColor& Value) { return { FVector(Value.R, Value.G, Value.B) }; }
+	FVectorValue AttributeValueToVector3(const FVector3f& Value) { return { FVector(Value) }; }
+	FVectorValue AttributeValueToVector3(const FVector3d& Value) { return { Value }; }
+	FVectorValue AttributeValueToVector3(const FVector4f& Value) { return { FVector(Value.X, Value.Y, Value.Z) }; }
+	FVectorValue AttributeValueToVector3(const FVector4d& Value) { return { FVector(Value.X, Value.Y, Value.Z) }; }
+	FVectorValue AttributeValueToVector3(const bool& Value) { return { FVector(double(Value)) }; }
+	FVectorValue AttributeValueToVector3(const FTransform3f& Value) { return { FVector(Value.GetLocation()) }; }
+	FVectorValue AttributeValueToVector3(const FTransform& Value) { return { Value.GetLocation() }; }
+	FVectorValue AttributeValueToVector3(const FVector2f& Value) { return { FVector(Value.X, Value.Y, 0) }; }
+	FVectorValue AttributeValueToVector3(const FIntVector& Value) { return { FVector(Value.X, Value.Y, Value.Z) }; }
+	FVectorValue AttributeValueToVector3(const FIntVector2& Value) { return { FVector(Value.X, Value.Y, 0) }; }
+	FVectorValue AttributeValueToVector3(const FIntVector4& Value) { return { FVector(Value.X, Value.Y, Value.Z) }; }
+
+	template<typename T>
+	FVectorValue AttributeValueToVector3(const TArray<T>& Array)
+	{
+		FVectorValue Out;
+		Out.Reserve(Array.Num());
+
+		for (int32 Idx = 0; Idx < Array.Num(); ++Idx)
+		{
+			Out.Append(AttributeValueToVector3(Array[Idx]));
+		}
+
+		return Out;
+	}
+
+	template<typename T>
+	FVectorValue AttributeValueToVector3(const TSet<T>& Value)
+	{
+		return AttributeValueToVector3(Value.Array());
+	}
+
+	template<typename T>
+	TArray<FVectorValue> AttributeToVector3(const FManagedArrayCollection& InCollection, const FName& InAttributeName, const FName& InGroupName)
+	{
+		const TManagedArray<T>& AttributesArray = InCollection.GetAttribute<T>(InAttributeName, InGroupName);
+
+		TArray<FVectorValue> OutVectorArray;
+		OutVectorArray.Reserve(AttributesArray.Num());
+
+		for (int32 i = 0; i < AttributesArray.Num(); ++i)
+		{
+			FVectorValue ValueAsString = AttributeValueToVector3(AttributesArray[i]);
+			OutVectorArray.Add(MoveTemp(ValueAsString));
+		}
+
+		return OutVectorArray;
+	}
+
+	TArray<FVectorValue> ConvertAttributeToVector3(const FManagedArrayCollection& InCollection, const FName& InAttributeName, const FName& InGroupName)
+	{
+		const FManagedArrayCollection::EArrayType ArrayType = InCollection.GetAttributeType(InAttributeName, InGroupName);
+
+		TArray<FVectorValue> AttributeAsVector;
+
+		switch (ArrayType)
+		{
+		case FManagedArrayCollection::EArrayType::FFloatType:
+			AttributeAsVector = AttributeToVector3<float>(InCollection, InAttributeName, InGroupName);
+			break;
+		case FManagedArrayCollection::EArrayType::FDoubleType:
+			AttributeAsVector = AttributeToVector3<double>(InCollection, InAttributeName, InGroupName);
+			break;
+		case FManagedArrayCollection::EArrayType::FUInt8Type:
+			AttributeAsVector = AttributeToVector3<uint8>(InCollection, InAttributeName, InGroupName);
+			break;
+		case FManagedArrayCollection::EArrayType::FInt32Type:
+			AttributeAsVector = AttributeToVector3<int32>(InCollection, InAttributeName, InGroupName);
+			break;
+		case FManagedArrayCollection::EArrayType::FBoolType:
+			AttributeAsVector = AttributeToVector3<bool>(InCollection, InAttributeName, InGroupName);
+			break;
+		case FManagedArrayCollection::EArrayType::FStringType:
+			AttributeAsVector = AttributeToVector3<FString>(InCollection, InAttributeName, InGroupName);
+			break;
+		case FManagedArrayCollection::EArrayType::FNameType:
+			AttributeAsVector = AttributeToVector3<FName>(InCollection, InAttributeName, InGroupName);
+			break;
+		case FManagedArrayCollection::EArrayType::FLinearColorType:
+			AttributeAsVector = AttributeToVector3<FLinearColor>(InCollection, InAttributeName, InGroupName);
+			break;
+		case FManagedArrayCollection::EArrayType::FVectorType:
+			AttributeAsVector = AttributeToVector3<FVector3f>(InCollection, InAttributeName, InGroupName);
+			break;
+		case FManagedArrayCollection::EArrayType::FVector2DType:
+			AttributeAsVector = AttributeToVector3<FVector2f>(InCollection, InAttributeName, InGroupName);
+			break;
+		case FManagedArrayCollection::EArrayType::FVector3dType:
+			AttributeAsVector = AttributeToVector3<FVector3d>(InCollection, InAttributeName, InGroupName);
+			break;
+		case FManagedArrayCollection::EArrayType::FVector4fType:
+			AttributeAsVector = AttributeToVector3<FVector4f>(InCollection, InAttributeName, InGroupName);
+			break;
+		case FManagedArrayCollection::EArrayType::FIntVectorType:
+			AttributeAsVector = AttributeToVector3<FIntVector>(InCollection, InAttributeName, InGroupName);
+			break;
+		case FManagedArrayCollection::EArrayType::FIntVector2Type:
+			AttributeAsVector = AttributeToVector3<FIntVector2>(InCollection, InAttributeName, InGroupName);
+			break;
+		case FManagedArrayCollection::EArrayType::FFVectorArrayType:
+			AttributeAsVector = AttributeToVector3<TArray<FVector3f>>(InCollection, InAttributeName, InGroupName);
+			break;
+		case FManagedArrayCollection::EArrayType::FIntVector3ArrayType:
+			AttributeAsVector = AttributeToVector3<TArray<FIntVector3>>(InCollection, InAttributeName, InGroupName);
+			break;
+		case FManagedArrayCollection::EArrayType::FIntVector2ArrayType:
+			AttributeAsVector = AttributeToVector3<TArray<FIntVector2>>(InCollection, InAttributeName, InGroupName);
+			break;
+		case FManagedArrayCollection::EArrayType::FVector4fArrayType:
+			AttributeAsVector = AttributeToVector3<TArray<FVector4f>>(InCollection, InAttributeName, InGroupName);
+			break;
+		case FManagedArrayCollection::EArrayType::FTransformType:
+			AttributeAsVector = AttributeToVector3<FTransform>(InCollection, InAttributeName, InGroupName);
+			break;
+		case FManagedArrayCollection::EArrayType::FVector2DArrayType:
+			AttributeAsVector = AttributeToVector3<TArray<FVector2f>>(InCollection, InAttributeName, InGroupName);
+			break;
+		case FManagedArrayCollection::EArrayType::FIntArrayType:
+			AttributeAsVector = AttributeToVector3<TSet<int32>>(InCollection, InAttributeName, InGroupName);
+			break;
+		case FManagedArrayCollection::EArrayType::FInt32ArrayType:
+			AttributeAsVector = AttributeToVector3<TArray<int32>>(InCollection, InAttributeName, InGroupName);
+			break;
+		case FManagedArrayCollection::EArrayType::FFloatArrayType:
+			AttributeAsVector = AttributeToVector3<TArray<float>>(InCollection, InAttributeName, InGroupName);
+			break;
+		case FManagedArrayCollection::EArrayType::FTransform3fType:
+			AttributeAsVector = AttributeToVector3<FTransform3f>(InCollection, InAttributeName, InGroupName);
+			break;
+		case FManagedArrayCollection::EArrayType::FIntVector4Type:
+			AttributeAsVector = AttributeToVector3<FIntVector4>(InCollection, InAttributeName, InGroupName);
+			break;
+		default:
+			break;
+		}
+
+		return AttributeAsVector;
+	}
+};
+
 void FPVDebugVisualizationBase::Draw(FVisualizerDrawContext& InContext)
 {
 	DrawPivotPoints(InContext);
@@ -38,7 +442,7 @@ void FPVDebugVisualizationBase::DrawPivotPoints(const FVisualizerDrawContext& In
 		for (int32 i = 0; i < Positions.Num(); i++)
 		{
 			const FVector Pos(Positions[i]);
-			AddDrawAsPoint(Pos, Settings.PointScale);
+			AddDrawAsPoint(Pos, Settings.AnchorPointsScale);
 		}	
 	}
 }
@@ -51,48 +455,46 @@ void FPVDebugVisualizationBase::DrawAttributes(FVisualizerDrawContext& InContext
 	const FName AttributeName = VisualizationSettings.AttributeToFilter;
 	const FName GroupName = VisualizationSettings.GetGroupToFilter();
 
-	if (GroupName.IsNone() || AttributeName.IsNone())
-	{
-		return;
-	}
-
-	if (!Collection.HasAttribute(AttributeName, GroupName) || !Collection.HasGroup(GroupName))
+	if (GroupName.IsNone() || AttributeName.IsNone() || !Collection.HasAttribute(AttributeName, GroupName) || !Collection.HasGroup(GroupName))
 	{
 		return;
 	}
 	
-	const FManagedArrayCollection::EArrayType AttributeType = Collection.GetAttributeType(AttributeName, GroupName);
-	const FPVAttributeValuePtr AttributeValue = FPVAttributeValueInterface::Create(GroupName, AttributeName, AttributeType, Collection);
-
-	if (!AttributeValue)
-	{
-		return;
-	}
-
 	const int32 NumElements = Collection.NumElements(GroupName);
 
-	if (AttributeValue->IsScalar() || VisualizationSettings.VisualizationMode == EPVDebugValueVisualizationMode::Text)
+	switch (VisualizationSettings.VisualizationMode)
 	{
-		const TArray<FText> TextValueArray = AttributeValue->ToText();
-		check(TextValueArray.Num() == NumElements);
-		
-		for (int32 i = 0; i < NumElements; i++)
+		case EPVDebugValueVisualizationMode::Text:
 		{
-			FVector3f OutPos;
-			float OutScale;
-			GetPivot(Collection, i, OutPos, OutScale);
+			using namespace PVE::AttributeToString;
 
-			DrawAsText(InContext.SceneSetupParams, TextValueArray[i], OutPos, OutScale, VisualizationSettings.TextScale, VisualizationSettings.Color);
-		}
-	}
-	else if (VisualizationSettings.VisualizationMode == EPVDebugValueVisualizationMode::Direction)
-	{
-		TObjectPtr<UPVLineBatchComponent> LineBatchComponent = GetOrCreateLineComponent(InContext.SceneSetupParams);
+			TArray<FString> AttributeAsString = ConvertAttributeToString(Collection, AttributeName, GroupName);
+			if (AttributeAsString.Num() != NumElements)
+			{
+				break;
+			}
 
-		if (AttributeType == FManagedArrayCollection::EArrayType::FFVectorArrayType)
+			for (int32 i = 0; i < NumElements; i++)
+			{
+				FVector3f OutPos;
+				float OutScale;
+				GetPivot(Collection, i, OutPos, OutScale);
+
+				DrawAsText(InContext.SceneSetupParams, FText::FromString(AttributeAsString[i]), OutPos, OutScale, VisualizationSettings.TextScale, VisualizationSettings.Color);
+			}
+			break;
+		}
+		case EPVDebugValueVisualizationMode::Direction:
 		{
-			const TArray<TArray<FVector>> VectorArrays = AttributeValue->ToVectorArrays();
-			check(VectorArrays.Num() == NumElements);
+			using namespace PVE::AttributeToVector3;
+
+			const TArray<FVectorValue> AttributeAsVector = ConvertAttributeToVector3(Collection, AttributeName, GroupName);
+			if (AttributeAsVector.Num() != NumElements)
+			{
+				break;
+			}
+
+			TObjectPtr<UPVLineBatchComponent> LineBatchComponent = GetOrCreateLineComponent(InContext.SceneSetupParams);
 
 			for (int32 i = 0; i < NumElements; i++)
 			{
@@ -100,22 +502,53 @@ void FPVDebugVisualizationBase::DrawAttributes(FVisualizerDrawContext& InContext
 				float OutScale;
 				GetPivot(Collection, i, OutPos, OutScale);
 
-				for (int32 j = 0; j < VectorArrays[i].Num(); ++j)
+				for (int32 j = 0; j < AttributeAsVector[i].Num(); ++j)
 				{
-					const FVector& DirectionVector = VectorArrays[i][j];
-					FVector StartPos(OutPos);
-					FVector EndPos = StartPos + DirectionVector * VisualizationSettings.DirectionStrength;
+					const float BaseDirectionStrength = 8.f;
 
-					FColor DrawColor = FLinearColor::LerpUsingHSV(FLinearColor::Black, VisualizationSettings.Color.ReinterpretAsLinear(), j / (float)VectorArrays[i].Num()).ToFColor(true);
+					const FVector& DirectionVector = AttributeAsVector[i][j];
+					const FVector StartPos(OutPos);
+					const FVector EndPos = StartPos + DirectionVector * BaseDirectionStrength * VisualizationSettings.GizmoScale;
+
+					const FColor DrawColor = FLinearColor::LerpUsingHSV(VisualizationSettings.Color.ReinterpretAsLinear(), FLinearColor::Black, j / (float)AttributeAsVector[i].Num()).ToFColor(true);
 					DrawAsDirection(LineBatchComponent, StartPos, EndPos, DrawColor);
-					DrawAsText(InContext.SceneSetupParams, FText::FromString(FString::Printf(TEXT("%d"), j)), FVector3f(EndPos), 0.f, VisualizationSettings.TextScale, DrawColor);
+
+					if (AttributeAsVector[i].Num() > 1)
+					{
+						DrawAsText(InContext.SceneSetupParams, FText::FromString(FString::Printf(TEXT("%d"), j)), FVector3f(EndPos), 0.f, VisualizationSettings.TextScale, DrawColor);
+					}
 				}
 			}
+			break;
 		}
-		else
+		case EPVDebugValueVisualizationMode::Point:
 		{
-			const TArray<FVector> Vectors = AttributeValue->ToVectors();
-			check(Vectors.Num() == NumElements);
+			using namespace PVE::AttributeToVector3;
+
+			const TArray<FVectorValue> AttributeAsVector = ConvertAttributeToVector3(Collection, AttributeName, GroupName);
+			if (AttributeAsVector.Num() != NumElements)
+			{
+				break;
+			}
+
+			for (int32 i = 0; i < NumElements; i++)
+			{
+				for (int32 j = 0; j < AttributeAsVector[i].Num(); ++j)
+				{
+					AddDrawAsPoint(AttributeAsVector[i][j], VisualizationSettings.GizmoScale);
+				}
+			}
+			break;
+		}
+		case EPVDebugValueVisualizationMode::Sphere:
+		{
+			TArray<double> AttributeAsScalar = PVE::AttributeToScalar::ConvertAttributeToScalar(Collection, AttributeName, GroupName);
+			if (AttributeAsScalar.Num() != NumElements)
+			{
+				break;
+			}
+
+			TObjectPtr<UPVLineBatchComponent> LineBatchComponent = GetOrCreateLineComponent(InContext.SceneSetupParams);
 
 			for (int32 i = 0; i < NumElements; i++)
 			{
@@ -123,36 +556,37 @@ void FPVDebugVisualizationBase::DrawAttributes(FVisualizerDrawContext& InContext
 				float OutScale;
 				GetPivot(Collection, i, OutPos, OutScale);
 
-				const FVector StartPos(OutPos);
-				const FVector EndPos = StartPos + Vectors[i].GetSafeNormal() * VisualizationSettings.DirectionStrength;
-
-				DrawAsDirection(LineBatchComponent, StartPos, EndPos, VisualizationSettings.Color);
+				DrawAsSphere(LineBatchComponent, FVector(OutPos), AttributeAsScalar[i] * VisualizationSettings.GizmoScale, VisualizationSettings.Color);
 			}
+			break;
 		}
 	}
-	else if (VisualizationSettings.VisualizationMode == EPVDebugValueVisualizationMode::Point)
+
+	if (PointsToDraw.Num() > 0)
 	{
-		const TArray<FVector> Vectors = AttributeValue->ToVectors();
-		check(Vectors.Num() == NumElements);
+		if (VisualizationSettings.bDrawPointAsMesh)
+		{
+			TArray<FTransform> PointTransforms;
+			PointTransforms.Reserve(PointsToDraw.Num());
+			for (const auto& It : PointsToDraw)
+			{
+				PointTransforms.Add(FTransform(FRotator::ZeroRotator, It.Key, FVector(It.Value * 0.005f)));
+			}
 
-		for (int32 i = 0; i < NumElements; i++)
+			DrawAsMeshPoints(GetOrCreatePointComponent(InContext.SceneSetupParams), PointTransforms);
+		}
+		else
 		{
-			AddDrawAsPoint(Vectors[i], VisualizationSettings.PointScale);
+			DrawAsPoints(GetOrCreateLineComponent(InContext.SceneSetupParams), PointsToDraw, VisualizationSettings.Color);
 		}
-	}	
 
-	if (PointsToDraw.Num() > 0)
-	{
-		DrawAsPoint(GetOrCreatePointComponent(InContext.SceneSetupParams), PointsToDraw);
+		PointsToDraw.Reset();
 	}
 }
 
 void FPVDebugVisualizationBase::AddDrawAsPoint(const FVector& InPos, float PointScale)
 {
-	auto Transform = FTransform(InPos);
-	Transform.SetScale3D(FVector(PointScale, PointScale, PointScale));
-
-	PointsToDraw.Add(Transform);
+	PointsToDraw.Add({ InPos, PointScale });
 }
 
 void FPVDebugVisualizationBase::DrawAsText(FPCGSceneSetupParams& InOutParams, const FText& InTextToDraw, const FVector3f& InPos, float InScale, float TextSize, const FLinearColor& InColor)
@@ -173,11 +607,23 @@ void FPVDebugVisualizationBase::DrawAsText(FPCGSceneSetupParams& InOutParams, co
 void FPVDebugVisualizationBase::DrawAsDirection(UPVLineBatchComponent* LineBatchComponent, const FVector& InStartPos, const FVector& InEndPos,
 	const FLinearColor& InColor)
 {
-	check(LineBatchComponent)
+	check(LineBatchComponent);
 	LineBatchComponent->AddLine(InStartPos, InEndPos, InColor);
 }
 
-void FPVDebugVisualizationBase::DrawAsPoint(UInstancedStaticMeshComponent* InstancedStaticMeshComponent, const TArray<FTransform>& InTransforms, const FLinearColor& InColor)
+void FPVDebugVisualizationBase::DrawAsSphere(UPVLineBatchComponent* LineBatchComponent, const FVector& InPos, float Radius, const FLinearColor& InColor)
+{
+	check(LineBatchComponent);
+	LineBatchComponent->AddSphere(InPos, Radius, InColor);
+}
+
+void FPVDebugVisualizationBase::DrawAsPoints(UPVLineBatchComponent* LineBatchComponent, const TArray<TPair<FVector, float>>& InPoints, const FLinearColor& InColor)
+{
+	check(LineBatchComponent);
+	LineBatchComponent->AddPoints(InPoints, InColor);
+}
+
+void FPVDebugVisualizationBase::DrawAsMeshPoints(UInstancedStaticMeshComponent* InstancedStaticMeshComponent, const TArray<FTransform>& InTransforms, const FLinearColor& InColor)
 {
 	if (InTransforms.Num() == 0)
 	{
diff --git a/Engine/Plugins/Experimental/ProceduralVegetationEditor/Source/ProceduralVegetationEditor/Private/Visualizations/DebugVisualization/PVDebugVisualizationBase.h b/Engine/Plugins/Experimental/ProceduralVegetationEditor/Source/ProceduralVegetationEditor/Private/Visualizations/DebugVisualization/PVDebugVisualizationBase.h
--- a/Engine/Plugins/Experimental/ProceduralVegetationEditor/Source/ProceduralVegetationEditor/Private/Visualizations/DebugVisualization/PVDebugVisualizationBase.h
+++ b/Engine/Plugins/Experimental/ProceduralVegetationEditor/Source/ProceduralVegetationEditor/Private/Visualizations/DebugVisualization/PVDebugVisualizationBase.h
@@ -30,7 +30,9 @@ protected:
 	
 	static void DrawAsText(FPCGSceneSetupParams& InOutParams, const FText& InTextToDraw, const FVector3f& InPos, float InScale, float TextSize, const FLinearColor& InColor = FLinearColor::White);
 	static void DrawAsDirection(UPVLineBatchComponent* LineBatchComponent, const FVector& InStartPos, const FVector& InEndPos, const FLinearColor& InColor = FLinearColor::White);
-	static void DrawAsPoint(UInstancedStaticMeshComponent* InstancedStaticMeshComponent, const TArray<FTransform>& InTransforms, const FLinearColor& InColor = FLinearColor::White);
+	static void DrawAsSphere(UPVLineBatchComponent* LineBatchComponent, const FVector& InPos, float Radius, const FLinearColor& InColor = FLinearColor::White);
+	static void DrawAsPoints(UPVLineBatchComponent* LineBatchComponent, const TArray<TPair<FVector, float>>& InPoints, const FLinearColor& InColor = FLinearColor::White);
+	static void DrawAsMeshPoints(UInstancedStaticMeshComponent* InstancedStaticMeshComponent, const TArray<FTransform>& InTransforms, const FLinearColor& InColor = FLinearColor::White);
 	
 	static UPVLineBatchComponent* GetOrCreateLineComponent(FPCGSceneSetupParams& InOutParams);
 	static UInstancedStaticMeshComponent* GetOrCreatePointComponent(FPCGSceneSetupParams& InOutParams);
@@ -38,7 +40,7 @@ protected:
 	static void AddToScene(FPCGSceneSetupParams& InOutParams, UPrimitiveComponent* InComponent);
 	
 private:
-	TArray<FTransform> PointsToDraw;
+	TArray<TPair<FVector, float>> PointsToDraw;
 };
 typedef TSharedPtr<FPVDebugVisualizationBase> FPVDebugVisualizationPtr;
 
diff --git a/Engine/Plugins/Experimental/ProceduralVegetationEditor/Source/ProceduralVegetationEditor/Private/Visualizations/PVLineBatchComponent.cpp b/Engine/Plugins/Experimental/ProceduralVegetationEditor/Source/ProceduralVegetationEditor/Private/Visualizations/PVLineBatchComponent.cpp
--- a/Engine/Plugins/Experimental/ProceduralVegetationEditor/Source/ProceduralVegetationEditor/Private/Visualizations/PVLineBatchComponent.cpp
+++ b/Engine/Plugins/Experimental/ProceduralVegetationEditor/Source/ProceduralVegetationEditor/Private/Visualizations/PVLineBatchComponent.cpp
@@ -7,6 +7,7 @@
 FPVLineSceneProxy::FPVLineSceneProxy(const UPVLineBatchComponent* InComponent)
 	: FPrimitiveSceneProxy(InComponent)
 	, Lines(InComponent->Lines)
+	, Points(InComponent->Points)
 	, PointSize(InComponent->PointSize)
 {}
 
@@ -52,6 +53,19 @@ void FPVLineSceneProxy::GetDynamicMeshElements(
 		}
 	};
 
+	const static auto DrawPoints = [](FPrimitiveDrawInterface* PDI, const TArray<FPVPointInfo>& Points)
+	{
+		for (const FPVPointInfo& Point : Points)
+		{
+			PDI->DrawPoint(
+				Point.PointLocation,
+				Point.Color,
+				Point.PointSize,
+				Point.DepthPriorityGroup
+			);
+		}
+	};
+
 	for (int32 ViewIndex = 0; ViewIndex < Views.Num(); ViewIndex++)
 	{
 		if (VisibilityMap & (1 << ViewIndex))
@@ -65,6 +79,7 @@ void FPVLineSceneProxy::GetDynamicMeshElements(
 #endif
 
 			DrawLines(PDI, Lines, PointSize);
+			DrawPoints(PDI, Points);
 		}
 	}
 }
@@ -101,7 +116,7 @@ void UPVLineBatchComponent::InitBounds()
 	MarkRenderStateDirty();
 }
 
-void UPVLineBatchComponent::AddLine(const FVector& InStartPos, const FVector& InEndPos, const FLinearColor& InColor,const ESceneDepthPriorityGroup InDepthPriorityGroup, const EPointDrawSettings InPointDrawSettings)
+void UPVLineBatchComponent::AddLine(const FVector& InStartPos, const FVector& InEndPos, const FLinearColor& InColor, const ESceneDepthPriorityGroup InDepthPriorityGroup, const EPointDrawSettings InPointDrawSettings)
 {
 	Lines.Emplace(InStartPos, InEndPos, InColor, InDepthPriorityGroup, InPointDrawSettings);
 
@@ -111,9 +126,78 @@ void UPVLineBatchComponent::AddLine(const FVector& InStartPos, const FVector& In
 	MarkRenderStateDirty();
 }
 
+void UPVLineBatchComponent::AddSphere(const FVector& InPos, float Radius, const FLinearColor& InColor, const ESceneDepthPriorityGroup InDepthPriorityGroup)
+{
+	if (Radius <= 0)
+	{
+		return;
+	}
+
+	const static int32 Segments = 6;
+
+	const float AngleInc = 2.f * UE_PI / Segments;
+	int32 NumSegmentsY = Segments;
+	float Latitude = AngleInc;
+	float SinY1 = 0.0f, CosY1 = 1.0f;
+
+	while (NumSegmentsY--)
+	{
+		const float SinY2 = FMath::Sin(Latitude);
+		const float CosY2 = FMath::Cos(Latitude);
+
+		FVector Vertex1 = FVector(SinY1, 0.0f, CosY1) * Radius + InPos;
+		FVector Vertex3 = FVector(SinY2, 0.0f, CosY2) * Radius + InPos;
+		float Longitude = AngleInc;
+
+		int32 NumSegmentsX = Segments;
+		while (NumSegmentsX--)
+		{
+			const float SinX = FMath::Sin(Longitude);
+			const float CosX = FMath::Cos(Longitude);
+
+			const FVector Vertex2 = FVector((CosX * SinY1), (SinX * SinY1), CosY1) * Radius + InPos;
+			const FVector Vertex4 = FVector((CosX * SinY2), (SinX * SinY2), CosY2) * Radius + InPos;
+
+			Lines.Emplace(Vertex1, Vertex2, InColor, InDepthPriorityGroup, EPointDrawSettings::None);
+			Lines.Emplace(Vertex1, Vertex3, InColor, InDepthPriorityGroup, EPointDrawSettings::None);
+
+			BBox += Vertex1;
+			BBox += Vertex2;
+			BBox += Vertex3;
+
+			Vertex1 = Vertex2;
+			Vertex3 = Vertex4;
+			Longitude += AngleInc;
+		}
+		SinY1 = SinY2;
+		CosY1 = CosY2;
+		Latitude += AngleInc;
+	}
+
+	MarkRenderStateDirty();
+}
+
+void UPVLineBatchComponent::AddPoint(const FVector& PointLocation, float InPointSize, const FLinearColor& InColor, const ESceneDepthPriorityGroup InDepthPriorityGroup)
+{
+	Points.Emplace(PointLocation, InPointSize, InColor, InDepthPriorityGroup);
+	BBox += PointLocation;
+	MarkRenderStateDirty();
+}
+
+void UPVLineBatchComponent::AddPoints(const TArray<TPair<FVector, float>>& InPoints, const FLinearColor& InColor, const ESceneDepthPriorityGroup InDepthPriorityGroup)
+{
+	for (const auto& It: InPoints)
+	{
+		Points.Emplace(It.Key, It.Value, InColor, InDepthPriorityGroup);
+		BBox += It.Key;
+	}
+	MarkRenderStateDirty();
+}
+
 void UPVLineBatchComponent::Flush()
 {
 	Lines.Empty();
+	Points.Empty();
 	BBox = FBox(ForceInit);
 	MarkRenderStateDirty();
 }
diff --git a/Engine/Plugins/Experimental/ProceduralVegetationEditor/Source/ProceduralVegetationEditor/Private/Visualizations/PVLineBatchComponent.h b/Engine/Plugins/Experimental/ProceduralVegetationEditor/Source/ProceduralVegetationEditor/Private/Visualizations/PVLineBatchComponent.h
--- a/Engine/Plugins/Experimental/ProceduralVegetationEditor/Source/ProceduralVegetationEditor/Private/Visualizations/PVLineBatchComponent.h
+++ b/Engine/Plugins/Experimental/ProceduralVegetationEditor/Source/ProceduralVegetationEditor/Private/Visualizations/PVLineBatchComponent.h
@@ -29,10 +29,19 @@ struct FPVLineInfo
 	EPointDrawSettings PointDrawSettings = EPointDrawSettings::None;
 };
 
+struct FPVPointInfo
+{
+	FVector PointLocation;
+	float PointSize;
+	FLinearColor Color;
+	ESceneDepthPriorityGroup DepthPriorityGroup = SDPG_World;
+};
+
 class FPVLineSceneProxy final : public FPrimitiveSceneProxy
 {
 private:
 	TArray<FPVLineInfo> Lines;
+	TArray<FPVPointInfo> Points;
 	float PointSize;
 
 public:
@@ -67,6 +76,15 @@ public:
 				const ESceneDepthPriorityGroup InDepthPriorityGroup = ESceneDepthPriorityGroup::SDPG_World,
 				const EPointDrawSettings InPointDrawSettings = EPointDrawSettings::None);
 
+	void AddSphere(const FVector& InPos, float Radius, const FLinearColor& InColor = FLinearColor::White,
+		const ESceneDepthPriorityGroup InDepthPriorityGroup = ESceneDepthPriorityGroup::SDPG_World);
+
+	void AddPoint(const FVector& PointLocation, float PointSize, const FLinearColor& InColor = FLinearColor::White,
+		const ESceneDepthPriorityGroup InDepthPriorityGroup = ESceneDepthPriorityGroup::SDPG_World);
+
+	void AddPoints(const TArray<TPair<FVector, float>>& InPoints, const FLinearColor& InColor = FLinearColor::White,
+		const ESceneDepthPriorityGroup InDepthPriorityGroup = ESceneDepthPriorityGroup::SDPG_World);
+
 	void Flush();
 	
 	virtual FPrimitiveSceneProxy* CreateSceneProxy() override;
@@ -75,7 +93,8 @@ public:
 protected:
 	FBox BBox;
 	TArray<FPVLineInfo> Lines;
-	
+	TArray<FPVPointInfo> Points;
+
 private:
 	const float PointSize = 5.0f;
 };
