[PVE] Gracefully replace older assets in case of replace export policy

#Synced-CL 48384499

diff --git a/Engine/Plugins/Experimental/ProceduralVegetationEditor/Source/ProceduralVegetationEditor/Private/Helpers/PVExportHelper.cpp b/Engine/Plugins/Experimental/ProceduralVegetationEditor/Source/ProceduralVegetationEditor/Private/Helpers/PVExportHelper.cpp
--- a/Engine/Plugins/Experimental/ProceduralVegetationEditor/Source/ProceduralVegetationEditor/Private/Helpers/PVExportHelper.cpp
+++ b/Engine/Plugins/Experimental/ProceduralVegetationEditor/Source/ProceduralVegetationEditor/Private/Helpers/PVExportHelper.cpp
@@ -254,6 +254,29 @@ namespace PV::Export::Internal
 		return Package;
 	}
 
+	void ReplaceAssetInPackage(const TObjectPtr<UObject>& NewObject)
+	{
+		UPackage* Package = NewObject->GetPackage();
+		
+		const FString ObjectName = FPackageName::GetShortName(Package->GetName());
+		TArray<UObject*> PackageObjects;
+		GetObjectsWithPackage(Package, PackageObjects, false);
+		PackageObjects = PackageObjects.FilterByPredicate(
+			[&](UObject* Object)
+			{
+				return Object != NewObject;
+			}
+		);
+		for (UObject* ExistingObject : PackageObjects)
+		{
+			const FString NewName = MakeUniqueObjectName(GetTransientPackage(), ExistingObject->GetClass(), *(ExistingObject->GetName() + "_DEADOBJECT")).ToString();
+			ExistingObject->Rename(*NewName, nullptr, REN_DontCreateRedirectors | REN_DoNotDirty | REN_NonTransactional);
+			ExistingObject->MarkAsGarbage();
+			IAssetRegistry::Get()->AssetDeleted(ExistingObject);
+		}
+		NewObject->Rename(*ObjectName, Package, REN_DontCreateRedirectors | REN_DoNotDirty | REN_NonTransactional);
+	}
+
 	TObjectPtr<UDynamicMesh> CollectionToDynamicMesh(const FManagedArrayCollection& Collection)
 	{
 		TRACE_CPUPROFILER_EVENT_SCOPE(PV::Export::CollectionToDynamicMesh);
@@ -1657,7 +1680,7 @@ namespace PV::Export
 						return PV::Analytics::EExportResult::Success;
 					case EExportResult::Fail:
 						return PV::Analytics::EExportResult::Failed;
-					case EExportResult::Canceled:
+					case EExportResult::Cancelled:
 						return PV::Analytics::EExportResult::Failed;
 					case EExportResult::Skipped:
 						return PV::Analytics::EExportResult::Skipped;
@@ -1686,10 +1709,11 @@ namespace PV::Export
 			return ExitWithResult(EExportResult::Fail);
 		};
 
+		Internal::EGetValidAssetPathResult AssetPathResult;
+		
 		FString MeshName = ExportParams.MeshName.ToString();
 		FString MeshPackageName = ExportParams.GetOutputMeshPackagePath();
-
-		Internal::EGetValidAssetPathResult AssetPathResult = Internal::GetValidAssetPath(ExportParams.ReplacePolicy, ExportParams.GetMeshClass(), MeshPackageName, MeshName);
+		AssetPathResult = Internal::GetValidAssetPath(ExportParams.ReplacePolicy, ExportParams.GetMeshClass(), MeshPackageName, MeshName);
 		if (AssetPathResult != Internal::EGetValidAssetPathResult::Success)
 		{
 			return ExitWithAssetPathResult(AssetPathResult);
@@ -1700,16 +1724,7 @@ namespace PV::Export
 
 		if (ExportParams.ExportMeshType == EPVExportMeshType::StaticMesh)
 		{
-			TObjectPtr<UPackage> MeshPackage = Internal::CreateAssetPackage(MeshPackageName);
-			if (TObjectPtr<UStaticMesh> ExistingMeshAsset = Cast<UStaticMesh>(MeshPackage->FindAssetInPackage()))
-			{
-				ExistingMeshAsset->SetExternalPackage(nullptr);
-				ExistingMeshAsset->MarkAsGarbage();
-			}
-
-			const TObjectPtr<UStaticMesh> ExportMesh = NewObject<UStaticMesh>(MeshPackage, *MeshName, RF_Standalone | RF_Public);
-
-			ModifiedPackages.Add(ExportMesh->GetOutermost());
+			const TObjectPtr<UStaticMesh> ExportMesh = NewObject<UStaticMesh>(Internal::CreateAssetPackage(MeshPackageName), NAME_None, RF_Public | RF_Standalone);
 
 			if (!PV::Export::Internal::ExportCollectionToStaticMesh(
 				ExportMesh,
@@ -1722,9 +1737,13 @@ namespace PV::Export
 				StatusReportCallback
 			))
 			{
-				return ExitWithResult(EExportResult::Canceled);
+				ExportMesh->MarkAsGarbage();
+				return ExitWithResult(EExportResult::Cancelled);
 			}
 
+			Internal::ReplaceAssetInPackage(ExportMesh);
+			ModifiedPackages.Add(ExportMesh->GetOutermost());
+
 			OutCreatedAssets.Add(MeshPackageName);
 
 			Internal::AttachProceduralVegetationLink(ExportMesh, InProceduralVegetation);
@@ -1733,73 +1752,33 @@ namespace PV::Export
 		{
 			FString SkeletonName = ExportParams.GetOutputSkeletonName();
 			FString SkeletonPackageName = ExportParams.GetOutputSkeletonPackagePath();
-
 			AssetPathResult = Internal::GetValidAssetPath(ExportParams.ReplacePolicy, USkeleton::StaticClass(), SkeletonPackageName, SkeletonName);
 			if (AssetPathResult != Internal::EGetValidAssetPathResult::Success)
 			{
 				return ExitWithAssetPathResult(AssetPathResult);
 			}
 
-			TObjectPtr<UPackage> MeshPackage = Internal::CreateAssetPackage(MeshPackageName);
-			TObjectPtr<UPackage> SkeletonPackage = Internal::CreateAssetPackage(SkeletonPackageName);
-
-			TObjectPtr<USkeletalMesh> ExistingMeshAsset = Cast<USkeletalMesh>(MeshPackage->FindAssetInPackage());
-
-			TObjectPtr<USkeletalMesh> ExportMesh = nullptr;
-			
-			if (ExistingMeshAsset)
-			{
-				ExportMesh = ExistingMeshAsset;
-				ExistingMeshAsset->ClearMeshDescriptionAndBulkData(0);
-				FNaniteAssemblyData& NaniteAssemblyData = ExistingMeshAsset->NaniteSettings.NaniteAssemblyData;
-				if (NaniteAssemblyData.IsValid())
-				{
-					NaniteAssemblyData.Nodes.Empty();
-					NaniteAssemblyData.Parts.Empty();
-				}
-			}
-			else
-			{
-				ExportMesh = NewObject<USkeletalMesh>(MeshPackage, *MeshName, RF_Standalone | RF_Public);
-			}
-			
-			if (TObjectPtr<USkeleton> ExistingSkeletonAsset = Cast<USkeleton>(SkeletonPackage->FindAssetInPackage()))
-			{
-				ExistingSkeletonAsset->SetExternalPackage(nullptr);
-				ExistingSkeletonAsset->MarkAsGarbage();
-			}
+			const TObjectPtr<USkeletalMesh> ExportMesh = NewObject<USkeletalMesh>(Internal::CreateAssetPackage(MeshPackageName), NAME_None, RF_Standalone | RF_Public);
+			const TObjectPtr<USkeleton> ExportSkeleton = NewObject<USkeleton>(Internal::CreateAssetPackage(SkeletonPackageName), NAME_None, RF_Standalone | RF_Public);
 
-			const TObjectPtr<USkeleton> ExportSkeleton = NewObject<USkeleton>(SkeletonPackage, *SkeletonName, RF_Standalone | RF_Public);
-			
 			ExportMesh->SetSkeleton(ExportSkeleton);
 
-			TObjectPtr<UPhysicsAsset> PhysicsAsset;
 			FString PhysicsAssetName = ExportParams.GetOutputPhysicsAssetName();
 			FString PhysicsPackageName = ExportParams.GetOutputPhysicsAssetPackagePath();
+			AssetPathResult = Internal::GetValidAssetPath(ExportParams.ReplacePolicy, UPhysicsAsset::StaticClass(), PhysicsPackageName, PhysicsAssetName);
+			if (AssetPathResult != Internal::EGetValidAssetPathResult::Success)
+			{
+				return ExitWithAssetPathResult(AssetPathResult);
+			}
 
+			TObjectPtr<UPhysicsAsset> PhysicsAsset = nullptr;
 			if (ExportParams.IsCollisionEnable())
 			{
-				Internal::EGetValidAssetPathResult PhysicsAssetPathResult = Internal::GetValidAssetPath(ExportParams.ReplacePolicy, UPhysicsAsset::StaticClass(), PhysicsPackageName, PhysicsAssetName);
-				if (PhysicsAssetPathResult != Internal::EGetValidAssetPathResult::Success)
-				{
-					return ExitWithAssetPathResult(PhysicsAssetPathResult);
-				}
-				
-				TObjectPtr<UPackage> PhysicsPackage = Internal::CreateAssetPackage(PhysicsPackageName);
-				
-				if (TObjectPtr<UPhysicsAsset> ExistingPhysicAsset = Cast<UPhysicsAsset>(PhysicsPackage->FindAssetInPackage()))
-				{
-					ExistingPhysicAsset->SetExternalPackage(nullptr);
-					ExistingPhysicAsset->MarkAsGarbage();
-				}
-				
-				PhysicsAsset = NewObject<UPhysicsAsset>(PhysicsPackage, *PhysicsAssetName, RF_Standalone | RF_Public);
-				ModifiedPackages.Add(PhysicsAsset->GetOutermost());
+				PhysicsAsset = NewObject<UPhysicsAsset>(Internal::CreateAssetPackage(PhysicsPackageName), NAME_None, RF_Standalone | RF_Public);
+				ExportMesh->SetPhysicsAsset(PhysicsAsset);
+				PhysicsAsset->SetPreviewMesh(ExportMesh, false);
 			}
 
-			ModifiedPackages.Add(ExportMesh->GetOutermost());
-			ModifiedPackages.Add(ExportSkeleton->GetOutermost());
-
 			if (!PV::Export::Internal::ExportCollectionToSkeletalMesh(
 				ExportMesh,
 				Collection,
@@ -1813,18 +1792,36 @@ namespace PV::Export
 				StatusReportCallback
 			))
 			{
-				return ExitWithResult(EExportResult::Canceled);
+				ExportMesh->SetSkeleton(nullptr);
+				if (PhysicsAsset)
+				{
+					ExportMesh->SetPhysicsAsset(nullptr);
+					PhysicsAsset->MarkAsGarbage();
+				}
+				
+				ExportSkeleton->MarkAsGarbage();
+				ExportMesh->MarkAsGarbage();
+				
+				return ExitWithResult(EExportResult::Cancelled);
 			}
+			
+			if (PhysicsAsset)
+			{
+				Internal::ReplaceAssetInPackage(PhysicsAsset);
+				ModifiedPackages.Add(PhysicsAsset->GetOutermost());
 
+				OutCreatedAssets.Add(PhysicsPackageName);
+				Internal::AttachProceduralVegetationLink(PhysicsAsset, InProceduralVegetation);
+			}
+			
+			Internal::ReplaceAssetInPackage(ExportSkeleton);
+			Internal::ReplaceAssetInPackage(ExportMesh);
+			ModifiedPackages.Add(ExportMesh->GetOutermost());
+			ModifiedPackages.Add(ExportSkeleton->GetOutermost());
+			
 			OutCreatedAssets.Add(MeshPackageName);
 			OutCreatedAssets.Add(SkeletonPackageName);
 
-			if(PhysicsAsset)
-			{
-				PhysicsAsset->PreviewSkeletalMesh = ExportMesh;
-				ExportMesh->SetPhysicsAsset(PhysicsAsset);
-				Internal::AttachProceduralVegetationLink(PhysicsAsset, InProceduralVegetation);
-			}
 			Internal::AttachProceduralVegetationLink(ExportMesh, InProceduralVegetation);
 			Internal::AttachProceduralVegetationLink(ExportSkeleton, InProceduralVegetation);
 		}
diff --git a/Engine/Plugins/Experimental/ProceduralVegetationEditor/Source/ProceduralVegetationEditor/Private/Helpers/PVExportHelper.h b/Engine/Plugins/Experimental/ProceduralVegetationEditor/Source/ProceduralVegetationEditor/Private/Helpers/PVExportHelper.h
--- a/Engine/Plugins/Experimental/ProceduralVegetationEditor/Source/ProceduralVegetationEditor/Private/Helpers/PVExportHelper.h
+++ b/Engine/Plugins/Experimental/ProceduralVegetationEditor/Source/ProceduralVegetationEditor/Private/Helpers/PVExportHelper.h
@@ -29,7 +29,7 @@ namespace PV::Export
 	{
 		Success,
 		Fail,
-		Canceled,
+		Cancelled,
 		Skipped
 	};
 
diff --git a/Engine/Plugins/Experimental/ProceduralVegetationEditor/Source/ProceduralVegetationEditor/Private/PVEditor.cpp b/Engine/Plugins/Experimental/ProceduralVegetationEditor/Source/ProceduralVegetationEditor/Private/PVEditor.cpp
--- a/Engine/Plugins/Experimental/ProceduralVegetationEditor/Source/ProceduralVegetationEditor/Private/PVEditor.cpp
+++ b/Engine/Plugins/Experimental/ProceduralVegetationEditor/Source/ProceduralVegetationEditor/Private/PVEditor.cpp
@@ -729,7 +729,7 @@ FReply FPVEditor::ExportInternal()
 			
 			break;
 		}
-		else if (ExportResult == PV::Export::EExportResult::Canceled)
+		else if (ExportResult == PV::Export::EExportResult::Cancelled)
 		{
 			FString MeshNames = "{ ";
 			for (int32 j = i; j < NodesToExport.Num(); j++)
