[PVE] Debug Loops

#virtualized
#Synced-CL 47309591

diff --git a/Engine/Plugins/Experimental/ProceduralVegetationEditor/Source/ProceduralVegetation/Private/Helpers/PVUtilities.cpp b/Engine/Plugins/Experimental/ProceduralVegetationEditor/Source/ProceduralVegetation/Private/Helpers/PVUtilities.cpp
--- a/Engine/Plugins/Experimental/ProceduralVegetationEditor/Source/ProceduralVegetation/Private/Helpers/PVUtilities.cpp
+++ b/Engine/Plugins/Experimental/ProceduralVegetationEditor/Source/ProceduralVegetation/Private/Helpers/PVUtilities.cpp
@@ -273,6 +273,19 @@ namespace PV::Utilities
 		return NumTriangles;
 	}
 
+	FLinearColor GetRandomHueColor(float Alpha)
+	{
+		Alpha = FMath::Clamp(Alpha, 0.f, 1.f);
+		const float HueDeg = 240.f * Alpha;
+		const FRandomStream Rand(Alpha * 255);
+		const float V = FMath::GetMappedRangeValueUnclamped(
+			FVector2f{0.0f, 1.0f},
+			FVector2f{0.25f, 1.0f},
+			Rand.FRand()
+		);
+		return FLinearColor(HueDeg, 0.9, V).HSVToLinearRGB();
+	};
+
 	bool IsValidPVData(const FManagedArrayCollection& Collection)
 	{
 		Facades::FPointFacade PointFacade = Facades::FPointFacade(Collection);
diff --git a/Engine/Plugins/Experimental/ProceduralVegetationEditor/Source/ProceduralVegetation/Private/Implementations/PVGravity.cpp b/Engine/Plugins/Experimental/ProceduralVegetationEditor/Source/ProceduralVegetation/Private/Implementations/PVGravity.cpp
--- a/Engine/Plugins/Experimental/ProceduralVegetationEditor/Source/ProceduralVegetation/Private/Implementations/PVGravity.cpp
+++ b/Engine/Plugins/Experimental/ProceduralVegetationEditor/Source/ProceduralVegetation/Private/Implementations/PVGravity.cpp
@@ -87,6 +87,8 @@ void FPVGravity::ApplyGravity(
 	FVector3f PivotPosition
 )
 {
+	PVE_LOOP_DEBUG_CHECK(return);
+	
 	const auto CalculateDownForce = [](const FVector3f InRelativeDirection, const int InPointIndex, const float InPScaleGradient,
 	                                          const FPVGravityParams& InGravityParams,
 	                                          const TArray<FVector3f>& InPhototropicDirections) -> FQuat4f
@@ -187,6 +189,14 @@ void FPVGravity::ApplyGravity(
 			PhototropicDirections);
 		TotalDownForce = DownForce * TotalDownForce;
 
+		PVE_LOOP_DEBUG_PARAMS_START()
+			PVE_LOOP_DEBUG_DIRECTION_PARAM("Tangent", Position, Tangent);
+			PVE_LOOP_DEBUG_TEXT_PARAM("Length from Root", Position, FText::FromString(FString::Printf(TEXT("%.2f"), LengthFromRoot)));
+			PVE_LOOP_DEBUG_VECTOR_PARAM("Offset", Position, PivotPosition + DownForce * RelativeDirection);
+		PVE_LOOP_DEBUG_PARAMS_END()
+
+		PVE_LOOP_DEBUG_STEP(return);
+
 		// Update attachment points for all foliage elements for current branch between the current and previous point 
 		TArray<int32> FoliageEntryIds = FoliageFacade.GetFoliageEntryIdsForBranch(BranchIndex);
 		for (const int32 FoliageEntryId : FoliageEntryIds)
diff --git a/Engine/Plugins/Experimental/ProceduralVegetationEditor/Source/ProceduralVegetation/Private/Implementations/PVGravity.h b/Engine/Plugins/Experimental/ProceduralVegetationEditor/Source/ProceduralVegetation/Private/Implementations/PVGravity.h
--- a/Engine/Plugins/Experimental/ProceduralVegetationEditor/Source/ProceduralVegetation/Private/Implementations/PVGravity.h
+++ b/Engine/Plugins/Experimental/ProceduralVegetationEditor/Source/ProceduralVegetation/Private/Implementations/PVGravity.h
@@ -4,6 +4,9 @@
 
 #include "CoreMinimal.h"
 #include "GeometryCollection/ManagedArrayCollection.h"
+
+#include "Helpers/PVUtilities.h"
+
 #include "PVGravity.generated.h"
 
 namespace PV::Facades
diff --git a/Engine/Plugins/Experimental/ProceduralVegetationEditor/Source/ProceduralVegetation/Private/Nodes/PVBaseSettings.cpp b/Engine/Plugins/Experimental/ProceduralVegetationEditor/Source/ProceduralVegetation/Private/Nodes/PVBaseSettings.cpp
--- a/Engine/Plugins/Experimental/ProceduralVegetationEditor/Source/ProceduralVegetation/Private/Nodes/PVBaseSettings.cpp
+++ b/Engine/Plugins/Experimental/ProceduralVegetationEditor/Source/ProceduralVegetation/Private/Nodes/PVBaseSettings.cpp
@@ -85,7 +85,7 @@ void FPVBaseElement::PostExecuteInternal(FPCGContext* InContext) const
 {
 	IPCGElement::PostExecuteInternal(InContext);
 
-#if WITH_EDITOR
+#if WITH_EDITORONLY_DATA
 	const UPVBaseSettings* Settings = InContext->GetInputSettings<UPVBaseSettings>();
 	check(Settings);
 
@@ -95,7 +95,9 @@ void FPVBaseElement::PostExecuteInternal(FPCGContext* InContext) const
 		{
 			if (const UPVData* ProceduralVegetationData = Cast<UPVData>(OutputData.Data))
 			{
-				ProceduralVegetationData->SetDebugSettings(Settings->GetDebugVisualizationSettings());
+				FPVDebugSettings DebugSettings = ProceduralVegetationData->GetDebugSettings();
+				DebugSettings.VisualizationSettings = Settings->GetDebugVisualizationSettings().VisualizationSettings;
+				ProceduralVegetationData->SetDebugSettings(MoveTemp(DebugSettings));
 			}
 		}
 	}
diff --git a/Engine/Plugins/Experimental/ProceduralVegetationEditor/Source/ProceduralVegetation/Private/Nodes/PVGravitySettings.cpp b/Engine/Plugins/Experimental/ProceduralVegetationEditor/Source/ProceduralVegetation/Private/Nodes/PVGravitySettings.cpp
--- a/Engine/Plugins/Experimental/ProceduralVegetationEditor/Source/ProceduralVegetation/Private/Nodes/PVGravitySettings.cpp
+++ b/Engine/Plugins/Experimental/ProceduralVegetationEditor/Source/ProceduralVegetation/Private/Nodes/PVGravitySettings.cpp
@@ -61,6 +61,10 @@ bool FPVGravityElement::ExecuteInternal(FPCGContext* InContext) const
 			FManagedArrayCollection OutCollection = InputData->GetCollection();
 		
 			UPVGrowthData* OutManagedArrayCollectionData = FPCGContext::NewObject_AnyThread<UPVGrowthData>(InContext);
+
+#if WITH_EDITORONLY_DATA
+			PVE_LOOP_DEBUG_INIT(Settings, LoopDebugStepper, OutManagedArrayCollectionData);
+#endif
 			
 			if (Settings->GravitySettings.Gravity != 0.0f)
 			{
@@ -69,6 +73,8 @@ bool FPVGravityElement::ExecuteInternal(FPCGContext* InContext) const
 
 			OutManagedArrayCollectionData->Initialize(MoveTemp(OutCollection));
 			InContext->OutputData.TaggedData.Emplace(OutManagedArrayCollectionData);
+
+			PVE_LOOP_DEBUG_END();
 		}
 		else
 		{
diff --git a/Engine/Plugins/Experimental/ProceduralVegetationEditor/Source/ProceduralVegetation/Private/Nodes/PVGravitySettings.h b/Engine/Plugins/Experimental/ProceduralVegetationEditor/Source/ProceduralVegetation/Private/Nodes/PVGravitySettings.h
--- a/Engine/Plugins/Experimental/ProceduralVegetationEditor/Source/ProceduralVegetation/Private/Nodes/PVGravitySettings.h
+++ b/Engine/Plugins/Experimental/ProceduralVegetationEditor/Source/ProceduralVegetation/Private/Nodes/PVGravitySettings.h
@@ -30,6 +30,11 @@ protected:
 public:
 	UPROPERTY(EditAnywhere, Category="Gravity", meta=(ShowOnlyInnerProperties))
 	FPVGravityParams GravitySettings;
+
+#if WITH_EDITORONLY_DATA
+	UPROPERTY(EditAnywhere, Category = "Loop Debug", meta=(EditCondition = "bDebug", EditConditionHides))
+	FLoopDebugStepper LoopDebugStepper;
+#endif
 };
 
 class FPVGravityElement : public FPVBaseElement
diff --git a/Engine/Plugins/Experimental/ProceduralVegetationEditor/Source/ProceduralVegetation/Private/Nodes/PVObjectAvoidance.cpp b/Engine/Plugins/Experimental/ProceduralVegetationEditor/Source/ProceduralVegetation/Private/Nodes/PVObjectAvoidance.cpp
--- a/Engine/Plugins/Experimental/ProceduralVegetationEditor/Source/ProceduralVegetation/Private/Nodes/PVObjectAvoidance.cpp
+++ b/Engine/Plugins/Experimental/ProceduralVegetationEditor/Source/ProceduralVegetation/Private/Nodes/PVObjectAvoidance.cpp
@@ -16,6 +16,8 @@
 #include "Facades/PVBudVectorsFacade.h"
 #include "Facades/PVPointFacade.h"
 
+#include "Helpers/PVUtilities.h"
+
 FPCGDataTypeIdentifier UPVObjectAvoidanceSettings::GetInputPinTypeIdentifier() const
 {
 	return FPCGDataTypeIdentifier{FPVDataTypeInfoGrowth::AsId()};
@@ -66,6 +68,12 @@ bool FPVObjectAvoidanceElement::ExecuteInternal(FPCGContext* InContext) const
 		if (const UPVGrowthData* InputData = Cast<UPVGrowthData>(Inputs[0].Data))
 		{
 			FManagedArrayCollection SkeletonCollection = InputData->GetCollection();
+			UPVData* OutManagedArrayCollectionData = FPCGContext::NewObject_AnyThread<UPVGrowthData>(InContext);
+
+#if WITH_EDITORONLY_DATA
+			PVE_LOOP_DEBUG_INIT(InputSettings, LoopDebugStepper, OutManagedArrayCollectionData);
+#endif
+
 			for (const auto& [MeshPath, Transform] : InputSettings->ObjectAvoidanceSettings.Colliders)
 			{
 				const UStaticMesh* ColliderMesh = MeshPath.LoadSynchronous();
@@ -124,6 +132,8 @@ bool FPVObjectAvoidanceElement::ExecuteInternal(FPCGContext* InContext) const
 
 				for (int32 BranchIndex : BranchIndices)
 				{
+					PVE_LOOP_DEBUG_CHECK(break);
+					
 					const TArray<int32> BranchPoints = BranchFacade.GetPoints(BranchIndex);
 					for (int32 BranchPointIndex = 1; BranchPointIndex < BranchPoints.Num(); ++BranchPointIndex)
 					{
@@ -170,6 +180,17 @@ bool FPVObjectAvoidanceElement::ExecuteInternal(FPCGContext* InContext) const
 							const FVector OffsetPosition = PointPosition + AvoidanceDirection * FVector::Distance(PointPosition, NextPointPosition);
 							const FVector3f Offset = FVector3f(OffsetPosition - NextPointPosition);
 
+							PVE_LOOP_DEBUG_PARAMS_START()
+								PVE_LOOP_DEBUG_DIRECTION_PARAM("Ray", Ray.Origin, MaxLengthVector);
+								PVE_LOOP_DEBUG_DIRECTION_PARAM("Hit Normal", HitPoint, HitNormal);
+								PVE_LOOP_DEBUG_POINT_PARAM("Hit Point", HitPoint);
+								PVE_LOOP_DEBUG_DIRECTION_PARAM("Avoidance Direction", HitPoint, AvoidanceDirection);
+								PVE_LOOP_DEBUG_VECTOR_PARAM("Offset Vector", NextPointPosition, OffsetPosition);
+								PVE_LOOP_DEBUG_TEXT_PARAM("Hit Distance", (PointPosition + HitPoint) / 2.0f, FText::FromString(FString::Printf(TEXT("%.2f units"), DistanceToHit)));
+							PVE_LOOP_DEBUG_PARAMS_END()
+
+							PVE_LOOP_DEBUG_STEP(break);
+
 							const int32 StartIndex = FMath::Max(1, BranchPointIndex - InputSettings->ObjectAvoidanceSettings.Smoothness);
 							FVector3f PreviousSiblingPosition;
 							for (int32 SiblingPointIndex = StartIndex; SiblingPointIndex < BranchPointIndex; ++SiblingPointIndex)
@@ -277,9 +298,10 @@ bool FPVObjectAvoidanceElement::ExecuteInternal(FPCGContext* InContext) const
 #endif
 			}
 
-			UPVGrowthData* OutManagedArrayCollectionData = FPCGContext::NewObject_AnyThread<UPVGrowthData>(InContext);
 			OutManagedArrayCollectionData->Initialize(MoveTemp(SkeletonCollection));
 			InContext->OutputData.TaggedData.Emplace(OutManagedArrayCollectionData);
+
+			PVE_LOOP_DEBUG_END();
 		}
 	}
 
diff --git a/Engine/Plugins/Experimental/ProceduralVegetationEditor/Source/ProceduralVegetation/Public/DataTypes/PVData.h b/Engine/Plugins/Experimental/ProceduralVegetationEditor/Source/ProceduralVegetation/Public/DataTypes/PVData.h
--- a/Engine/Plugins/Experimental/ProceduralVegetationEditor/Source/ProceduralVegetation/Public/DataTypes/PVData.h
+++ b/Engine/Plugins/Experimental/ProceduralVegetationEditor/Source/ProceduralVegetation/Public/DataTypes/PVData.h
@@ -3,9 +3,11 @@
 #pragma once
 
 #include "PCGContext.h"
+
 #include "Data/PCGSpatialData.h"
+
 #include "GeometryCollection/ManagedArrayCollection.h"
-#include "Helpers/PVUtilities.h"
+
 #include "PVData.generated.h"
 
 struct FPCGContext;
@@ -41,8 +43,9 @@ enum class EPVDebugValueVisualizationMode
 {
 	Text,
 	Direction,
-	Point,
-	Sphere
+	Sphere,
+	Vector,
+	Point
 };
 
 USTRUCT()
@@ -106,6 +109,36 @@ struct FPVVisualizationSettings
 	}
 };
 
+USTRUCT()
+struct FLoopDebugData
+{
+	GENERATED_BODY()
+	
+	UPROPERTY(EditAnywhere, Category = DebugData)
+	FVector Pivot;
+
+	UPROPERTY(EditAnywhere, Category = DebugData)
+	TOptional<FVector> Direction;
+
+	UPROPERTY(EditAnywhere, Category = DebugData)
+	TOptional<FText> TextData;
+};
+
+USTRUCT()
+struct FPVLoopDebugSettings
+{
+	GENERATED_BODY()
+
+	UPROPERTY(EditAnywhere, Category = LoopDebugSettings)
+	EPVDebugValueVisualizationMode VisualizationMode = EPVDebugValueVisualizationMode::Text;
+
+	UPROPERTY(EditAnywhere, Category = LoopDebugSettings)
+	FText ParamName;
+
+	UPROPERTY(EditAnywhere, Category = LoopDebugSettings)
+	FLoopDebugData Data;
+};
+
 USTRUCT()
 struct FPVDebugSettings
 {
@@ -113,6 +146,11 @@ struct FPVDebugSettings
 
 	UPROPERTY(EditAnywhere, Category = DebugSettings)
 	TArray<FPVVisualizationSettings> VisualizationSettings;
+
+	UPROPERTY()
+	TArray<FPVLoopDebugSettings> LoopDebugVisualizationSettings;
+
+	bool bAutoFocusLoopDebug = false;
 };
 
 UCLASS(BlueprintType, ClassGroup = (Procedural))
@@ -139,8 +177,9 @@ public:
 	const FManagedArrayCollection& GetCollection() const { return Collection; }
 	FManagedArrayCollection& GetCollection() { return Collection; }
 	
-#if WITH_EDITOR
+#if WITH_EDITORONLY_DATA
 	void SetDebugSettings(const FPVDebugSettings& InSettings) const { DebugSettings = InSettings; }
+	void SetDebugSettings(FPVDebugSettings&& InSettings) const { DebugSettings = MoveTemp(InSettings); }
 	const FPVDebugSettings& GetDebugSettings() const { return DebugSettings; }
 #endif
 	
diff --git a/Engine/Plugins/Experimental/ProceduralVegetationEditor/Source/ProceduralVegetation/Public/Helpers/PVUtilities.h b/Engine/Plugins/Experimental/ProceduralVegetationEditor/Source/ProceduralVegetation/Public/Helpers/PVUtilities.h
--- a/Engine/Plugins/Experimental/ProceduralVegetationEditor/Source/ProceduralVegetation/Public/Helpers/PVUtilities.h
+++ b/Engine/Plugins/Experimental/ProceduralVegetationEditor/Source/ProceduralVegetation/Public/Helpers/PVUtilities.h
@@ -1,8 +1,12 @@
 ﻿// Copyright Epic Games, Inc. All Rights Reserved.
 #pragma once
 
+#include "DataTypes/PVData.h"
+
 #include "Facades/PVRenderingFacade.h"
 
+#include "PVUtilities.generated.h"
+
 class UInstancedStaticMeshComponent;
 class UMeshComponent;
 class FFoliageFacade;
@@ -15,9 +19,12 @@ namespace PV::Utilities
 	bool PROCEDURALVEGETATION_API DebugModeEnabled();
 
 	DECLARE_DELEGATE_OneParam(FFoliageComponentCreatedCallback, UMeshComponent*);
-	
-	void PROCEDURALVEGETATION_API AddFoliageInstancesToActor(const Facades::FFoliageFacade& InFoliageFacade, AActor* InActor, TMap<FString, TObjectPtr<UMeshComponent>>& OutInstancedComponentMap);
-	void PROCEDURALVEGETATION_API AddFoliageInstances(const Facades::FFoliageFacade& InFoliageFacade, const UObject* InParent, FFoliageComponentCreatedCallback InCallback, TMap<FString, TObjectPtr<UMeshComponent>>& OutInstancedComponentMap);
+
+	void PROCEDURALVEGETATION_API AddFoliageInstancesToActor(const Facades::FFoliageFacade& InFoliageFacade, AActor* InActor,
+	                                                         TMap<FString, TObjectPtr<UMeshComponent>>& OutInstancedComponentMap);
+	void PROCEDURALVEGETATION_API AddFoliageInstances(const Facades::FFoliageFacade& InFoliageFacade, const UObject* InParent,
+	                                                  FFoliageComponentCreatedCallback InCallback,
+	                                                  TMap<FString, TObjectPtr<UMeshComponent>>& OutInstancedComponentMap);
 
 	bool PROCEDURALVEGETATION_API ValidateAssetPathAndName(const FString& MeshName, const FString& Path, UClass* InClass, FString& OutError);
 
@@ -25,6 +32,157 @@ namespace PV::Utilities
 	bool PROCEDURALVEGETATION_API DoesConflictingPackageExist(const FString& PackageName, UClass* InClass);
 	bool PackageExists(const FString& LongPackageName, UClass* AssetClass);
 	int32 PROCEDURALVEGETATION_API GetMeshTriangles(const FString InMeshPath);
+	FLinearColor PROCEDURALVEGETATION_API GetRandomHueColor(float Alpha);
 
 	bool IsValidPVData(const FManagedArrayCollection& Collection);
 }
+
+USTRUCT()
+struct FLoopDebugStepper
+{
+	GENERATED_BODY()
+
+	UPROPERTY(EditAnywhere, Category="Loop Debug", meta=(InlineEditConditionToggle))
+	bool bDebug = false;
+
+	UPROPERTY(EditAnywhere, Category="Loop Debug")
+	mutable bool bDebugFocus = false;
+
+	UPROPERTY(EditAnywhere, Category="Loop Debug", meta=(ClampMin=1, UIMin=1, EditCondition="bDebug"))
+	uint32 MaxSteps = 1;
+
+	void IncrementStep()
+	{
+		if (MaxSteps != UINT32_MAX) MaxSteps++;
+	}
+
+	void DecrementStep()
+	{
+		if (MaxSteps > 1) MaxSteps--;
+	}
+};
+
+struct FLoopDebuggerState
+{
+	uint32 NumSteps = 0;
+	const FLoopDebugStepper* LoopDebugStepper = nullptr;
+	UPVData* Data = nullptr;
+
+	bool IsLastLoopIteration() const
+	{
+		check(LoopDebugStepper);
+		return NumSteps == LoopDebugStepper->MaxSteps;
+	}
+
+	bool IterationsCompleted() const
+	{
+		check(LoopDebugStepper);
+		return NumSteps == LoopDebugStepper->MaxSteps;
+	}
+	
+#if WITH_EDITORONLY_DATA
+	FPVDebugSettings GetDebugSettings() const
+	{
+		check(Data);
+		return CopyTemp(Data->GetDebugSettings());
+	}
+
+	void SetDebugSettings(FPVDebugSettings&& InSettings) const
+	{
+		check(Data);
+		Data->SetDebugSettings(MoveTemp(InSettings));
+	}
+#endif
+	
+};
+
+
+#if WITH_EDITORONLY_DATA
+
+inline TUniquePtr<FLoopDebuggerState> GlobalLoopDebuggerState;
+
+#define PVE_LOOP_DEBUG_INIT(InputSettings, LoopDebugStepperName, PVData)													\
+	if (InputSettings->bDebug && InputSettings->LoopDebugStepperName.bDebug)												\
+	{																														\
+		GlobalLoopDebuggerState = MakeUnique<FLoopDebuggerState>(0, &InputSettings->LoopDebugStepperName, PVData);			\
+	}
+
+#define PVE_LOOP_DEBUG_PARAMS_START()																						\
+	if (																													\
+		GlobalLoopDebuggerState &&																							\
+		GlobalLoopDebuggerState->IsLastLoopIteration()																		\
+	)																														\
+	{																														\
+		FPVDebugSettings __DebugSettings = GlobalLoopDebuggerState->GetDebugSettings();										\
+		__DebugSettings.bAutoFocusLoopDebug = GlobalLoopDebuggerState->LoopDebugStepper->bDebugFocus;
+
+#define PVE_LOOP_DEBUG_POINT_PARAM(InParamName, InPivot)																	\
+		__DebugSettings.LoopDebugVisualizationSettings.Emplace(																\
+			EPVDebugValueVisualizationMode::Point,																			\
+			FText::FromString(TEXT(InParamName)),																			\
+			FLoopDebugData(FVector(InPivot), NullOpt, NullOpt)																\
+		);
+
+#define PVE_LOOP_DEBUG_DIRECTION_PARAM(InParamName, InPivot, InDirection)													\
+		__DebugSettings.LoopDebugVisualizationSettings.Emplace(																\
+			EPVDebugValueVisualizationMode::Direction,																		\
+			FText::FromString(TEXT(InParamName)),																			\
+			FLoopDebugData(FVector(InPivot), FVector(InDirection), NullOpt)													\
+		);
+
+#define PVE_LOOP_DEBUG_VECTOR_PARAM(InParamName, InPivot, InVector)															\
+		__DebugSettings.LoopDebugVisualizationSettings.Emplace(																\
+			EPVDebugValueVisualizationMode::Vector,																			\
+			FText::FromString(TEXT(InParamName)),																			\
+			FLoopDebugData(FVector(InPivot), FVector(InVector), NullOpt)													\
+		);
+
+#define PVE_LOOP_DEBUG_TEXT_PARAM(InParamName, InPivot, InText)																\
+		__DebugSettings.LoopDebugVisualizationSettings.Emplace(																\
+			EPVDebugValueVisualizationMode::Text,																			\
+			FText::FromString(TEXT(InParamName)),																			\
+			FLoopDebugData(FVector(InPivot), NullOpt, InText)																\
+		);
+
+#define PVE_LOOP_DEBUG_PARAMS_END()																							\
+		if (GlobalLoopDebuggerState)																						\
+			GlobalLoopDebuggerState->SetDebugSettings(MoveTemp(__DebugSettings));											\
+	}
+
+#define PVE_LOOP_DEBUG_CHECK(ExitMethod)																					\
+	if (																													\
+		GlobalLoopDebuggerState &&																							\
+		GlobalLoopDebuggerState->IsLastLoopIteration()																		\
+	)																														\
+	{																														\
+		GlobalLoopDebuggerState->NumSteps++;																				\
+		ExitMethod;																											\
+	}
+
+#define PVE_LOOP_DEBUG_STEP(ExitMethod)																						\
+	if (GlobalLoopDebuggerState)																							\
+	{																														\
+		PVE_LOOP_DEBUG_CHECK(ExitMethod)																					\
+		GlobalLoopDebuggerState->NumSteps++;																				\
+	}
+
+#define PVE_LOOP_DEBUG_END()																								\
+	GlobalLoopDebuggerState.Reset();																						\
+
+#else
+
+#define PVE_LOOP_DEBUG_INIT(InputSettings, LoopDebugStepperName, PVData)
+
+#define PVE_LOOP_DEBUG_PARAMS_START()
+#define PVE_LOOP_DEBUG_POINT_PARAM(InParamName, InPivot)
+#define PVE_LOOP_DEBUG_DIRECTION_PARAM(InParamName, InPivot, InDirection)
+#define PVE_LOOP_DEBUG_VECTOR_PARAM(InParamName, InPivot, InVector)
+#define PVE_LOOP_DEBUG_TEXT_PARAM(InParamName, InPivot, InText)
+#define PVE_LOOP_DEBUG_PARAMS_END()
+
+#define PVE_LOOP_DEBUG_CHECK(ExitMethod)
+#define PVE_LOOP_DEBUG_STEP(ExitMethod)
+
+#define PVE_LOOP_DEBUG_END()
+
+#endif
diff --git a/Engine/Plugins/Experimental/ProceduralVegetationEditor/Source/ProceduralVegetation/Public/Nodes/PVBaseSettings.h b/Engine/Plugins/Experimental/ProceduralVegetationEditor/Source/ProceduralVegetation/Public/Nodes/PVBaseSettings.h
--- a/Engine/Plugins/Experimental/ProceduralVegetationEditor/Source/ProceduralVegetation/Public/Nodes/PVBaseSettings.h
+++ b/Engine/Plugins/Experimental/ProceduralVegetationEditor/Source/ProceduralVegetation/Public/Nodes/PVBaseSettings.h
@@ -25,6 +25,8 @@ public:
 	virtual void SetCurrentRenderType(TArray<EPVRenderType> InRenderTypes);
 	
 	FPVDebugSettings GetDebugVisualizationSettings() const { return DebugVisualizationSettings; }
+
+	virtual void PostEditChangeProperty(FPropertyChangedEvent& PropertyChangedEvent) override;
 #endif
 
 protected:
@@ -36,10 +38,6 @@ protected:
 	virtual FPCGDataTypeIdentifier GetInputPinTypeIdentifier() const;
 	virtual FPCGDataTypeIdentifier GetOutputPinTypeIdentifier() const;
 
-#if WITH_EDITOR
-	virtual void PostEditChangeProperty(FPropertyChangedEvent& PropertyChangedEvent) override;
-#endif
-	
 private:
 	
 #if WITH_EDITORONLY_DATA
diff --git a/Engine/Plugins/Experimental/ProceduralVegetationEditor/Source/ProceduralVegetation/Public/Nodes/PVObjectAvoidance.h b/Engine/Plugins/Experimental/ProceduralVegetationEditor/Source/ProceduralVegetation/Public/Nodes/PVObjectAvoidance.h
--- a/Engine/Plugins/Experimental/ProceduralVegetationEditor/Source/ProceduralVegetation/Public/Nodes/PVObjectAvoidance.h
+++ b/Engine/Plugins/Experimental/ProceduralVegetationEditor/Source/ProceduralVegetation/Public/Nodes/PVObjectAvoidance.h
@@ -67,6 +67,11 @@ protected:
 public:
 	UPROPERTY(EditAnywhere, Category = "Object Avoidance", meta = (ShowOnlyInnerProperties))
 	FPVObjectAvoidanceParams ObjectAvoidanceSettings;
+
+#if WITH_EDITORONLY_DATA
+	UPROPERTY(EditAnywhere, Category = "Loop Debug", meta=(EditCondition = "bDebug", EditConditionHides))
+	FLoopDebugStepper LoopDebugStepper;
+#endif
 };
 
 class FPVObjectAvoidanceElement : public FPVBaseElement
diff --git a/Engine/Plugins/Experimental/ProceduralVegetationEditor/Source/ProceduralVegetationEditor/Private/Customizations/PVLoopDebugStepperCustomization.cpp b/Engine/Plugins/Experimental/ProceduralVegetationEditor/Source/ProceduralVegetationEditor/Private/Customizations/PVLoopDebugStepperCustomization.cpp
new file mode 100644
--- /dev/null
+++ b/Engine/Plugins/Experimental/ProceduralVegetationEditor/Source/ProceduralVegetationEditor/Private/Customizations/PVLoopDebugStepperCustomization.cpp
@@ -0,0 +1,187 @@
+// Copyright Epic Games, Inc. All Rights Reserved.
+
+#include "PVLoopDebugStepperCustomization.h"
+
+#include "DetailWidgetRow.h"
+
+#include "Helpers/PVUtilities.h"
+
+#include "Widgets/Input/SButton.h"
+#include "Widgets/Input/SNumericEntryBox.h"
+
+TSharedRef<IPropertyTypeCustomization> FPVLoopDebugStepperCustomization::MakeInstance()
+{
+	return MakeShareable(new FPVLoopDebugStepperCustomization);
+}
+
+void FPVLoopDebugStepperCustomization::CustomizeHeader(TSharedRef<IPropertyHandle> InPropertyHandle, FDetailWidgetRow& HeaderRow,
+                                                       IPropertyTypeCustomizationUtils& InCustomizationUtils)
+{
+	PropertyHandle = InPropertyHandle;
+
+	TArray<UObject*> OuterObjects;
+	InPropertyHandle->GetOuterObjects(OuterObjects);
+	TArray<void*> StructPtrs;
+	InPropertyHandle->AccessRawData(StructPtrs);
+
+	if (StructPtrs.Num() == 1)
+	{
+		LoopDebugStepper = static_cast<FLoopDebugStepper*>(StructPtrs[0]);
+		ParentObject = OuterObjects[0];
+
+		check(LoopDebugStepper);
+		check(ParentObject);
+
+		HeaderRow
+			.Visibility(TAttribute<EVisibility>::CreateLambda([]()
+				{
+					return PV::Utilities::DebugModeEnabled()
+						? EVisibility::Visible
+						: EVisibility::Collapsed;
+				}))
+			.NameContent()
+			[
+				InPropertyHandle->CreatePropertyNameWidget()
+			]
+			.ValueContent()
+			.HAlign(HAlign_Fill)
+			[
+				SNew(SHorizontalBox)
+				+ SHorizontalBox::Slot()
+				.HAlign(HAlign_Left)
+				[
+					SNew(SNumericEntryBox<uint32>)
+					.MinDesiredValueWidth(107)
+					.AllowSpin(true)
+					.AllowWheel(true)
+					.Delta(1)
+					.CtrlMultiplier(5)
+					.MinValue(0)
+					.MinSliderValue(0)
+					.MaxValue(NullOpt)
+					.MaxSliderValue(NullOpt)
+					.Value_Lambda([this]()
+						{
+							return LoopDebugStepper->MaxSteps;
+						})
+					.OnValueChanged_Lambda([this](const float NewValue)
+						{
+							LoopDebugStepper->MaxSteps = NewValue;
+						})
+					.OnValueCommitted_Lambda([this](uint32 NewValue, ETextCommit::Type CommitType)
+						{
+							LoopDebugStepper->MaxSteps = NewValue;
+							FPropertyChangedEvent PropertyChangedEvent = FPropertyChangedEvent(PropertyHandle->GetProperty(),
+								EPropertyChangeType::ValueSet);
+							ParentObject->PostEditChangeProperty(PropertyChangedEvent);
+						})
+				]
+				+ SHorizontalBox::Slot()
+				.HAlign(HAlign_Left)
+				[
+					SNew(SHorizontalBox)
+					+ SHorizontalBox::Slot()
+					[
+						SNew(SButton)
+						.ButtonStyle(FAppStyle::Get(), "SimpleButton")
+						.ContentPadding(0.0f)
+						.OnClicked_Lambda([this]()
+							{
+								LoopDebugStepper->DecrementStep();
+								FPropertyChangedEvent PropertyChangedEvent = FPropertyChangedEvent(PropertyHandle->GetProperty(),
+									EPropertyChangeType::ValueSet);
+								ParentObject->PostEditChangeProperty(PropertyChangedEvent);
+								return FReply::Handled();
+							}
+						)
+						[
+							SNew(SImage)
+							.Image_Lambda([this]()
+								{
+									return FAppStyle::Get().GetBrush("Icons.ChevronLeft");
+								})
+						]
+					]
+					+ SHorizontalBox::Slot()
+					[
+						SNew(SButton)
+						.ButtonStyle(FAppStyle::Get(), "SimpleButton")
+						.ContentPadding(0.0f)
+						.OnClicked_Lambda([this]()
+							{
+								LoopDebugStepper->IncrementStep();
+								FPropertyChangedEvent PropertyChangedEvent = FPropertyChangedEvent(PropertyHandle->GetProperty(),
+									EPropertyChangeType::ValueSet);
+								ParentObject->PostEditChangeProperty(PropertyChangedEvent);
+								return FReply::Handled();
+							}
+						)
+						[
+							SNew(SImage)
+							.Image_Lambda([this]()
+								{
+									return FAppStyle::Get().GetBrush("Icons.ChevronRight");
+								})
+						]
+					]
+					+ SHorizontalBox::Slot()
+					[
+						SNew(SButton)
+						.ButtonStyle(FAppStyle::Get(), "SimpleButton")
+						.ContentPadding(2.0f)
+						.OnClicked_Lambda([this]()
+							{
+								LoopDebugStepper->bDebug = !LoopDebugStepper->bDebug;
+								FPropertyChangedEvent PropertyChangedEvent = FPropertyChangedEvent(PropertyHandle->GetProperty(),
+									EPropertyChangeType::ValueSet);
+								ParentObject->PostEditChangeProperty(PropertyChangedEvent);
+								return FReply::Handled();
+							}
+						)
+						[
+							SNew(SImage)
+							.Image_Lambda([this]()
+								{
+									return LoopDebugStepper->bDebug
+										? FAppStyle::Get().GetBrush("GenericStop")
+										: FAppStyle::Get().GetBrush("GenericPlay");
+								})
+						]
+					]
+					+ SHorizontalBox::Slot()
+					[
+						SNew(SCheckBox)
+						.ToolTipText(FText::FromString("Auto-Focus debug stats"))
+						.Style(FAppStyle::Get(), "ToggleButtonCheckbox")
+						.OnCheckStateChanged_Lambda([this](const ECheckBoxState State)
+							{
+								LoopDebugStepper->bDebugFocus = State == ECheckBoxState::Checked;
+								if (LoopDebugStepper->bDebugFocus)
+								{
+									FPropertyChangedEvent PropertyChangedEvent = FPropertyChangedEvent(PropertyHandle->GetProperty(),
+										EPropertyChangeType::ValueSet);
+									ParentObject->PostEditChangeProperty(PropertyChangedEvent);
+								}
+							}
+						)
+						.IsChecked_Lambda([this]()
+							{
+								return LoopDebugStepper->bDebugFocus
+									? ECheckBoxState::Checked
+									: ECheckBoxState::Unchecked;
+							})
+						.Padding(4.0f)
+						[
+							SNew(SImage)
+							.Image(FAppStyle::Get().GetBrush("Symbols.SearchGlass"))
+							.ColorAndOpacity(FSlateColor::UseForeground())
+						]
+					]
+				]
+			];
+	}
+}
+
+void FPVLoopDebugStepperCustomization::CustomizeChildren(TSharedRef<IPropertyHandle> InPropertyHandle, IDetailChildrenBuilder& ChildBuilder,
+                                                         IPropertyTypeCustomizationUtils& InCustomizationUtils)
+{}
diff --git a/Engine/Plugins/Experimental/ProceduralVegetationEditor/Source/ProceduralVegetationEditor/Private/Customizations/PVLoopDebugStepperCustomization.h b/Engine/Plugins/Experimental/ProceduralVegetationEditor/Source/ProceduralVegetationEditor/Private/Customizations/PVLoopDebugStepperCustomization.h
new file mode 100644
--- /dev/null
+++ b/Engine/Plugins/Experimental/ProceduralVegetationEditor/Source/ProceduralVegetationEditor/Private/Customizations/PVLoopDebugStepperCustomization.h
@@ -0,0 +1,31 @@
+// Copyright Epic Games, Inc. All Rights Reserved.
+
+#pragma once
+
+#include "IPropertyTypeCustomization.h"
+
+struct FLoopDebugStepper;
+
+class FPVLoopDebugStepperCustomization : public IPropertyTypeCustomization
+{
+public:
+	static TSharedRef<IPropertyTypeCustomization> MakeInstance();
+
+	virtual void CustomizeHeader(
+		TSharedRef<IPropertyHandle> InPropertyHandle,
+		FDetailWidgetRow& HeaderRow,
+		IPropertyTypeCustomizationUtils& InCustomizationUtils
+	) override;
+
+	virtual void CustomizeChildren(
+		TSharedRef<IPropertyHandle> InPropertyHandle,
+		IDetailChildrenBuilder& ChildBuilder,
+		IPropertyTypeCustomizationUtils& InCustomizationUtils
+	)
+	override;
+
+private:
+	FLoopDebugStepper* LoopDebugStepper = nullptr;
+	TSharedPtr<IPropertyHandle> PropertyHandle = nullptr;
+	UObject* ParentObject = nullptr;
+};
diff --git a/Engine/Plugins/Experimental/ProceduralVegetationEditor/Source/ProceduralVegetationEditor/Private/PVEditor.cpp b/Engine/Plugins/Experimental/ProceduralVegetationEditor/Source/ProceduralVegetationEditor/Private/PVEditor.cpp
--- a/Engine/Plugins/Experimental/ProceduralVegetationEditor/Source/ProceduralVegetationEditor/Private/PVEditor.cpp
+++ b/Engine/Plugins/Experimental/ProceduralVegetationEditor/Source/ProceduralVegetationEditor/Private/PVEditor.cpp
@@ -334,6 +334,25 @@ void FPVEditor::OnSelectedNodesChanged(const TSet<UObject*>& InNewSelection)
 
 	UPCGEditorGraphNodeBase* EdNode = nullptr;
 
+	if (UPCGNode* OldSelectedNode = GetSelectedNode() ? GetSelectedNode()->GetPCGNode() : nullptr)
+	{
+		if (UPVBaseSettings* Settings = Cast<UPVBaseSettings>(OldSelectedNode->GetSettings()))
+		{
+			for (FStructProperty* Property : TFieldRange<FStructProperty>(Settings->GetClass(), EFieldIterationFlags::None))
+			{
+				if (Property->Struct == FLoopDebugStepper::StaticStruct())
+				{
+					FLoopDebugStepper* const LoopDebugStepper = Property->ContainerPtrToValuePtr<FLoopDebugStepper>(Settings);
+					if (LoopDebugStepper->bDebug)
+					{
+						LoopDebugStepper->bDebug = false;
+						FPropertyChangedEvent PropertyChangedEvent = FPropertyChangedEvent(Property, EPropertyChangeType::ValueSet);
+						Settings->PostEditChangeProperty(PropertyChangedEvent);
+					}
+				}
+			}
+		}
+	}
 	SelectedNodes.Empty();
 
 	for(UObject* Selection: InNewSelection)
diff --git a/Engine/Plugins/Experimental/ProceduralVegetationEditor/Source/ProceduralVegetationEditor/Private/ProceduralVegetationEditorModule.cpp b/Engine/Plugins/Experimental/ProceduralVegetationEditor/Source/ProceduralVegetationEditor/Private/ProceduralVegetationEditorModule.cpp
--- a/Engine/Plugins/Experimental/ProceduralVegetationEditor/Source/ProceduralVegetationEditor/Private/ProceduralVegetationEditorModule.cpp
+++ b/Engine/Plugins/Experimental/ProceduralVegetationEditor/Source/ProceduralVegetationEditor/Private/ProceduralVegetationEditorModule.cpp
@@ -37,11 +37,13 @@
 
 #include "GameFramework/Actor.h"
 
-#include "Helpers/PVExportHelper.h"
-
 #include "Modules/ModuleManager.h"
 #include "MessageLogModule.h"
 
+#include "Customizations/PVLoopDebugStepperCustomization.h"
+
+#include "Helpers/PVUtilities.h"
+
 #include "Templates/SharedPointer.h"
 
 #include "Visualizations/PVDataVisualization.h"
@@ -67,6 +69,7 @@ void FProceduralVegetationEditorModule::StartupModule()
 	FPropertyEditorModule& PropertyEditor = FModuleManager::LoadModuleChecked<FPropertyEditorModule>("PropertyEditor");
 	PropertyEditor.RegisterCustomPropertyTypeLayout(FPVExportParams::StaticStruct()->GetFName(), FOnGetPropertyTypeCustomizationInstance::CreateStatic(&FPVOutputSettingsCustomizations::MakeInstance));
 	PropertyEditor.RegisterCustomPropertyTypeLayout(FPVFloatRamp::StaticStruct()->GetFName(), FOnGetPropertyTypeCustomizationInstance::CreateStatic(&FPVFloatRampCustomization::MakeInstance));
+	PropertyEditor.RegisterCustomPropertyTypeLayout(FLoopDebugStepper::StaticStruct()->GetFName(), FOnGetPropertyTypeCustomizationInstance::CreateStatic(&FPVLoopDebugStepperCustomization::MakeInstance));
 
 	FMessageLogModule& MessageLogModule = FModuleManager::LoadModuleChecked<FMessageLogModule>("MessageLog");
 	MessageLogModule.RegisterLogListing(PVEditor::MessageLogName, LOCTEXT("PVEditorLogLabel", "Procedural Vegetation Editor"));
@@ -84,6 +87,7 @@ void FProceduralVegetationEditorModule::ShutdownModule()
 		{
 			PropertyEditor->UnregisterCustomPropertyTypeLayout(FPVExportParams::StaticStruct()->GetFName());
 			PropertyEditor->UnregisterCustomPropertyTypeLayout(FPVFloatRamp::StaticStruct()->GetFName());
+			PropertyEditor->UnregisterCustomPropertyTypeLayout(FLoopDebugStepper::StaticStruct()->GetFName());
 		}
 	}
 
diff --git a/Engine/Plugins/Experimental/ProceduralVegetationEditor/Source/ProceduralVegetationEditor/Private/SPVEditorViewport.cpp b/Engine/Plugins/Experimental/ProceduralVegetationEditor/Source/ProceduralVegetationEditor/Private/SPVEditorViewport.cpp
--- a/Engine/Plugins/Experimental/ProceduralVegetationEditor/Source/ProceduralVegetationEditor/Private/SPVEditorViewport.cpp
+++ b/Engine/Plugins/Experimental/ProceduralVegetationEditor/Source/ProceduralVegetationEditor/Private/SPVEditorViewport.cpp
@@ -262,7 +262,7 @@ void SPVEditorViewport::PopulateViewportOverlays(TSharedRef<SOverlay> Overlay)
 				.Padding(0, 0, 0, 4)
 				[
 					SNew(SRichTextBlock)
-					.Text(FText::FromString(TEXT("<TextBlock.ShadowedText>Foliage Legend</>")))
+					.Text(FText::FromString(TEXT("<TextBlock.ShadowedText>Legend</>")))
 					.Justification(ETextJustify::Center)
 					.TextStyle(&FCoreStyle::Get().GetWidgetStyle<FTextBlockStyle>( "BoldFont" ))
 				]
diff --git a/Engine/Plugins/Experimental/ProceduralVegetationEditor/Source/ProceduralVegetationEditor/Private/Visualizations/DebugVisualization/PVDebugVisualizer.cpp b/Engine/Plugins/Experimental/ProceduralVegetationEditor/Source/ProceduralVegetationEditor/Private/Visualizations/DebugVisualization/PVDebugVisualizer.cpp
--- a/Engine/Plugins/Experimental/ProceduralVegetationEditor/Source/ProceduralVegetationEditor/Private/Visualizations/DebugVisualization/PVDebugVisualizer.cpp
+++ b/Engine/Plugins/Experimental/ProceduralVegetationEditor/Source/ProceduralVegetationEditor/Private/Visualizations/DebugVisualization/PVDebugVisualizer.cpp
@@ -2,34 +2,152 @@
 
 #include "PVDebugVisualizer.h"
 
+#include "AdvancedPreviewScene.h"
 #include "PVDebugVisualizationBase.h"
 #include "PVDebugVisualizations.h"
+#include "SPVEditorViewport.h"
+
+#include "Components/TextRenderComponent.h"
+
 #include "DataTypes/PVData.h"
 
-FVisualizerDrawContext::FVisualizerDrawContext(const FManagedArrayCollection& InCollection, const FPVVisualizationSettings& InSettings,
-	FPCGSceneSetupParams& InSceneParams)
+#include "Materials/MaterialInterface.h"
+
+#include "Visualizations/PVLineBatchComponent.h"
+
+FVisualizerDrawContext::FVisualizerDrawContext(
+	const FManagedArrayCollection& InCollection,
+	const FPVVisualizationSettings& InSettings,
+	FPCGSceneSetupParams& InSceneParams
+)
 	: Collection(InCollection)
 	, VisualizationSettings(InSettings)
 	, SceneSetupParams(InSceneParams)
-{
-}
+{}
 
 void FPVDebugVisualizer::Draw(const UPVData* InData, FPCGSceneSetupParams& InOutParams)
 {
 	check(InData);
-	
+
 	const TArray<FPVVisualizationSettings>& VisualizationSettings = InData->GetDebugSettings().VisualizationSettings;
+	for (const FPVVisualizationSettings& VisualizationSetting : VisualizationSettings)
+	{
+		FVisualizerDrawContext Context(InData->GetCollection(), VisualizationSetting, InOutParams);
+		if (const FPVDebugVisualizationPtr VisualizationPtr = CreateVisualizer(VisualizationSetting))
+		{
+			VisualizationPtr->Draw(Context);
+		}
+	}
+}
+
+void FPVDebugVisualizer::DrawLoopDebug(const UPVData* InData, FPCGSceneSetupParams& InOutParams)
+{
+	check(InData);
+
+	TArray<TPair<FColor, FText>> LegendItems;
+	const TArray<FPVLoopDebugSettings>& LoopDebugVisualizationSettings = InData->GetDebugSettings().LoopDebugVisualizationSettings;
 
-	for (int32 i = 0; i < VisualizationSettings.Num(); i++)
+	float Alpha = 0.0f;
+	const float AlphaStep = 1.0f / static_cast<float>(LoopDebugVisualizationSettings.Num());
+	UPVLineBatchComponent* LineBatchComponent = NewObject<UPVLineBatchComponent>();
+
+	FBox BoundingBox(ForceInit);
+	for (const FPVLoopDebugSettings& LoopDebugVisualizationSetting : LoopDebugVisualizationSettings)
 	{
-		const auto& Settings = VisualizationSettings[i];
+		const FLinearColor Color = PV::Utilities::GetRandomHueColor(Alpha);
+		VisualizeLoopDebuggerParam(InOutParams, LoopDebugVisualizationSetting, LineBatchComponent, Color, BoundingBox);
+		LegendItems.Emplace(Color.ToFColorSRGB(), LoopDebugVisualizationSetting.ParamName);
+		Alpha += AlphaStep;
+	}
+
+	InOutParams.ManagedResources.Add(LineBatchComponent);
+	InOutParams.Scene->AddComponent(LineBatchComponent, FTransform());
 
-		FVisualizerDrawContext Context(InData->GetCollection(), Settings, InOutParams);
+	const TSharedPtr<SPVEditorViewport> ViewportWidget = StaticCastSharedPtr<SPVEditorViewport>(
+		InOutParams.EditorViewportClient->GetEditorViewportWidget()
+	);
+	ViewportWidget->PopulateLegendOverlayText(LegendItems);
 
-		if (FPVDebugVisualizationPtr VisualizationPtr = CreateVisualizer(Settings))
+	if (!BoundingBox.IsValid)
+	{
+		return;
+	}
+
+	if (InData->GetDebugSettings().bAutoFocusLoopDebug)
+	{
+		if (const TSharedPtr<FEditorViewportClient> ViewportClient = ViewportWidget->GetViewportClient())
 		{
-			VisualizationPtr->Draw(Context);
-		}	
+			ViewportClient->FocusViewportOnBox(BoundingBox.ExpandBy(0.25f));
+		}
+	}
+}
+
+void FPVDebugVisualizer::VisualizeLoopDebuggerParam(
+	FPCGSceneSetupParams& InOutParams,
+	const FPVLoopDebugSettings& Settings,
+	UPVLineBatchComponent* LineBatcher,
+	FLinearColor Color,
+	FBox& BoundingBox
+)
+{
+	constexpr float DefaultTextScale = 0.05f;
+
+	if (Settings.VisualizationMode == EPVDebugValueVisualizationMode::Point)
+	{
+		LineBatcher->AddPoint(Settings.Data.Pivot, 5, Color, SDPG_Foreground);
+		BoundingBox += Settings.Data.Pivot;
+	}
+	else if (Settings.VisualizationMode == EPVDebugValueVisualizationMode::Direction && Settings.Data.Direction.IsSet())
+	{
+		const FVector StartPoint = Settings.Data.Pivot;
+		const FVector EndPoint = Settings.Data.Pivot + Settings.Data.Direction.GetValue();
+		BoundingBox += StartPoint;
+		// BoundingBox += EndPoint;
+		const float ArrowSize = 0.0005f * Settings.Data.Direction->Length();
+		{
+			const FVector Dir = (EndPoint - StartPoint).GetUnsafeNormal();
+			FVector Up = FVector::UpVector;
+			FVector Right = Dir ^ Up;
+			if (!Right.IsNormalized())
+			{
+				Dir.FindBestAxisVectors(Up, Right);
+			}
+			const FVector Origin = FVector::ZeroVector;
+			FMatrix TM;
+			TM.SetAxes(&Dir, &Right, &Up, &Origin);
+
+			const float ArrowSqrt = FMath::Sqrt(ArrowSize);
+
+			LineBatcher->AddLine(StartPoint, EndPoint, Color, SDPG_Foreground);
+			LineBatcher->AddLine(EndPoint, EndPoint + TM.TransformPosition(FVector(-ArrowSqrt, ArrowSqrt, 0)), Color, SDPG_Foreground);
+			LineBatcher->AddLine(EndPoint, EndPoint + TM.TransformPosition(FVector(-ArrowSqrt, -ArrowSqrt, 0)), Color, SDPG_Foreground);
+		}
+	}
+	else if (Settings.VisualizationMode == EPVDebugValueVisualizationMode::Vector && Settings.Data.Direction.IsSet())
+	{
+		LineBatcher->AddLine(
+			Settings.Data.Pivot,
+			Settings.Data.Direction.GetValue(),
+			Color,
+			SDPG_Foreground
+		);
+		BoundingBox += Settings.Data.Pivot;
+		// BoundingBox += Settings.Data.Direction.GetValue();
+	}
+	else if (Settings.VisualizationMode == EPVDebugValueVisualizationMode::Text && Settings.Data.TextData.IsSet())
+	{
+		const TObjectPtr<UTextRenderComponent> Component = NewObject<UTextRenderComponent>(GetTransientPackage(), NAME_None, RF_Transient);
+		Component->SetText(Settings.Data.TextData.GetValue());
+		Component->SetTextRenderColor(Color.ToFColorSRGB());
+		Component->SetWorldSize(DefaultTextScale);
+		Component->SetHorizontalAlignment(EHTA_Center);
+		Component->SetVerticalAlignment(EVRTA_TextCenter);
+		Component->SetGenerateOverlapEvents(false);
+		Component->SetMaterial(0, LoadObject<UMaterialInterface>(nullptr, TEXT("/ProceduralVegetationEditor/Materials/LoopDebugTextMaterial.LoopDebugTextMaterial")));
+
+		InOutParams.ManagedResources.Add(Component);
+		InOutParams.Scene->AddComponent(Component, FTransform(Settings.Data.Pivot));
+		BoundingBox += Settings.Data.Pivot;
 	}
 }
 
@@ -37,19 +155,19 @@ FPVDebugVisualizationPtr FPVDebugVisualizer::CreateVisualizer(const FPVVisualiza
 {
 	switch (InSettings.DebugType)
 	{
-		case EPVDebugType::Point:
-			return MakeShared<FPVPointDebugVisualization>();
+	case EPVDebugType::Point:
+		return MakeShared<FPVPointDebugVisualization>();
 
-		case EPVDebugType::Foliage:
-			return MakeShared<FPVFoliageDebugVisualization>();
+	case EPVDebugType::Foliage:
+		return MakeShared<FPVFoliageDebugVisualization>();
 
-		case EPVDebugType::Branches:
-			return MakeShared<FPVBranchDebugVisualization>();
+	case EPVDebugType::Branches:
+		return MakeShared<FPVBranchDebugVisualization>();
 
-		case EPVDebugType::Custom:
-			return MakeShared<FPVCustomDebugVisualization>(InSettings.CustomPivotPositionAttributeName, InSettings.CustomPivotScaleAttributeName, InSettings.CustomGroupToFilter);
+	case EPVDebugType::Custom:
+		return MakeShared<FPVCustomDebugVisualization>(InSettings.CustomPivotPositionAttributeName, InSettings.CustomPivotScaleAttributeName, InSettings.CustomGroupToFilter);
 
-		default:
-			return nullptr;
+	default:
+		return nullptr;
 	}
-}
+}
\ No newline at end of file
diff --git a/Engine/Plugins/Experimental/ProceduralVegetationEditor/Source/ProceduralVegetationEditor/Private/Visualizations/DebugVisualization/PVDebugVisualizer.h b/Engine/Plugins/Experimental/ProceduralVegetationEditor/Source/ProceduralVegetationEditor/Private/Visualizations/DebugVisualization/PVDebugVisualizer.h
--- a/Engine/Plugins/Experimental/ProceduralVegetationEditor/Source/ProceduralVegetationEditor/Private/Visualizations/DebugVisualization/PVDebugVisualizer.h
+++ b/Engine/Plugins/Experimental/ProceduralVegetationEditor/Source/ProceduralVegetationEditor/Private/Visualizations/DebugVisualization/PVDebugVisualizer.h
@@ -5,6 +5,8 @@
 
 #include "PVDebugVisualizationBase.h"
 
+#include "DataTypes/PVData.h"
+
 struct FManagedArrayCollection;
 struct FPVVisualizationSettings;
 class FPVDebugVisualizationBase;
@@ -23,6 +25,16 @@ class FPVDebugVisualizer
 {
 public:
 	static void Draw(const UPVData* InData, FPCGSceneSetupParams& InOutParams);
+	static void DrawLoopDebug(const UPVData* InData, FPCGSceneSetupParams& InOutParams);
+
 private:
+	static void VisualizeLoopDebuggerParam(
+		FPCGSceneSetupParams& InOutParams,
+		const FPVLoopDebugSettings& Settings,
+		UPVLineBatchComponent* LineBatcher,
+		FLinearColor Color,
+		FBox& BoundingBox
+	);
+
 	static FPVDebugVisualizationPtr CreateVisualizer(const FPVVisualizationSettings& InSettings);
-};
\ No newline at end of file
+};
diff --git a/Engine/Plugins/Experimental/ProceduralVegetationEditor/Source/ProceduralVegetationEditor/Private/Visualizations/PVDataVisualization.cpp b/Engine/Plugins/Experimental/ProceduralVegetationEditor/Source/ProceduralVegetationEditor/Private/Visualizations/PVDataVisualization.cpp
--- a/Engine/Plugins/Experimental/ProceduralVegetationEditor/Source/ProceduralVegetationEditor/Private/Visualizations/PVDataVisualization.cpp
+++ b/Engine/Plugins/Experimental/ProceduralVegetationEditor/Source/ProceduralVegetationEditor/Private/Visualizations/PVDataVisualization.cpp
@@ -82,9 +82,10 @@ FPCGSetupSceneFunc FPVDataVisualization::GetViewportSetupFunc(const UPCGSettings
 		InOutParams.Scene->SetFloorVisibility(true);
 		InOutParams.Scene->SetFloorOffset(1.0f);
 
+		InOutParams.EditorViewportClient->OverrideNearClipPlane(0.001f);
+
 		// Reapply view mode settings as for some reason SetFloorVisibility changes EngineFlags
 		InOutParams.EditorViewportClient->SetViewMode(InOutParams.EditorViewportClient->GetViewMode());
-			
 		InOutParams.EditorViewportClient->SetRealtime(true);
 		if (!InOutParams.Scene->IsGridEnabled())
 		{
@@ -119,6 +120,17 @@ FPCGSetupSceneFunc FPVDataVisualization::GetViewportSetupFunc(const UPCGSettings
 				TObjectPtr<UStaticMeshComponent> StaticMeshComponent = NewObject<UStaticMeshComponent>(GetTransientPackage(), NAME_None, RF_Transient);
 				StaticMeshComponent->SetStaticMesh(Mesh.LoadSynchronous());
 				StaticMeshComponent->UpdateBounds();
+				if (WeakData->GetDebugSettings().LoopDebugVisualizationSettings.Num())
+				{
+					const int32 NumMaterials = StaticMeshComponent->GetNumMaterials();
+					for (int32 Index = 0; Index < NumMaterials; Index++)
+					{
+						StaticMeshComponent->SetMaterial(
+							Index,
+							LoadObject<UMaterialInterface>(nullptr, TEXT("/ProceduralVegetationEditor/Materials/WireframeMaterial.WireframeMaterial"))
+						);
+					}
+				}
 
 				InOutParams.ManagedResources.Add(StaticMeshComponent);
 				InOutParams.Scene->AddComponent(StaticMeshComponent, Transform);
@@ -127,9 +139,13 @@ FPCGSetupSceneFunc FPVDataVisualization::GetViewportSetupFunc(const UPCGSettings
 
 		if (PV::Utilities::DebugModeEnabled())
 		{
-			if (WeakSettings.IsValid() && WeakSettings->bDebug)
+			if (WeakSettings.IsValid())
 			{
-				FPVDebugVisualizer::Draw(WeakData.Get(), InOutParams);
+				if (WeakSettings->bDebug)
+				{
+					FPVDebugVisualizer::Draw(WeakData.Get(), InOutParams);
+					FPVDebugVisualizer::DrawLoopDebug(WeakData.Get(), InOutParams);
+				}
 			}
 		}
 
@@ -400,24 +416,11 @@ void FPVDataVisualization::FoliageAttachmentPointsRenderer(const FManagedArrayCo
 	SphereInstancer->SetMaterial(0, LoadObject<UMaterialInterface>(nullptr, TEXT("/ProceduralVegetationEditor/Materials/FoliageAttachmentPointMaterial.FoliageAttachmentPointMaterial")));
 	SphereInstancer->SetNumCustomDataFloats(3);
 
-	const auto GetUniqueColor = [](float Alpha)
-	{
-		Alpha = FMath::Clamp(Alpha, 0.f, 1.f);
-		const float HueDeg = 240.f * Alpha;
-		FMath::SRandInit(Alpha * 255);
-		const float V = FMath::GetMappedRangeValueUnclamped(
-			FVector2f{0.0f, 1.0f},
-			FVector2f{0.25f, 1.0f},
-			FMath::SRand()
-		);
-		return FLinearColor(HueDeg, 0.9, V).HSVToLinearRGB();
-	};
-
 	TArray<TPair<FColor, FText>> LegendItems;
 	const TArray<FString> FoliageNames = FoliageFacade.GetFoliageNames();
 	for (int32 FoliageIndex = 0; FoliageIndex < FoliageNames.Num(); FoliageIndex++)
 	{
-		const FLinearColor Color = GetUniqueColor(static_cast<float>(FoliageIndex) / static_cast<float>(FoliageNames.Num()));
+		const FLinearColor Color = PV::Utilities::GetRandomHueColor(static_cast<float>(FoliageIndex) / static_cast<float>(FoliageNames.Num()));
 		LegendItems.Add({Color.ToFColorSRGB(), FText::FromString(FSoftObjectPath(FoliageNames[FoliageIndex]).GetAssetName())});
 	}
 
@@ -427,7 +430,7 @@ void FPVDataVisualization::FoliageAttachmentPointsRenderer(const FManagedArrayCo
 		const PV::Facades::FFoliageEntryData FoliageData = FoliageFacade.GetFoliageEntry(PointIndex);
 
 		Bounds += static_cast<FVector>(FoliageData.PivotPoint);
-		const FLinearColor Color = GetUniqueColor(static_cast<float>(FoliageData.NameId) / static_cast<float>(FoliageNames.Num()));
+		const FLinearColor Color = PV::Utilities::GetRandomHueColor(static_cast<float>(FoliageData.NameId) / static_cast<float>(FoliageNames.Num()));
 		LineBatchComponent->AddLine(
 			static_cast<FVector>(FoliageData.PivotPoint),
 			static_cast<FVector>(FoliageData.PivotPoint + FoliageData.UpVector * 5.0f),
diff --git a/Engine/Plugins/Experimental/ProceduralVegetationEditor/Source/ProceduralVegetationEditor/Private/Visualizations/PVScaleVisualizationComponent.cpp b/Engine/Plugins/Experimental/ProceduralVegetationEditor/Source/ProceduralVegetationEditor/Private/Visualizations/PVScaleVisualizationComponent.cpp
--- a/Engine/Plugins/Experimental/ProceduralVegetationEditor/Source/ProceduralVegetationEditor/Private/Visualizations/PVScaleVisualizationComponent.cpp
+++ b/Engine/Plugins/Experimental/ProceduralVegetationEditor/Source/ProceduralVegetationEditor/Private/Visualizations/PVScaleVisualizationComponent.cpp
@@ -8,6 +8,8 @@
 
 #include "GameFramework/WorldSettings.h"
 
+#include "Materials/MaterialInterface.h"
+
 UPVScaleVisualizationComponent::UPVScaleVisualizationComponent()
 	: Bounds(ForceInit)
 {
@@ -76,12 +78,12 @@ void UPVScaleVisualizationComponent::UpdateScaleVisualizations()
 		FVector StartPosition = Bounds.Origin;
 		StartPosition.X -= Bounds.BoxExtent.X;
 		StartPosition.Y -= Bounds.BoxExtent.Y;
-		StartPosition.Z -= Bounds.BoxExtent.Z;
+		StartPosition.Z = 0;
 
 		FVector EndPosition = Bounds.Origin;
 		EndPosition.X += Bounds.BoxExtent.X;
 		EndPosition.Y -= Bounds.BoxExtent.Y;
-		EndPosition.Z -= Bounds.BoxExtent.Z;
+		EndPosition.Z = 0;
 
 		const double Length = (Bounds.BoxExtent.X * 2.0f) / WorldToMetersScale;
 		const TObjectPtr<UTextRenderComponent> TextRenderComponent = GetTextRenderComponent(0);
@@ -92,6 +94,7 @@ void UPVScaleVisualizationComponent::UpdateScaleVisualizations()
 		TextRenderComponent->SetHorizontalAlignment(EHTA_Left);
 		TextRenderComponent->SetVerticalAlignment(EVRTA_TextBottom);
 		TextRenderComponent->SetWorldSize(TextSize);
+		TextRenderComponent->SetMaterial(0, LoadObject<UMaterialInterface>(nullptr, TEXT("/ProceduralVegetationEditor/Materials/ScaleTextMaterial.ScaleTextMaterial")));
 		TextRenderComponent->SetText(FText::FromString(FString::Printf(TEXT("%.2fm"), Length)));
 
 		AddLine(StartPosition, EndPosition, FLinearColor::Red);
@@ -101,12 +104,12 @@ void UPVScaleVisualizationComponent::UpdateScaleVisualizations()
 		FVector StartPosition = Bounds.Origin;
 		StartPosition.X -= Bounds.BoxExtent.X;
 		StartPosition.Y -= Bounds.BoxExtent.Y;
-		StartPosition.Z -= Bounds.BoxExtent.Z;
+		StartPosition.Z = 0;
 
 		FVector EndPosition = Bounds.Origin;
 		EndPosition.X -= Bounds.BoxExtent.X;
 		EndPosition.Y += Bounds.BoxExtent.Y;
-		EndPosition.Z -= Bounds.BoxExtent.Z;
+		EndPosition.Z = 0;
 
 		const double Length = (Bounds.BoxExtent.Y * 2.0f) / WorldToMetersScale;
 		const TObjectPtr<UTextRenderComponent> TextRenderComponent = GetTextRenderComponent(1);
@@ -115,6 +118,7 @@ void UPVScaleVisualizationComponent::UpdateScaleVisualizations()
 		TextRenderComponent->SetHorizontalAlignment(EHTA_Right);
 		TextRenderComponent->SetVerticalAlignment(EVRTA_TextBottom);
 		TextRenderComponent->SetWorldSize(TextSize);
+		TextRenderComponent->SetMaterial(0, LoadObject<UMaterialInterface>(nullptr, TEXT("/ProceduralVegetationEditor/Materials/ScaleTextMaterial.ScaleTextMaterial")));
 		TextRenderComponent->SetText(FText::FromString(FString::Printf(TEXT("%.2fm"), Length)));
 
 		AddLine(StartPosition, EndPosition, FLinearColor::Green);
@@ -139,6 +143,7 @@ void UPVScaleVisualizationComponent::UpdateScaleVisualizations()
 		TextRenderComponent->SetHorizontalAlignment(EHTA_Center);
 		TextRenderComponent->SetVerticalAlignment(EVRTA_TextBottom);
 		TextRenderComponent->SetWorldSize(TextSize);
+		TextRenderComponent->SetMaterial(0, LoadObject<UMaterialInterface>(nullptr, TEXT("/ProceduralVegetationEditor/Materials/ScaleTextMaterial.ScaleTextMaterial")));
 		TextRenderComponent->SetText(FText::FromString(FString::Printf(TEXT("%.2fm"), Length)));
 		
 		AddLine(StartPosition, EndPosition, FLinearColor::Blue);
