Add helper functions for traversing a PVE tree
- Add support for traversing both split point and fused point trees
- Update Slope node to use new helper functions

This is a pre-cursor CL for the PVE static mesh importer. The PVFacadeHelper is used extensively by the PVE StaticMesh importer, hence why it needs to be submitted beforehand.

The changes to the Slope node were tested by running the new and old implementation in parallel, and asserting that the results were identical.

https://jira.it.epicgames.com/browse/PVE-75

#Synced-CL 48685397

diff --git a/Engine/Plugins/Experimental/ProceduralVegetationEditor/Source/ProceduralVegetation/Private/Implementations/PVSlope.cpp b/Engine/Plugins/Experimental/ProceduralVegetationEditor/Source/ProceduralVegetation/Private/Implementations/PVSlope.cpp
--- a/Engine/Plugins/Experimental/ProceduralVegetationEditor/Source/ProceduralVegetation/Private/Implementations/PVSlope.cpp
+++ b/Engine/Plugins/Experimental/ProceduralVegetationEditor/Source/ProceduralVegetation/Private/Implementations/PVSlope.cpp
@@ -1,6 +1,7 @@
 // Copyright Epic Games, Inc. All Rights Reserved.
 
 #include "PVSlope.h"
+#include "Helpers/PVFacadeHelper.h"
 #include "Facades/PVBranchFacade.h"
 #include "Facades/PVBudVectorsFacade.h"
 #include "Facades/PVFoliageFacade.h"
@@ -19,113 +20,6 @@ namespace PVSlope
 		return YawRotation * LeanRotation;
 	}
 
-	TArray<int32> GetBranchSegmentChildren(
-		const PV::Facades::FBranchFacade& BranchFacade, 
-		const PV::Facades::FPointFacade& PointFacade,
-		int32 BranchIndex, 
-		int32 BranchPointIndex,
-		int32 PreviousBranchPointIndex)
-	{
-		const TArray<int32>& BranchChildren = BranchFacade.GetChildren(BranchIndex);
-		const TManagedArray<int32>& BranchNumbers = BranchFacade.GetBranchNumbers();
-
-		const float BranchPointLengthFromRoot = PointFacade.GetLengthFromRoot(BranchPointIndex);
-		const float PreviousBranchPointLengthFromRoot = PointFacade.GetLengthFromRoot(PreviousBranchPointIndex);
-
-		TArray<int32> ChildrenIndices;
-		ChildrenIndices.Reserve(BranchChildren.Num());
-
-		for (const int32 BranchChild : BranchChildren)
-		{
-			const int32 ChildIndex = BranchNumbers.Find(BranchChild);
-			if (ChildIndex == INDEX_NONE)
-			{
-				continue;
-			}
-
-			const int32 ParentIndex = BranchFacade.GetParentIndex(ChildIndex);
-			if (ParentIndex != BranchIndex)
-			{
-				continue;
-			}
-
-			const TArray<int32>& ChildBranchPoints = BranchFacade.GetPoints(ChildIndex);
-			if (ChildBranchPoints.Num() == 0)
-			{
-				continue;
-			}
-
-			const int32 FirstChildPointIndex = ChildBranchPoints[0];
-			const float ChildLengthFromRoot = PointFacade.GetLengthFromRoot(FirstChildPointIndex);
-			if (ChildLengthFromRoot > BranchPointLengthFromRoot || ChildLengthFromRoot <= PreviousBranchPointLengthFromRoot)
-			{
-				continue;
-			}
-
-			ChildrenIndices.Add(ChildIndex);
-		}
-
-		return ChildrenIndices;
-	}
-
-	TArray<int32> GetBranchSegmentFoliageEntryIds(
-		PV::Facades::FFoliageFacade& FoliageFacade,
-		const PV::Facades::FPointFacade& PointFacade,
-		int32 BranchIndex,
-		int32 BranchPointIndex,
-		int32 PreviousBranchPointIndex)
-	{
-		TArray<int32> OutFoliageEntryIds;
-
-		const float BranchPointLengthFromRoot = PointFacade.GetLengthFromRoot(BranchPointIndex);
-		const float PreviousBranchPointLengthFromRoot = PointFacade.GetLengthFromRoot(PreviousBranchPointIndex);
-
-		TArray<int32> FoliageEntryIds = FoliageFacade.GetFoliageEntryIdsForBranch(BranchIndex);
-		OutFoliageEntryIds.Reserve(FoliageEntryIds.Num());
-
-		for (const int32 FoliageEntryId : FoliageEntryIds)
-		{
-			const float FoliageEntryLengthFromRoot = FoliageFacade.GetLengthFromRoot(FoliageEntryId);
-
-			if (FoliageEntryLengthFromRoot <= BranchPointLengthFromRoot && FoliageEntryLengthFromRoot > PreviousBranchPointLengthFromRoot)
-			{
-				OutFoliageEntryIds.Add(FoliageEntryId);
-			}
-		}
-
-		return OutFoliageEntryIds;
-	}
-
-	void CalcBranchLengthSquared_Recursive(
-		const PV::Facades::FBranchFacade& BranchFacade,
-		const PV::Facades::FPointFacade& PointFacade,
-		int32 BranchIndex,
-		const FVector3f& BranchStartPosition,
-		float& OutLengthSquared)
-	{
-		const TArray<int32>& CurrentBranchPoints = BranchFacade.GetPoints(BranchIndex);
-		const TManagedArray<FVector3f>& PointPositions = PointFacade.GetPositions();
-
-		for (int32 Index = 0; Index < CurrentBranchPoints.Num(); ++Index)
-		{
-			const int32 PointIndex = CurrentBranchPoints[Index];
-			const int32 PrevPointIndex = Index > 0 ? CurrentBranchPoints[Index - 1] : INDEX_NONE;
-			const FVector3f& PointPosition = PointPositions[PointIndex];
-
-			const float PointDistSquared = (PointPosition - BranchStartPosition).SizeSquared();
-			if (PointDistSquared > OutLengthSquared)
-			{
-				OutLengthSquared = PointDistSquared;
-			}
-
-			const TArray<int32> BranchSegmentChildren = GetBranchSegmentChildren(BranchFacade, PointFacade, BranchIndex, PointIndex, PrevPointIndex);
-			for (const int32 ChildIndex : BranchSegmentChildren)
-			{
-				CalcBranchLengthSquared_Recursive(BranchFacade, PointFacade, ChildIndex, BranchStartPosition, OutLengthSquared);
-			}
-		}
-	}
-
 	float CalcBranchLength(
 		const PV::Facades::FBranchFacade& BranchFacade, 
 		const PV::Facades::FPointFacade& PointFacade, 
@@ -141,8 +35,21 @@ namespace PVSlope
 		const FVector3f& RootPosition = PointPositions[CurrentBranchPoints[0]];
 
 		float BranchLengthSquared = 0;
-		FVector3f DebugLargetsPoint = FVector3f::ZeroVector;
-		CalcBranchLengthSquared_Recursive(BranchFacade, PointFacade, BranchIndex, RootPosition, BranchLengthSquared);
+		PV::FacadeHelper::WalkBranchPoints(
+			BranchFacade,
+			PointFacade,
+			BranchIndex,
+			[&](int32 BranchIndex, int32 BranchPointIndex, int32 PointIndex, const TArray<int32>& ChildPointIndices)->bool
+			{
+				const FVector3f& PointPosition = PointPositions[PointIndex];
+				const float PointDistSquared = (PointPosition - RootPosition).SizeSquared();
+				if (PointDistSquared > BranchLengthSquared)
+				{
+					BranchLengthSquared = PointDistSquared;
+				}
+				return true;
+			}
+		);
 
 		return BranchLengthSquared > 0 ? FMath::Sqrt(BranchLengthSquared) : 0.f;
 	}
@@ -187,7 +94,7 @@ namespace PVSlope
 				CurrentBudDirections[i] = SlopeRotation.RotateVector(CurrentBudDirections[i]);
 			}
 
-			const TArray<int32> BranchSegmentFoliageEntryIds = GetBranchSegmentFoliageEntryIds(FoliageFacade, PointFacade, BranchIndex, PointIndex, PrevPointIndex);
+			const TArray<int32> BranchSegmentFoliageEntryIds = PV::FacadeHelper::GetBranchPointFoliageEntryIds(FoliageFacade, BranchFacade, PointFacade, BranchIndex, PointIndex);
 			for (int32 FoliageEntryId : BranchSegmentFoliageEntryIds)
 			{
 				const FVector3f FoliageEntryPivotPoint = FoliageFacade.GetPivotPoint(FoliageEntryId);
@@ -209,7 +116,7 @@ namespace PVSlope
 			Position = NewPosition;
 			ParentPosition = Position;
 
-			const TArray<int32> BranchSegmentChildren = GetBranchSegmentChildren(BranchFacade, PointFacade, BranchIndex, PointIndex, PrevPointIndex);
+			const TArray<int32> BranchSegmentChildren = PV::FacadeHelper::GetBranchPointChildPoints(BranchFacade, PointFacade, BranchIndex, PointIndex);
 			for (const int32 ChildIndex : BranchSegmentChildren)
 			{
 				ApplySlope_Recursive(
diff --git a/Engine/Plugins/Experimental/ProceduralVegetationEditor/Source/ProceduralVegetation/Public/Helpers/PVFacadeHelper.h b/Engine/Plugins/Experimental/ProceduralVegetationEditor/Source/ProceduralVegetation/Public/Helpers/PVFacadeHelper.h
new file mode 100644
--- /dev/null
+++ b/Engine/Plugins/Experimental/ProceduralVegetationEditor/Source/ProceduralVegetation/Public/Helpers/PVFacadeHelper.h
@@ -0,0 +1,335 @@
+// Copyright Epic Games, Inc. All Rights Reserved.
+#pragma once
+
+#include "Facades/PVPlantFacade.h"
+#include "Facades/PVFoliageFacade.h"
+#include "Facades/PVBranchFacade.h"
+#include "Facades/PVPointFacade.h"
+
+namespace PV::FacadeHelper
+{
+	template <typename StateType>
+	void WalkBranches(
+		const PV::Facades::FPlantFacade& PlantFacade, 
+		const PV::Facades::FBranchFacade& BranchFacade, 
+		const StateType& RootBranchState, 
+		const TFunction<bool(int32, StateType&)>& Body
+	)
+	{
+		const TFunction<bool(int32, const StateType&)> WalkBranches_Recursive = [&](int32 BranchIndex, const StateType& State) -> bool
+		{
+			StateType BranchState = State;
+			if (!Body(BranchIndex, BranchState))
+			{
+				return false;
+			}
+
+			const TArray<int32>& BranchChildren = BranchFacade.GetChildren(BranchIndex);
+			for (int32 BranchChild : BranchChildren)
+			{
+				const int32 ChildBranchIndex = BranchFacade.GetBranchNumbers().Find(BranchChild);
+				check(ChildBranchIndex != INDEX_NONE);
+
+				if (!WalkBranches_Recursive(ChildBranchIndex, BranchState))
+				{
+					return false;
+				}
+			}
+
+			return true;
+		};
+
+		for (const int32 TrunkIndex : PlantFacade.GetTrunkIndices())
+		{
+			if (!WalkBranches_Recursive(TrunkIndex, RootBranchState))
+			{
+				return;
+			}
+		}
+	}
+
+	void WalkBranches(
+		const PV::Facades::FPlantFacade& PlantFacade, 
+		const PV::Facades::FBranchFacade& BranchFacade, 
+		const TFunction<bool(int32 BranchIndex, int32 Depth)>& Body
+	)
+	{
+		int32 Depth = 0;
+		WalkBranches<int32>(PlantFacade, BranchFacade, Depth, [&](int32 BranchIndex, int32& Depth) 
+		{
+			if (!Body(BranchIndex, Depth))
+			{
+				return false;
+			}
+			Depth += 1;
+			return true;
+		});
+	}
+
+	void WalkBranches(
+		const PV::Facades::FPlantFacade& PlantFacade, 
+		const PV::Facades::FBranchFacade& BranchFacade, 
+		const TFunction<bool(int32 BranchIndex)>& Body
+	)
+	{
+		WalkBranches<int32>(PlantFacade, BranchFacade, INDEX_NONE, [&](int32 BranchIndex, int32&) 
+		{
+			if (!Body(BranchIndex))
+			{
+				return false;
+			}
+			return true;
+		});
+	}
+
+	TArray<int32> GetBranchPointFoliageEntryIds(
+		const PV::Facades::FFoliageFacade& FoliageFacade,
+		const PV::Facades::FBranchFacade& BranchFacade,
+		const PV::Facades::FPointFacade& PointFacade,
+		int32 BranchIndex,
+		int32 PointIndex)
+	{
+		const TArray<int32> BranchPoints = BranchFacade.GetPoints(BranchIndex);
+		const int32 PointBranchIndex = BranchPoints.Find(PointIndex);
+		if (PointBranchIndex == INDEX_NONE)
+		{
+			return TArray<int32>();
+		}
+
+		const int32 PreviousPointBranchIndex = PointBranchIndex > 0 ? PointBranchIndex - 1 : INDEX_NONE;
+		if (PreviousPointBranchIndex == INDEX_NONE)
+		{
+			return TArray<int32>();
+		}
+
+		const int32 PreviousPointIndex = BranchPoints[PreviousPointBranchIndex];
+
+		const float BranchPointLengthFromRoot = PointFacade.GetLengthFromRoot(PointIndex);
+		const float PreviousBranchPointLengthFromRoot = PointFacade.GetLengthFromRoot(PreviousPointIndex);
+
+		const TArray<int32> FoliageEntryIds = FoliageFacade.GetFoliageEntryIdsForBranch(BranchIndex);
+
+		TArray<int32> OutFoliageEntryIds;
+		OutFoliageEntryIds.Reserve(FoliageEntryIds.Num());
+
+		for (const int32 FoliageEntryId : FoliageEntryIds)
+		{
+			const float FoliageEntryLengthFromRoot = FoliageFacade.GetLengthFromRoot(FoliageEntryId);
+
+			if (FoliageEntryLengthFromRoot <= BranchPointLengthFromRoot && FoliageEntryLengthFromRoot > PreviousBranchPointLengthFromRoot)
+			{
+				OutFoliageEntryIds.Add(FoliageEntryId);
+			}
+		}
+
+		return OutFoliageEntryIds;
+	}
+
+	enum EWalkBranchMode
+	{
+		FusedPoints,
+		SplitPoints
+	};
+
+	template<EWalkBranchMode WalkBranchMode = EWalkBranchMode::SplitPoints>
+	TArray<int32> GetBranchPointChildPoints(const PV::Facades::FBranchFacade& BranchFacade, const PV::Facades::FPointFacade& PointFacade, int32 BranchIndex, int32 PointIndex);
+
+	template<> TArray<int32> GetBranchPointChildPoints<EWalkBranchMode::FusedPoints>(
+		const PV::Facades::FBranchFacade& BranchFacade, 
+		const PV::Facades::FPointFacade& PointFacade, 
+		int32 BranchIndex,
+		int32 PointIndex
+	)
+	{
+		const TArray<int32>& BranchChildren = BranchFacade.GetChildren(BranchIndex);
+		const TManagedArray<int32>& BranchNumbers = BranchFacade.GetBranchNumbers();
+
+		TArray<int32> ChildrenIndices;
+		ChildrenIndices.Reserve(BranchChildren.Num());
+
+		for (const int32 BranchChild : BranchChildren)
+		{
+			const int32 ChildIndex = BranchNumbers.Find(BranchChild);
+			if (ChildIndex == INDEX_NONE)
+			{
+				continue;
+			}
+
+			const int32 ParentIndex = BranchFacade.GetParentIndex(ChildIndex);
+			if (ParentIndex != BranchIndex)
+			{
+				continue;
+			}
+
+			const TArray<int32>& ChildBranchPoints = BranchFacade.GetPoints(ChildIndex);
+			if (ChildBranchPoints.Num() == 0)
+			{
+				continue;
+			}
+
+			const int32 FirstChildPointIndex = ChildBranchPoints[0];
+			if (FirstChildPointIndex == PointIndex)
+			{
+				ChildrenIndices.Add(ChildIndex);
+			}
+		}
+
+		return ChildrenIndices;
+	};
+
+	template<> TArray<int32> GetBranchPointChildPoints<EWalkBranchMode::SplitPoints>(
+		const PV::Facades::FBranchFacade& BranchFacade, 
+		const PV::Facades::FPointFacade& PointFacade, 
+		int32 BranchIndex, 
+		int32 PointIndex
+	)
+	{
+		const TArray<int32> BranchPoints = BranchFacade.GetPoints(BranchIndex);
+		const int32 PointBranchIndex = BranchPoints.Find(PointIndex);
+		if (PointBranchIndex == INDEX_NONE)
+		{
+			return TArray<int32>();
+		}
+
+		const int32 PreviousPointBranchIndex = PointBranchIndex > 0 ? PointBranchIndex - 1 : INDEX_NONE;
+		if (PreviousPointBranchIndex == INDEX_NONE)
+		{
+			return TArray<int32>();
+		}
+
+		const int32 PreviousPointIndex = BranchPoints[PreviousPointBranchIndex];
+
+		const TArray<int32>& BranchChildren = BranchFacade.GetChildren(BranchIndex);
+		const TManagedArray<int32>& BranchNumbers = BranchFacade.GetBranchNumbers();
+
+		const float BranchPointLengthFromRoot = PointFacade.GetLengthFromRoot(PointIndex);
+		const float PreviousBranchPointLengthFromRoot = PointFacade.GetLengthFromRoot(PreviousPointIndex);
+
+		TArray<int32> ChildrenIndices;
+		ChildrenIndices.Reserve(BranchChildren.Num());
+
+		for (const int32 BranchChild : BranchChildren)
+		{
+			const int32 ChildIndex = BranchNumbers.Find(BranchChild);
+			if (ChildIndex == INDEX_NONE)
+			{
+				continue;
+			}
+
+			const int32 ParentIndex = BranchFacade.GetParentIndex(ChildIndex);
+			if (ParentIndex != BranchIndex)
+			{
+				continue;
+			}
+
+			const TArray<int32>& ChildBranchPoints = BranchFacade.GetPoints(ChildIndex);
+			if (ChildBranchPoints.Num() == 0)
+			{
+				continue;
+			}
+
+			const int32 FirstChildPointIndex = ChildBranchPoints[0];
+			const float ChildLengthFromRoot = PointFacade.GetLengthFromRoot(FirstChildPointIndex);
+			if (ChildLengthFromRoot > BranchPointLengthFromRoot || ChildLengthFromRoot <= PreviousBranchPointLengthFromRoot)
+			{
+				continue;
+			}
+
+			ChildrenIndices.Add(ChildIndex);
+		}
+
+		return ChildrenIndices;
+	};
+
+	template <typename StateType, EWalkBranchMode WalkBranchMode = EWalkBranchMode::SplitPoints>
+	bool WalkBranchPoints(
+		const PV::Facades::FBranchFacade& BranchFacade,
+		const PV::Facades::FPointFacade& PointFacade,
+		int32 BranchIndex,
+		const StateType& RootPointState,
+		const TFunction<bool(int32 BranchIndex, int32 BranchPointIndex, int32 PointIndex, const TArray<int32>& ChildPointIndices, StateType& PointState)>& Body
+	)
+	{
+		const TFunction<bool(int32, bool, const StateType&)> WalkPoints = [&](int32 BranchIndex, bool bSkipFirstBranchPoint, const StateType& ParentState) -> bool
+		{
+			StateType PointState = ParentState;
+			const TArray<int32>& CurrentBranchPoints = BranchFacade.GetPoints(BranchIndex);
+
+			const int32 StartIndex = bSkipFirstBranchPoint ? 1 : 0;
+			for (int32 Index = StartIndex; Index < CurrentBranchPoints.Num(); ++Index)
+			{
+				const int32 PointIndex = CurrentBranchPoints[Index];
+				const TArray<int32> BranchPointChildren = GetBranchPointChildPoints<WalkBranchMode>(BranchFacade, PointFacade, BranchIndex, PointIndex);
+
+				if (!Body(BranchIndex, Index, PointIndex, BranchPointChildren, PointState))
+				{
+					return false;
+				}
+
+				for (const int32 ChildPointIndex : BranchPointChildren)
+				{
+					if (!WalkPoints(ChildPointIndex, WalkBranchMode == EWalkBranchMode::FusedPoints, PointState))
+					{
+						return false;
+					}
+				}
+			}
+
+			return true;
+		};
+
+		return WalkPoints(BranchIndex, false, RootPointState);
+	}
+
+	template <EWalkBranchMode WalkBranchMode = EWalkBranchMode::SplitPoints>
+	bool WalkBranchPoints(
+		const PV::Facades::FBranchFacade& BranchFacade,
+		const PV::Facades::FPointFacade& PointFacade,
+		int32 BranchIndex,
+		const TFunction<bool(int32 BranchIndex, int32 BranchPointIndex, int32 PointIndex, const TArray<int32>& ChildPointIndices)>& Body
+	)
+	{
+		return WalkBranchPoints<int32, WalkBranchMode>(
+			BranchFacade,
+			PointFacade,
+			BranchIndex,
+			0,
+			[&](int32 BranchIndex, int32 BranchPointIndex, int32 PointIndex, const TArray<int32>& ChildPointIndices, int32& PointState)
+			{
+				return Body(BranchIndex, BranchPointIndex, PointIndex, ChildPointIndices);
+			}
+		);
+	}
+
+	template <typename StateType, EWalkBranchMode WalkBranchMode = EWalkBranchMode::SplitPoints>
+	void WalkPlantPoints(
+		const PV::Facades::FPlantFacade& PlantFacade,
+		const PV::Facades::FBranchFacade& BranchFacade,
+		const PV::Facades::FPointFacade& PointFacade,
+		const StateType& RootPointState,
+		const TFunction<bool(int32 BranchIndex, int32 BranchPointIndex, int32 PointIndex, const TArray<int32>& ChildPointIndices, StateType& PointState)>& Body
+	)
+	{
+		for (const int32 TrunkIndex : PlantFacade.GetTrunkIndices())
+		{
+			if (!WalkBranchPoints<StateType, WalkBranchMode>(BranchFacade, PointFacade, TrunkIndex, RootPointState, Body))
+			{
+				return;
+			}
+		}
+	}
+
+	template <EWalkBranchMode WalkBranchMode = EWalkBranchMode::SplitPoints>
+	void WalkPlantPoints(
+		const PV::Facades::FPlantFacade& PlantFacade,
+		const PV::Facades::FBranchFacade& BranchFacade,
+		const PV::Facades::FPointFacade& PointFacade,
+		const TFunction<bool(int32 BranchIndex, int32 BranchPointIndex, int32 PointIndex, const TArray<int32>& ChildPointIndices)>& Body
+	)
+	{
+		WalkPlantPoints<int32, WalkBranchMode>(PlantFacade, BranchFacade, PointFacade, INDEX_NONE, [&](int32 BranchIndex, int32 BranchPointIndex, int32 PointIndex, const TArray<int32>& ChildPointIndices, int32&)
+		{
+			Body(BranchIndex, BranchPointIndex, PointIndex, ChildPointIndices);
+		});
+	}
+}
