First draft of StaticMesh to PVE growth data node

#Synced-CL 48685397

diff --git a/Engine/Plugins/Experimental/ProceduralVegetationEditor/Source/ProceduralVegetation/Private/IPVEditorInterface.cpp b/Engine/Plugins/Experimental/ProceduralVegetationEditor/Source/ProceduralVegetation/Private/IPVEditorInterface.cpp
new file mode 100644
--- /dev/null
+++ b/Engine/Plugins/Experimental/ProceduralVegetationEditor/Source/ProceduralVegetation/Private/IPVEditorInterface.cpp
@@ -0,0 +1,9 @@
+// Copyright Epic Games, Inc. All Rights Reserved.
+
+#pragma once
+
+#include "IPVEditorInterface.h"
+
+#if WITH_EDITOR
+PV::IPVEditorInterface* PV::GPVEditorInterface = nullptr;
+#endif
\ No newline at end of file
diff --git a/Engine/Plugins/Experimental/ProceduralVegetationEditor/Source/ProceduralVegetation/Private/Nodes/PVStaticMeshImporterSettings.cpp b/Engine/Plugins/Experimental/ProceduralVegetationEditor/Source/ProceduralVegetation/Private/Nodes/PVStaticMeshImporterSettings.cpp
new file mode 100644
--- /dev/null
+++ b/Engine/Plugins/Experimental/ProceduralVegetationEditor/Source/ProceduralVegetation/Private/Nodes/PVStaticMeshImporterSettings.cpp
@@ -0,0 +1,74 @@
+// Copyright Epic Games, Inc. All Rights Reserved.
+
+#include "PVStaticMeshImporterSettings.h"
+#include "IPVEditorInterface.h"
+#include "PCGContext.h"
+#include "ProceduralVegetationModule.h"
+#include "DataTypes/PVGrowthData.h"
+#include "ProceduralVegetationPreset.h"
+#include "Facades/PVBranchFacade.h"
+#include "Facades/PVPointFacade.h"
+
+#define LOCTEXT_NAMESPACE "PVStaticMeshImporterSettings"
+
+#if WITH_EDITOR
+FText UPVStaticMeshImporterSettings::GetDefaultNodeTitle() const
+{ 
+	return LOCTEXT("NodeTitle", "Procedural Vegetation Static Mesh Importer"); 
+}
+
+FText UPVStaticMeshImporterSettings::GetNodeTooltipText() const
+{	
+	return LOCTEXT("NodeTooltip", "");
+}
+#endif
+
+TArray<FPCGPinProperties> UPVStaticMeshImporterSettings::InputPinProperties() const
+{
+	return TArray<FPCGPinProperties>();
+}
+
+FPCGDataTypeIdentifier UPVStaticMeshImporterSettings::GetOutputPinTypeIdentifier() const
+{
+	return FPCGDataTypeIdentifier{ FPVDataTypeInfoGrowth::AsId() };
+}
+
+FPCGElementPtr UPVStaticMeshImporterSettings::CreateElement() const
+{
+	return MakeShared<FPVStaticMeshImporterElement>();
+}
+
+bool FPVStaticMeshImporterElement::ExecuteInternal(FPCGContext* InContext) const
+{
+	TRACE_CPUPROFILER_EVENT_SCOPE(FPVStaticMeshImporterElement::ExecuteInternal);
+
+	check(InContext);
+
+	const UPVStaticMeshImporterSettings* Settings = InContext->GetInputSettings<UPVStaticMeshImporterSettings>();
+	check(Settings);
+
+	if (!Settings->Params.StaticMeshAsset)
+	{
+		PCGLog::LogWarningOnGraph(LOCTEXT("MissingStaticMeshAsset", "Static Mesh asset not set"), InContext);
+		return true;
+	}
+
+#if WITH_EDITOR
+	FManagedArrayCollection Collection;
+	if (!PV::GPVEditorInterface || !PV::GPVEditorInterface->ImportGrowthDataFromStaticMesh(Settings->Params, Collection))
+	{
+		return false;
+	}
+
+	UPVGrowthData* OutVariantData = FPCGContext::NewObject_AnyThread<UPVGrowthData>(InContext);
+	OutVariantData->Initialize(MoveTemp(Collection));
+
+	FPCGTaggedData& CollectionOutput = InContext->OutputData.TaggedData.Emplace_GetRef();
+	CollectionOutput.Data = OutVariantData;
+	CollectionOutput.Pin = TEXT("Output");
+#endif
+
+	return true;
+}
+
+#undef LOCTEXT_NAMESPACE
\ No newline at end of file
diff --git a/Engine/Plugins/Experimental/ProceduralVegetationEditor/Source/ProceduralVegetation/Private/Nodes/PVStaticMeshImporterSettings.h b/Engine/Plugins/Experimental/ProceduralVegetationEditor/Source/ProceduralVegetation/Private/Nodes/PVStaticMeshImporterSettings.h
new file mode 100644
--- /dev/null
+++ b/Engine/Plugins/Experimental/ProceduralVegetationEditor/Source/ProceduralVegetation/Private/Nodes/PVStaticMeshImporterSettings.h
@@ -0,0 +1,41 @@
+// Copyright Epic Games, Inc. All Rights Reserved.
+#pragma once
+
+#include "CoreMinimal.h"
+#include "Nodes/PVBaseSettings.h"
+#include "PVImportStaticMeshParams.h"
+
+#include "PVStaticMeshImporterSettings.generated.h"
+
+UCLASS(BlueprintType, ClassGroup = (Procedural))
+class UPVStaticMeshImporterSettings : public UPVBaseSettings
+{
+	GENERATED_BODY()
+
+public:
+	
+#if WITH_EDITOR
+	virtual FName GetDefaultNodeName() const override { return FName(TEXT("ProceduralVegetationStaticMeshImporter")); }
+	virtual FText GetDefaultNodeTitle() const override;
+	virtual FText GetNodeTooltipText() const override;
+	virtual EPCGSettingsType GetType() const override { return EPCGSettingsType::InputOutput; }
+	virtual TArray<EPVRenderType> GetDefaultRenderType() const override { return TArray{ EPVRenderType::PointData }; }
+	virtual TArray<EPVRenderType> GetSupportedRenderTypes() const override { return TArray{ EPVRenderType::PointData, EPVRenderType::Mesh }; }
+#endif
+
+protected:
+	virtual TArray<FPCGPinProperties> InputPinProperties() const override;
+	virtual FPCGDataTypeIdentifier GetOutputPinTypeIdentifier() const override;
+	virtual FPCGElementPtr CreateElement() const override;
+
+public:
+	UPROPERTY(EditAnywhere, Category="Importer Settings", meta = (ShowOnlyInnerProperties))
+	FPVImportStaticMeshParams Params;
+
+};
+
+class FPVStaticMeshImporterElement : public FPVBaseElement
+{
+protected:
+	virtual bool ExecuteInternal(FPCGContext* InContext) const override;
+};
diff --git a/Engine/Plugins/Experimental/ProceduralVegetationEditor/Source/ProceduralVegetation/Public/IPVEditorInterface.h b/Engine/Plugins/Experimental/ProceduralVegetationEditor/Source/ProceduralVegetation/Public/IPVEditorInterface.h
new file mode 100644
--- /dev/null
+++ b/Engine/Plugins/Experimental/ProceduralVegetationEditor/Source/ProceduralVegetation/Public/IPVEditorInterface.h
@@ -0,0 +1,24 @@
+// Copyright Epic Games, Inc. All Rights Reserved.
+
+#pragma once
+
+#include "PVImportStaticMeshParams.h"
+
+struct FManagedArrayCollection;
+struct FPVImportStaticMeshParams;
+class UStaticMesh;
+
+#if WITH_EDITOR
+namespace PV
+{
+	class PROCEDURALVEGETATION_API IPVEditorInterface
+	{
+	public:
+		virtual ~IPVEditorInterface() = default;
+
+		virtual bool ImportGrowthDataFromStaticMesh(const FPVImportStaticMeshParams& InParams, FManagedArrayCollection& OutCollection) = 0;
+	};
+
+	extern PROCEDURALVEGETATION_API IPVEditorInterface* GPVEditorInterface;
+};
+#endif
\ No newline at end of file
diff --git a/Engine/Plugins/Experimental/ProceduralVegetationEditor/Source/ProceduralVegetation/Public/PVImportStaticMeshParams.h b/Engine/Plugins/Experimental/ProceduralVegetationEditor/Source/ProceduralVegetation/Public/PVImportStaticMeshParams.h
new file mode 100644
--- /dev/null
+++ b/Engine/Plugins/Experimental/ProceduralVegetationEditor/Source/ProceduralVegetation/Public/PVImportStaticMeshParams.h
@@ -0,0 +1,44 @@
+// Copyright Epic Games, Inc. All Rights Reserved.
+#pragma once
+
+#include "Engine/EngineTypes.h"
+
+#include "PVImportStaticMeshParams.generated.h"
+
+class UStaticMesh;
+
+UENUM(BlueprintType)
+enum class EPVImportStaticMeshDebugState : uint8
+{
+	None,
+	InitialMesh,
+	KeepMaterial,
+	ComputeNormals,
+	AverageNormals,
+	SmoothNormalsWithLengthPreservation,
+	GetSpherical,
+	BlurRadius,
+	ExtractReductionCloud,
+	RadialFuse,
+	GroupPathPoints,
+	ComputeDist,
+	DrawBranches,
+	RemoveOverlaps,
+	ExtractBranches,
+	MAX UMETA(Hidden)
+};
+
+USTRUCT(BlueprintType)
+struct FPVImportStaticMeshParams
+{
+	GENERATED_BODY()
+
+	UPROPERTY(EditAnywhere, Category = "Asset")
+	TObjectPtr<UStaticMesh> StaticMeshAsset = nullptr;
+
+	UPROPERTY(EditAnywhere, Category = "Asset")
+	FName MaterialFilter = NAME_None;
+
+	UPROPERTY(EditAnywhere, Category = "Asset")
+	EPVImportStaticMeshDebugState DebugState = EPVImportStaticMeshDebugState::None;
+};
diff --git a/Engine/Plugins/Experimental/ProceduralVegetationEditor/Source/ProceduralVegetationEditor/Private/Importer/PVImportHelpers.cpp b/Engine/Plugins/Experimental/ProceduralVegetationEditor/Source/ProceduralVegetationEditor/Private/Importer/PVImportHelpers.cpp
new file mode 100644
--- /dev/null
+++ b/Engine/Plugins/Experimental/ProceduralVegetationEditor/Source/ProceduralVegetationEditor/Private/Importer/PVImportHelpers.cpp
@@ -0,0 +1,1132 @@
+// Copyright Epic Games, Inc. All Rights Reserved.
+#include "PVImportHelpers.h"
+
+#include "ProceduralVegetationPreset.h"
+
+#include "Helpers/PVFacadeHelper.h"
+#include "Facades/PVAttributesNames.h"
+#include "Facades/PVFoliageFacade.h"
+#include "Facades/PVPointFacade.h"
+#include "Facades/PVBranchFacade.h"
+#include "Facades/PVPlantFacade.h"
+
+#include "Engine/StaticMesh.h"
+
+void PV::Import::Helper::CreateEmptyGrowthData(FManagedArrayCollection& OutCollection, int32 NumBranchPoints, int32 NumFoliagePoints, int32 NumBranches, int32 NumFoliageMeshes)
+{
+	// PointGroup
+	OutCollection.AddGroup(PV::GroupNames::PointGroup);
+	OutCollection.AddElements(NumBranchPoints, PV::GroupNames::PointGroup);
+	auto& PointPositionAttribute = OutCollection.AddAttribute<FVector3f>(PV::AttributeNames::PointPosition, PV::GroupNames::PointGroup);
+	auto& LengthFromRootAttribute = OutCollection.AddAttribute<float>(PV::AttributeNames::LengthFromRoot, PV::GroupNames::PointGroup);
+	auto& LengthFromSeedAttribute = OutCollection.AddAttribute<float>(PV::AttributeNames::LengthFromSeed, PV::GroupNames::PointGroup);
+	auto& PointScaleGradientAttribute = OutCollection.AddAttribute<float>(PV::AttributeNames::PointScaleGradient, PV::GroupNames::PointGroup);
+	auto& HullGradientAttribute = OutCollection.AddAttribute<float>(PV::AttributeNames::HullGradient, PV::GroupNames::PointGroup);
+	auto& MainTrunkGradientAttribute = OutCollection.AddAttribute<float>(PV::AttributeNames::MainTrunkGradient, PV::GroupNames::PointGroup);
+	auto& GroundGradientAttribute = OutCollection.AddAttribute<float>(PV::AttributeNames::GroundGradient, PV::GroupNames::PointGroup);
+	auto& PointScaleAttribute = OutCollection.AddAttribute<float>(PV::AttributeNames::PointScale, PV::GroupNames::PointGroup);
+	auto& BudDirectionAttribute = OutCollection.AddAttribute<TArray<FVector3f>>(PV::AttributeNames::BudDirection, PV::GroupNames::PointGroup);
+	auto& BranchGradientAttribute = OutCollection.AddAttribute<float>(PV::AttributeNames::BranchGradient, PV::GroupNames::PointGroup);
+	auto& BudHormoneLevelsAttribute = OutCollection.AddAttribute<TArray<float>>(PV::AttributeNames::BudHormoneLevels, PV::GroupNames::PointGroup);
+	auto& PlantGradientAttribute = OutCollection.AddAttribute<float>(PV::AttributeNames::PlantGradient, PV::GroupNames::PointGroup);
+	auto& BudNumberAttribute = OutCollection.AddAttribute<int32>(PV::AttributeNames::BudNumber, PV::GroupNames::PointGroup);
+	auto& NjordPixelIndexAttribute = OutCollection.AddAttribute<float>(PV::AttributeNames::NjordPixelIndex, PV::GroupNames::PointGroup);
+	auto& BudLightDetectedAttribute = OutCollection.AddAttribute<TArray<float>>(PV::AttributeNames::BudLightDetected, PV::GroupNames::PointGroup);
+	auto& BudDevelopmentAttribute = OutCollection.AddAttribute<TArray<int32>>(PV::AttributeNames::BudDevelopment, PV::GroupNames::PointGroup);
+
+	// Zero initialize vector attributes
+	for (int32 i = 0; i < NumBranchPoints; ++i)
+	{
+		PointPositionAttribute[i] = FVector3f::ZeroVector;
+		LengthFromRootAttribute[i] = 0;
+		LengthFromSeedAttribute[i] = 0;
+		PointScaleGradientAttribute[i] = 0;
+		HullGradientAttribute[i] = 0;
+		MainTrunkGradientAttribute[i] = 0;
+		GroundGradientAttribute[i] = 0;
+		PointScaleAttribute[i] = 0.f;
+		BudDirectionAttribute[i] = { FVector3f::UpVector, FVector3f::UpVector, FVector3f::UpVector, FVector3f::UpVector, FVector3f::UpVector, FVector3f::UpVector };
+		BranchGradientAttribute[i] = 0;
+		BudHormoneLevelsAttribute[i] = { 0.f, 0.f, 0.f, 0.f, 0.f };
+		PlantGradientAttribute[i] = 0;
+		BudNumberAttribute[i] = 0;
+		NjordPixelIndexAttribute[i] = 0;
+		BudLightDetectedAttribute[i] = { 1 };
+		BudDevelopmentAttribute[i] = { 0, 0, 0, 0, 0, 0 };
+	}
+
+	// BranchGroup
+	OutCollection.AddGroup(PV::GroupNames::BranchGroup);
+	OutCollection.AddElements(NumBranches, PV::GroupNames::BranchGroup);
+	auto& BranchParentsAttribute = OutCollection.AddAttribute<TArray<int32>>(PV::AttributeNames::BranchParents, PV::GroupNames::BranchGroup);
+	auto& BranchChildrenAttribute = OutCollection.AddAttribute<TArray<int32>>(PV::AttributeNames::BranchChildren, PV::GroupNames::BranchGroup);
+	auto& BranchPointsAttribute = OutCollection.AddAttribute<TArray<int32>>(PV::AttributeNames::BranchPoints, PV::GroupNames::BranchGroup);
+	auto& BranchNumberAttribute = OutCollection.AddAttribute<int32>(PV::AttributeNames::BranchNumber, PV::GroupNames::BranchGroup);
+	auto& BranchSourceBudNumberAttribute = OutCollection.AddAttribute<int32>(PV::AttributeNames::BranchSourceBudNumber, PV::GroupNames::BranchGroup);
+	auto& BranchHierarchyNumberAttribute = OutCollection.AddAttribute<int32>(PV::AttributeNames::BranchHierarchyNumber, PV::GroupNames::BranchGroup);
+	auto& BranchParentNumberAttribute = OutCollection.AddAttribute<int32>(PV::AttributeNames::BranchParentNumber, PV::GroupNames::BranchGroup);
+	auto& PlantNumberAttribute = OutCollection.AddAttribute<int32>(PV::AttributeNames::PlantNumber, PV::GroupNames::BranchGroup);
+	auto& BranchFoliageIDsAttribute = OutCollection.AddAttribute<TArray<int32>>(PV::AttributeNames::BranchFoliageIDs, PV::GroupNames::BranchGroup);
+
+	for (int32 i = 0; i < NumBranches; ++i)
+	{
+		BranchParentsAttribute[i] = {};
+		BranchChildrenAttribute[i] = {};
+		BranchPointsAttribute[i] = {};
+		BranchNumberAttribute[i] = 0;
+		BranchSourceBudNumberAttribute[i] = 0;
+		BranchHierarchyNumberAttribute[i] = 0;
+		BranchParentNumberAttribute[i] = 0;
+		PlantNumberAttribute[i] = 0;
+		BranchFoliageIDsAttribute[i] = {};
+	}
+
+	// FoliageNamesGroup
+	OutCollection.AddGroup(PV::GroupNames::FoliageNamesGroup);
+	OutCollection.AddElements(NumFoliageMeshes, PV::GroupNames::FoliageNamesGroup);
+	auto& FoliageNamesAttribute = OutCollection.AddAttribute<FString>(PV::AttributeNames::FoliageName, PV::GroupNames::FoliageNamesGroup);
+
+	for (int32 i = 0; i < NumFoliageMeshes; ++i)
+	{
+		FoliageNamesAttribute[i] = TEXT("");
+	}
+	
+	// FoliageGroup
+	OutCollection.AddGroup(PV::GroupNames::FoliageGroup);
+	OutCollection.AddElements(NumFoliagePoints, PV::GroupNames::FoliageGroup);
+	auto& FoliageNameIDAttribute = OutCollection.AddAttribute<int32>(PV::AttributeNames::FoliageNameID, PV::GroupNames::FoliageGroup);
+	auto& FoliageBranchIDAttribute = OutCollection.AddAttribute<int32>(PV::AttributeNames::FoliageBranchID, PV::GroupNames::FoliageGroup);
+	auto& FoliagePivotPointAttribute = OutCollection.AddAttribute<FVector3f>(PV::AttributeNames::FoliagePivotPoint, PV::GroupNames::FoliageGroup);
+	auto& FoliageUPVectorAttribute = OutCollection.AddAttribute<FVector3f>(PV::AttributeNames::FoliageUPVector, PV::GroupNames::FoliageGroup);
+	auto& FoliageNormalVectorAttribute = OutCollection.AddAttribute<FVector3f>(PV::AttributeNames::FoliageNormalVector, PV::GroupNames::FoliageGroup);
+	auto& FoliageScaleAttribute = OutCollection.AddAttribute<float>(PV::AttributeNames::FoliageScale, PV::GroupNames::FoliageGroup);
+	auto& FoliageLengthFromRootAttribute = OutCollection.AddAttribute<float>(PV::AttributeNames::FoliageLengthFromRoot, PV::GroupNames::FoliageGroup);
+	auto& FoliageParentBoneIDAttribute = OutCollection.AddAttribute<int32>(PV::AttributeNames::FoliageParentBoneID, PV::GroupNames::FoliageGroup);
+
+	for (int32 i = 0; i < NumFoliagePoints; ++i)
+	{
+		FoliageNameIDAttribute[i] = 0;
+		FoliageBranchIDAttribute[i] = 0;
+		FoliagePivotPointAttribute[i] = FVector3f::ZeroVector;
+		FoliageUPVectorAttribute[i] = FVector3f::ZeroVector;
+		FoliageNormalVectorAttribute[i] = FVector3f::ZeroVector;
+		FoliageScaleAttribute[i] = 0;
+		FoliageLengthFromRootAttribute[i] = 0;
+		FoliageParentBoneIDAttribute[i] = 0;
+	}
+
+	// DetailsGroup
+	OutCollection.AddGroup(PV::GroupNames::DetailsGroup);
+	OutCollection.AddElements(1, PV::GroupNames::DetailsGroup);
+	auto& FoliagePathAttribute = OutCollection.AddAttribute<FString>(PV::AttributeNames::FoliagePath, PV::GroupNames::DetailsGroup);
+	auto& TrunkMaterialPathAttribute = OutCollection.AddAttribute<FString>(PV::AttributeNames::TrunkMaterialPath, PV::GroupNames::DetailsGroup);
+	auto& TrunkURangeAttribute = OutCollection.AddAttribute<TArray<FVector2f>>(PV::AttributeNames::TrunkURange, PV::GroupNames::DetailsGroup);
+	auto& GuidAttribute = OutCollection.AddAttribute<FGuid>(PV::AttributeNames::Guid, PV::GroupNames::DetailsGroup);
+	auto& LeafPhyllotaxyAttribute = OutCollection.AddAttribute<TArray<float>>(PV::AttributeNames::LeafPhyllotaxy, PV::GroupNames::DetailsGroup);
+
+	FoliagePathAttribute[0] = TEXT("");
+	TrunkMaterialPathAttribute[0] = 0;
+	TrunkURangeAttribute[0] = { };
+	GuidAttribute[0] = FGuid::NewGuid();
+	LeafPhyllotaxyAttribute[0] = { 0, 200, 50, 1, 1, 0, 0, 0 };
+}
+
+void PV::Import::Helper::FillGrowthDataFromBranchHierarchyDescription(FManagedArrayCollection& InOutCollection, const FPVBranchHierarchyDescription& InBranchHierarchy)
+{
+	PV::Facades::FPointFacade PointFacade(InOutCollection);
+	PV::Facades::FBranchFacade BranchFacade(InOutCollection);
+
+	const static auto BranchIndexToBranchNumber = [](int32 BranchIndex) 
+	{ 
+		// Branch Number is BranchIndex + 1 since branch with parent number 0 is a root in PVE and our root has parent index -1.
+		return BranchIndex + 1; 
+	}; 
+
+	TSet<int32> AddedNodes;
+	int32 FirstPointIndex = 0;
+	int32 BudNumber = 1;
+	const TFunction<void(int32, int32)> FillBranchesRecursive = [&](int32 BranchIndex, int32 Depth)
+	{
+		AddedNodes.Add(BranchIndex);
+
+		const FPVBranchDescription& Branch = InBranchHierarchy.Branches[BranchIndex];
+		const int32 LastPointIndex = FirstPointIndex + Branch.PointIndices.Num() - 1;
+
+		TArray<int32> BranchPoints;
+		BranchPoints.Reserve(Branch.PointIndices.Num());
+		for (int32 i = FirstPointIndex; i <= LastPointIndex; ++i)
+		{
+			const int32 PointIndex = Branch.PointIndices[i - FirstPointIndex];
+			BranchPoints.Add(PointIndex);
+
+			PointFacade.SetBudNumber(PointIndex, BudNumber);
+
+			// Setting the njord pixel index to the same as the bud number only works since we're traversing the branches recursivly
+			// This shold be set by the mesher but we set it here since it's easy to do so.
+			PointFacade.SetNjordPixelIndex(PointIndex, BudNumber);
+			
+			BudNumber++;
+			
+			// TODO: Set these based on this doc https://docs.google.com/document/d/1fWNcORvDreHBOakcFfiKPhSXhE5yVTCabn4eWZU1jWg/edit?tab=t.0
+			//BudHormoneLevels[i] = { 0.f, 0.f, 0.f, 0.f, 0.f };
+		}
+		BranchFacade.SetPoints(BranchIndex, MoveTemp(BranchPoints));
+
+		if (Branch.ParentBranchIndex != INDEX_NONE)
+		{
+			TArray<int32> Parents = BranchFacade.GetParents(Branch.ParentBranchIndex);
+			Parents.Add(BranchIndexToBranchNumber(Branch.ParentBranchIndex));
+			BranchFacade.SetParents(BranchIndex, MoveTemp(Parents));
+		}
+
+		TArray<int32> Children;
+		Children.SetNumUninitialized(Branch.ChildBranchIndices.Num());
+		for (int32 i = 0; i < Branch.ChildBranchIndices.Num(); ++i)
+		{
+			check(Branch.ChildBranchIndices[i] != InBranchHierarchy.RootBranchIndex);
+			check(!AddedNodes.Contains(Branch.ChildBranchIndices[i]));
+			Children[i] = BranchIndexToBranchNumber(Branch.ChildBranchIndices[i]);
+		}
+		BranchFacade.SetChildren(BranchIndex, MoveTemp(Children));
+
+		BranchFacade.SetParentBranchNumber(BranchIndex, BranchIndexToBranchNumber(Branch.ParentBranchIndex));
+		BranchFacade.SetBranchNumber(BranchIndex, BranchIndexToBranchNumber(BranchIndex));
+
+		BranchFacade.SetBranchSourceBudNumber(BranchIndex, -1);
+		BranchFacade.SetBranchHierarchyNumber(BranchIndex, Depth);
+		BranchFacade.SetBranchPlantNumber(BranchIndex, 0);
+
+		FirstPointIndex += LastPointIndex - FirstPointIndex + 1;
+
+		for (int32 ChildIndex : Branch.ChildBranchIndices)
+		{
+			FillBranchesRecursive(ChildIndex, Depth + 1);
+		}
+	};
+
+	FillBranchesRecursive(InBranchHierarchy.RootBranchIndex, 0);
+
+	TManagedArray<FVector3f>& Positions = PointFacade.ModifyPositions();
+	TManagedArray<float>& Scales = PointFacade.ModifyPointScales();
+	for (int32 i = 0; i < InBranchHierarchy.Points.Num(); ++i)
+	{
+		Positions[i] = InBranchHierarchy.Points[i];
+		Scales[i] = InBranchHierarchy.PointsRadii[i] > 0 ? InBranchHierarchy.PointsRadii[i] : 0;
+	}
+}
+
+void PV::Import::Helper::FillGrowthDataFromFoliageDescription(FManagedArrayCollection& Collection, const FPVFoliageDescription& InFoliageDescription, const FPVBranchHierarchyDescription& InBranchHierarchy)
+{
+	PV::Facades::FFoliageFacade FoliageFacade(Collection);
+	PV::Facades::FPointFacade PointFacade(Collection);
+
+	const auto FindClosestBranchSegment = [&](const FVector3f& Point, int32& OutBranchIndex, FVector3f& OutClosestPointOnBranch, int32& OutParentPointIndex)
+	{
+		float ClosestBranchDistSquared = MAX_FLT;
+		
+		OutBranchIndex = INDEX_NONE;
+		OutParentPointIndex = INDEX_NONE;
+		OutClosestPointOnBranch = FVector3f::ZeroVector;
+
+		for (int32 i = 0; i < InBranchHierarchy.Branches.Num(); ++i)
+		{
+			FVector3f ClosestPointOnBranch;
+			float BranchDistSquared;
+			int32 PointIndex;
+			DistToBranch(Point, InBranchHierarchy.Points, InBranchHierarchy.Branches[i].PointIndices, ClosestPointOnBranch, BranchDistSquared, PointIndex);
+			if (BranchDistSquared < ClosestBranchDistSquared)
+			{
+				ClosestBranchDistSquared = BranchDistSquared;
+
+				OutBranchIndex = i;
+				OutParentPointIndex = PointIndex;
+				OutClosestPointOnBranch = ClosestPointOnBranch;
+			}
+		}
+	};
+
+	for (int32 i = 0; i < InFoliageDescription.FoliageMeshes.Num(); ++i)
+	{
+		TObjectPtr<UStaticMesh> FoliageMesh = InFoliageDescription.FoliageMeshes[i];
+		FoliageFacade.SetFoliageName(i, FoliageMesh ? FoliageMesh->GetPathName() : TEXT("None"));
+	}
+
+	int32 FoliageEntryIndex = 0;
+	for (int32 i = 0; i < InFoliageDescription.FoliagePoints.Num(); ++i)
+	{
+		const FPVFoliagePointsDescription& FoliagePointsDescription = InFoliageDescription.FoliagePoints[i];
+		for (int32 j = 0; j < FoliagePointsDescription.Positions.Num(); ++j)
+		{
+			int32 ClosestBranchIndex;
+			FVector3f ClosestPointOnBranch;
+			int32 ParentPointIndex;
+			FindClosestBranchSegment(FoliagePointsDescription.Positions[j], ClosestBranchIndex, ClosestPointOnBranch, ParentPointIndex);
+
+			check(ClosestBranchIndex != INDEX_NONE);
+			check(ParentPointIndex != INDEX_NONE);
+
+			const float FoliageLengthFromRoot = PointFacade.GetLengthFromRoot(ParentPointIndex) + (PointFacade.GetPosition(ParentPointIndex) - ClosestPointOnBranch).Size();
+
+			PV::Facades::FFoliageEntryData FoliageEntry;
+			FoliageEntry.NameId = FoliagePointsDescription.FoliageIndices[j];
+			FoliageEntry.BranchId = ClosestBranchIndex;
+			FoliageEntry.PivotPoint = FoliagePointsDescription.Positions[j];
+			FoliageEntry.UpVector = FoliagePointsDescription.Orientations[j].GetUpVector();
+			FoliageEntry.NormalVector = FoliagePointsDescription.Orientations[j].GetRightVector();
+			FoliageEntry.Scale = FoliagePointsDescription.Scales[j].GetAbsMax();
+			FoliageEntry.LengthFromRoot = FoliageLengthFromRoot;
+			FoliageEntry.ParentBoneID = ClosestBranchIndex;
+			FoliageFacade.SetFoliageEntry(FoliageEntryIndex, FoliageEntry);
+
+			TArray<int32> FoliageEntryIdsForBranch = FoliageFacade.GetFoliageEntryIdsForBranch(ClosestBranchIndex);
+			FoliageEntryIdsForBranch.Add(FoliageEntryIndex);
+			FoliageFacade.SetFoliageIdsArray(ClosestBranchIndex, MoveTemp(FoliageEntryIdsForBranch));
+
+			FoliageEntryIndex++;
+		}
+	}
+}
+
+void PV::Import::Helper::ComputeLengthFromRoot(PV::Facades::FPointFacade& PointFacade, const PV::Facades::FBranchFacade& BranchFacade, const PV::Facades::FPlantFacade& PlantFacade)
+{
+	PV::FacadeHelper::WalkBranches(PlantFacade, BranchFacade, [&](int32 BranchIndex)->bool
+	{
+		TManagedArray<float>& LengthFromRoots = PointFacade.ModifyLengthFromRoots();
+		TManagedArray<float>& LengthFromSeeds = PointFacade.ModifyLengthFromSeeds();
+
+		const TArray<int32>& CurrentBranchPoints = BranchFacade.GetPoints(BranchIndex);
+		check(CurrentBranchPoints.Num() > 0);
+
+		const int32 FirstPointIndex = CurrentBranchPoints[0];
+		const int32 ParentBranchIndex = BranchFacade.GetParentBranchIndex(BranchIndex);
+
+		if (ParentBranchIndex == INDEX_NONE)
+		{
+			LengthFromRoots[FirstPointIndex] = 0.f; // No parent, length to root is 0
+			LengthFromSeeds[FirstPointIndex] = 0.f;
+		}
+		else
+		{
+			const TArray<int32>& ParentBranchPoints = BranchFacade.GetPoints(ParentBranchIndex);
+			check(ParentBranchPoints.Num() > 0);
+
+			const FVector3f& FirstPoint = PointFacade.GetPosition(FirstPointIndex);
+
+			int32 ClosestParentPointIndex = ParentBranchPoints[0];
+			float ClosestParentPointDistSquared = (PointFacade.GetPosition(ClosestParentPointIndex) - FirstPoint).SizeSquared();
+
+			for (int32 i = 1; i < ParentBranchPoints.Num(); ++i)
+			{
+				const int32 PointIndex = ParentBranchPoints[i];
+				const FVector3f ParentPoint = PointFacade.GetPosition(PointIndex);
+
+				const float DistSquared = (ParentPoint - FirstPoint).SizeSquared();
+				if (DistSquared > ClosestParentPointDistSquared)
+				{
+					continue;
+				}
+
+				const FVector3f ParentBranchDirection = [&]()
+				{
+					if (i < ParentBranchPoints.Num() - 1)
+					{
+						const int32 NextPointIndex = ParentBranchPoints[i + 1];
+						return (PointFacade.GetPosition(NextPointIndex) - PointFacade.GetPosition(PointIndex)).GetSafeNormal();
+					}
+
+					if (i > 0)
+					{
+						const int32 PrevPointIndex = ParentBranchPoints[i - 1];
+						return (PointFacade.GetPosition(PointIndex) - PointFacade.GetPosition(PrevPointIndex)).GetSafeNormal();
+					}
+
+					return FVector3f::UpVector;
+				}();
+
+				const FVector3f ParentPointToFirstPointDirection = (FirstPoint - ParentPoint).GetSafeNormal();
+				if (ParentPointToFirstPointDirection.Dot(ParentBranchDirection) > 0)
+				{
+					ClosestParentPointIndex = PointIndex;
+					ClosestParentPointDistSquared = DistSquared;
+				}
+			}
+
+			LengthFromRoots[FirstPointIndex] = LengthFromRoots[ClosestParentPointIndex] + (ClosestParentPointDistSquared > 0 ? FMath::Sqrt(ClosestParentPointDistSquared) : 0);
+			LengthFromSeeds[FirstPointIndex] = LengthFromRoots[FirstPointIndex];
+		}
+
+		for (int32 i = 1; i < CurrentBranchPoints.Num(); ++i)
+		{
+			const int32 PointIndex = CurrentBranchPoints[i];
+			const int32 PrevPointIndex = CurrentBranchPoints[i - 1];
+
+			const FVector3f& Position = PointFacade.GetPosition(PointIndex);
+			const FVector3f& PrevPosition = PointFacade.GetPosition(PrevPointIndex);
+
+			LengthFromRoots[PointIndex] = LengthFromRoots[PrevPointIndex] + (Position - PrevPosition).Size();
+			LengthFromSeeds[PointIndex] = LengthFromRoots[PointIndex];
+		}
+
+		return true;
+	});
+}
+
+void PV::Import::Helper::ComputeBudDevelopment(PV::Facades::FPointFacade& PointFacade, const PV::Facades::FBranchFacade& BranchFacade, const PV::Facades::FPlantFacade& PlantFacade)
+{
+	/*
+		0_Generatation: Generation is assumed to be the number of parents in the parents array and does not account for codominant structures. 
+		1_Age: Assuming the hierarchy is already extracted, starting a loop at the seed point and adding a value of 1 to all the connected points will result in an attribute resembling the age. 
+		       In the case of PVE there is no connection between a child and its parents so other metrics need to be implemented E.g utilizing branchSourceBudNumber or getting all points within a given distance.
+		2_Branch_Age: Branch age is simply the age of the first point on each branch.
+		3_Age_Senescence: Is defaulted to zero due to a lack of data and zero being the safest assumption.
+		4_Light_Senescence: Same as Age senescence.
+		5_Relative_Age: Same as age.
+	*/
+
+	PV::FacadeHelper::WalkPlantPoints<int32, PV::FacadeHelper::EWalkBranchMode::FusedPoints>(
+		PlantFacade, 
+		BranchFacade, 
+		PointFacade,
+		0, 
+		[&](int32 BranchIndex, int32 BranchPointIndex, int32 PointIndex, const TArray<int32>& ChildPointIndices, int32& NumWalkedPoints)->bool
+		{
+			NumWalkedPoints += 1;
+
+			const TArray<int32>& CurrentBranchPoints = BranchFacade.GetPoints(BranchIndex);
+			check(CurrentBranchPoints.Num() > 0);
+
+			const int32 FirstPointIndex = CurrentBranchPoints[0];
+			const int32 BranchAge = PointIndex == FirstPointIndex ? NumWalkedPoints : PointFacade.GetBudAge(FirstPointIndex);
+
+			TArray<int32> BudDevelopment = {
+				/*0_Generatation*/ BranchFacade.GetParents(BranchIndex).Num() + 1,
+				/*1_Age*/ NumWalkedPoints,
+				/*2_Branch_Age*/ BranchAge,
+				/*3_Age_Senescence*/ 0,
+				/*4_Light_Senescence*/ 0,
+				/*5_Relative_Age*/ NumWalkedPoints
+			};
+
+			PointFacade.SetBudDevelopment(PointIndex, MoveTemp(BudDevelopment));
+
+			return true;
+		}
+	);
+}
+
+void PV::Import::Helper::ComputeBudDirections(PV::Facades::FPointFacade& PointFacade, const PV::Facades::FBranchFacade& BranchFacade, const PV::Facades::FPlantFacade& PlantFacade)
+{
+	/*
+		0_Apical: Apical is generated by getting the direction from the previous point to the current point, unless its the last.
+		          The last point is the direction of the current point to the next.
+		1_Axillary: This is set by using a phyllotaxy looping over the points. If the current point in an iteration has a child, 
+		            the axillary direction will be set to the orthogonalized apical direction of the child and the phyllotaxy will reset at that point.
+		2_Light_Optimal: Is set to Apical
+		3_Light_Sub_Optimal: Is set to Apical.
+		4_Guide_Curve: Is set to Apical.
+		5 Up_Vector: This is created by operating over the prims in hierarchical order. For the root point of each plant we compute an orthogonal vector 
+		             between the ground plane and the Apical direction. From here we run a parallel transport of the vector throughout the branch. 
+					 For the direct children we compute the dihedral matrix between the parent apical and the child apical direction, transform the parent 
+					 up vector by this matrix and continue the parallel transport. This is repeated until the hierarchy is exhausted.
+	*/
+
+	PV::FacadeHelper::WalkPlantPoints<int32, PV::FacadeHelper::EWalkBranchMode::FusedPoints>(
+		PlantFacade,
+		BranchFacade,
+		PointFacade,
+		INDEX_NONE,
+		[&](int32 BranchIndex, int32 BranchPointIndex, int32 PointIndex, const TArray<int32>& ChildPointIndices, int32& PrevPointIndex)->bool
+		{
+			const TArray<int32>& BranchPoints = BranchFacade.GetPoints(BranchIndex);
+
+			const FVector3f Apical = [&]()->FVector3f
+			{
+				if (BranchPointIndex < BranchPoints.Num() - 1)
+				{
+					const int32 NextPointIndex = BranchPoints[BranchPointIndex + 1];
+					return (PointFacade.GetPosition(NextPointIndex) - PointFacade.GetPosition(PointIndex)).GetSafeNormal();
+				}
+				else if (BranchPointIndex > 0)
+				{
+					const int32 PrevPointIndex = BranchPoints[BranchPointIndex - 1];
+					return (PointFacade.GetPosition(PointIndex) - PointFacade.GetPosition(PrevPointIndex)).GetSafeNormal();
+				}
+
+				ensure(false); // We should never be able to get here unless we have a branch with only 1 point
+				return FVector3f::UpVector;
+			}();
+
+			const FVector3f UpVector = FRotationMatrix44f::MakeFromZX(Apical, FVector3f::ForwardVector).ToQuat().GetRightVector() * -1;
+			const FVector3f PrevAxillary = PrevPointIndex == INDEX_NONE ? FVector3f::ForwardVector : PointFacade.GetBudDirection(PrevPointIndex)[1];
+			const FVector3f Axillary = FQuat4f(Apical, FMath::DegreesToRadians(135.f)).RotateVector(PrevAxillary);
+
+			TArray<FVector3f> BudDirections = {
+				/*0_Apical*/Apical,
+				/*1_Axillary*/Axillary,
+				/*2_Light_Optimal*/Apical,
+				/*3_Light_Sub_Optimal*/Apical,
+				/*4_Guide_Curve*/Apical,
+				/*Up_Vector*/UpVector
+			};
+
+			PointFacade.SetBudDirections(PointIndex, MoveTemp(BudDirections));
+
+			PrevPointIndex = PointIndex;
+
+			return true;
+		}
+	);
+}
+
+void PV::Import::Helper::InterpolatePointsScales(PV::Facades::FPointFacade& PointFacade, const PV::Facades::FBranchFacade& BranchFacade, const PV::Facades::FPlantFacade& PlantFacade)
+{
+	PV::FacadeHelper::WalkBranches(PlantFacade, BranchFacade, [&](int32 BranchIndex)
+	{
+		const TArray<int32>& BranchPoints = BranchFacade.GetPoints(BranchIndex);
+
+		float CurrentPointScale = PointFacade.GetPointScale(BranchPoints.Last());
+		float CurrentLengthFromRoot = PointFacade.GetLengthFromRoot(BranchPoints.Last());
+
+		for (int32 i = BranchPoints.Num() - 1; i >= 0; --i)
+		{
+			float NextValidScale = -1;
+			float NextValidScale_LengthFromRoot = -1;
+			int32 NextValidScale_Index = INDEX_NONE;
+
+			for (int32 j = i - 1; j >= 0; --j)
+			{
+				const int32 PointIndex = BranchPoints[j];
+				const float PointScale = PointFacade.GetPointScale(PointIndex);
+				if (PointScale > 0)
+				{
+					NextValidScale = PointScale;
+					NextValidScale_LengthFromRoot = PointFacade.GetLengthFromRoot(PointIndex);
+					NextValidScale_Index = j;
+					break;
+				}
+			}
+
+			if (NextValidScale_Index == INDEX_NONE)
+			{
+				for (int32 j = i - 1; j >= 0; --j)
+				{
+					PointFacade.SetPointScale(BranchPoints[j], CurrentPointScale);
+				}
+				break;
+			}
+
+			if (NextValidScale_Index == i - 1)
+			{
+				CurrentPointScale = PointFacade.GetPointScale(BranchPoints[NextValidScale_Index]);
+				continue;
+			}
+
+			for (int32 j = i - 1; j >= NextValidScale_Index; --j)
+			{
+				const int32 PointIndex = BranchPoints[j];
+				const float PointLengthFromRoot = PointFacade.GetLengthFromRoot(PointIndex);
+				const float Alpha = FMath::GetMappedRangeValueClamped(FVector2D(CurrentLengthFromRoot, NextValidScale_LengthFromRoot), FVector2D(0, 1), PointLengthFromRoot);
+				const float PointScale = FMath::Lerp(CurrentPointScale, NextValidScale, Alpha);
+				PointFacade.SetPointScale(PointIndex, PointScale);
+			}
+
+			CurrentPointScale = NextValidScale;
+			CurrentLengthFromRoot = NextValidScale_LengthFromRoot;
+			i = NextValidScale_Index + 1;
+		}
+
+		return true;
+	});
+}
+
+bool PV::Import::Helper::SplitFusedPoints(FManagedArrayCollection& Collection)
+{
+	PV::Facades::FBranchFacade BranchFacade(Collection);
+
+	if (!Collection.HasGroup(PV::GroupNames::PointGroup) || !Collection.HasGroup(PV::GroupNames::BranchGroup))
+	{
+		return false;
+	}
+
+	auto& BranchPointsAttribute = Collection.ModifyAttribute<TArray<int32>>(PV::AttributeNames::BranchPoints, PV::GroupNames::BranchGroup);
+	auto& BranchParentNumberAttribute = Collection.ModifyAttribute<int32>(PV::AttributeNames::BranchParentNumber, PV::GroupNames::BranchGroup);
+
+	for (int32 i = 0; i < BranchPointsAttribute.Num(); ++i)
+	{
+		const int32 BranchIndex = i;
+		if (BranchParentNumberAttribute[BranchIndex] == 0)
+		{
+			continue;
+		}
+
+		auto& PointPositionAttribute = Collection.ModifyAttribute<FVector3f>(PV::AttributeNames::PointPosition, PV::GroupNames::PointGroup);
+		auto& LengthFromRootAttribute = Collection.ModifyAttribute<float>(PV::AttributeNames::LengthFromRoot, PV::GroupNames::PointGroup);
+		auto& LengthFromSeedAttribute = Collection.ModifyAttribute<float>(PV::AttributeNames::LengthFromSeed, PV::GroupNames::PointGroup);
+		auto& PointScaleGradientAttribute = Collection.ModifyAttribute<float>(PV::AttributeNames::PointScaleGradient, PV::GroupNames::PointGroup);
+		auto& HullGradientAttribute = Collection.ModifyAttribute<float>(PV::AttributeNames::HullGradient, PV::GroupNames::PointGroup);
+		auto& MainTrunkGradientAttribute = Collection.ModifyAttribute<float>(PV::AttributeNames::MainTrunkGradient, PV::GroupNames::PointGroup);
+		auto& GroundGradientAttribute = Collection.ModifyAttribute<float>(PV::AttributeNames::GroundGradient, PV::GroupNames::PointGroup);
+		auto& PointScaleAttribute = Collection.ModifyAttribute<float>(PV::AttributeNames::PointScale, PV::GroupNames::PointGroup);
+		auto& BudDirectionAttribute = Collection.ModifyAttribute<TArray<FVector3f>>(PV::AttributeNames::BudDirection, PV::GroupNames::PointGroup);
+		auto& BranchGradientAttribute = Collection.ModifyAttribute<float>(PV::AttributeNames::BranchGradient, PV::GroupNames::PointGroup);
+		auto& BudHormoneLevelsAttribute = Collection.ModifyAttribute<TArray<float>>(PV::AttributeNames::BudHormoneLevels, PV::GroupNames::PointGroup);
+		auto& PlantGradientAttribute = Collection.ModifyAttribute<float>(PV::AttributeNames::PlantGradient, PV::GroupNames::PointGroup);
+		auto& BudNumberAttribute = Collection.ModifyAttribute<int32>(PV::AttributeNames::BudNumber, PV::GroupNames::PointGroup);
+		auto& NjordPixelIndexAttribute = Collection.ModifyAttribute<float>(PV::AttributeNames::NjordPixelIndex, PV::GroupNames::PointGroup);
+		auto& BudLightDetectedAttribute = Collection.ModifyAttribute<TArray<float>>(PV::AttributeNames::BudLightDetected, PV::GroupNames::PointGroup);
+		auto& BudDevelopmentAttribute = Collection.ModifyAttribute<TArray<int32>>(PV::AttributeNames::BudDevelopment, PV::GroupNames::PointGroup);
+
+		const int32 NumPoints = PointPositionAttribute.Num();
+		Collection.Resize(NumPoints + 1, PV::GroupNames::PointGroup);
+
+		TArray<int32>& BranchPoints = BranchPointsAttribute[BranchIndex];
+		const int32 PrevPointIndex = BranchPoints[0];
+		const int32 NewPointIndex = NumPoints;
+
+		BranchPoints[0] = NewPointIndex;
+
+		PointPositionAttribute[NewPointIndex] = PointPositionAttribute[PrevPointIndex];
+		LengthFromRootAttribute[NewPointIndex] = LengthFromRootAttribute[PrevPointIndex];
+		LengthFromSeedAttribute[NewPointIndex] = LengthFromSeedAttribute[PrevPointIndex];
+		PointScaleGradientAttribute[NewPointIndex] = PointScaleGradientAttribute[PrevPointIndex];
+		HullGradientAttribute[NewPointIndex] = HullGradientAttribute[PrevPointIndex];
+		MainTrunkGradientAttribute[NewPointIndex] = MainTrunkGradientAttribute[PrevPointIndex];
+		GroundGradientAttribute[NewPointIndex] = GroundGradientAttribute[PrevPointIndex];
+		PointScaleAttribute[NewPointIndex] = PointScaleAttribute[PrevPointIndex];
+		BudDirectionAttribute[NewPointIndex] = BudDirectionAttribute[PrevPointIndex];
+		BranchGradientAttribute[NewPointIndex] = BranchGradientAttribute[PrevPointIndex];
+		BudHormoneLevelsAttribute[NewPointIndex] = BudHormoneLevelsAttribute[PrevPointIndex];
+		PlantGradientAttribute[NewPointIndex] = PlantGradientAttribute[PrevPointIndex];
+		BudNumberAttribute[NewPointIndex] = BudNumberAttribute[PrevPointIndex];
+		NjordPixelIndexAttribute[NewPointIndex] = NjordPixelIndexAttribute[PrevPointIndex];
+		BudLightDetectedAttribute[NewPointIndex] = BudLightDetectedAttribute[PrevPointIndex];
+		BudDevelopmentAttribute[NewPointIndex] = BudDevelopmentAttribute[PrevPointIndex];
+	}
+
+	return true;
+}
+
+void PV::Import::Helper::ConnectDisconnectedBranches(FPVBranchHierarchyDescription& InOutBranchHierarchy)
+{
+	struct FNewBranchRootIndex
+	{
+		int32 BranchIndex;
+		int32 NewRootPointIndex;
+	};
+	TArray<FNewBranchRootIndex> BranchRootsToFuse;
+	BranchRootsToFuse.Reserve(InOutBranchHierarchy.Branches.Num()); // We expect to fuse at most one point per branch
+
+	TArray<FNewBranchRootIndex> BranchRootsToConnect;
+	BranchRootsToConnect.Reserve(InOutBranchHierarchy.Branches.Num()); // We expect to connect at most one point per branch
+
+	TArray<float> ApicalDots;
+	TArray<float> PointDistances;
+
+	for (const FPVBranchDescription& Branch : InOutBranchHierarchy.Branches)
+	{
+		if (Branch.ParentBranchIndex == INDEX_NONE)
+		{
+			continue;
+		}
+
+		check(Branch.PointIndices.Num() > 1);
+			
+		const FPVBranchDescription& ParentBranch = InOutBranchHierarchy.Branches[Branch.ParentBranchIndex];
+			
+		const int32 FirstPointIndex = Branch.PointIndices[0];
+		const FVector3f& FirstPointLocation = InOutBranchHierarchy.Points[FirstPointIndex];
+
+		const int32 LastParentPointIndex = ParentBranch.PointIndices.Last();
+		const FVector3f& LastParentPointLocation = InOutBranchHierarchy.Points[LastParentPointIndex];
+
+		const float ParentPointDistSquared = (FirstPointLocation - LastParentPointLocation).SizeSquared();
+
+		constexpr float MaxFuseDistanceSquared = 1.f; // 1cm
+		if (ParentPointDistSquared <= MaxFuseDistanceSquared)
+		{
+			BranchRootsToFuse.Add({
+				.BranchIndex = Branch.BranchIndex,
+				.NewRootPointIndex = LastParentPointIndex,
+			});
+			continue;
+		}
+
+		const int32 SecondPointIndex = Branch.PointIndices[1];
+		const FVector3f& SecondPointLocation = InOutBranchHierarchy.Points[SecondPointIndex];
+		const FVector3f RootApicalDirection = (SecondPointLocation - FirstPointLocation).GetSafeNormal();
+
+		ApicalDots.SetNumUninitialized(ParentBranch.PointIndices.Num(), EAllowShrinking::No);
+		PointDistances.SetNumUninitialized(ParentBranch.PointIndices.Num(), EAllowShrinking::No);
+
+		int32 ClosestPointIndex = INDEX_NONE;
+		float ClosestPointDistance = MAX_flt;
+		for (int32 i = 0; i < ParentBranch.PointIndices.Num(); ++i)
+		{
+			const int32 ParentPointIndex = ParentBranch.PointIndices[i];
+			const FVector3f& ParentPointLocation = InOutBranchHierarchy.Points[ParentPointIndex];
+
+			const FVector3f DeltaLocation = (FirstPointLocation - ParentPointLocation);
+				
+			FVector3f ApicalDirection = FVector3f::ZeroVector;
+			float PointDist = 0;
+
+			const float PointDistSquared = DeltaLocation.SizeSquared();
+
+			if (PointDistSquared == 1.f)
+			{
+				ApicalDirection = DeltaLocation;
+			}
+			else if (PointDistSquared < UE_SMALL_NUMBER)
+			{
+				ApicalDirection = FVector3f::ZeroVector;
+			}
+			else
+			{
+				const float InvPointDist = FMath::InvSqrt(PointDistSquared);
+				ApicalDirection = FVector3f(DeltaLocation.X * InvPointDist, DeltaLocation.Y * InvPointDist, DeltaLocation.Z * InvPointDist);
+				PointDist = PointDistSquared * InvPointDist;
+			}
+
+			const float ApicalDotProduct = RootApicalDirection.Dot(ApicalDirection);
+			ApicalDots[i] = ApicalDotProduct;
+			PointDistances[i] = PointDist;
+
+			if (PointDist < ClosestPointDistance)
+			{
+				ClosestPointIndex = i;
+				ClosestPointDistance = PointDist;
+			}
+		}
+
+		int32 ClosestApicalIndex = ClosestPointIndex;
+		for (int32 i = 0; i < ParentBranch.PointIndices.Num(); ++i)
+		{
+			const float PercentageDiff = FMath::Abs(PointDistances[ClosestPointIndex] - PointDistances[i]) / PointDistances[i];
+			if (PercentageDiff < 0.5)
+			{
+				if (ApicalDots[i] > ApicalDots[ClosestApicalIndex])
+				{
+					ClosestApicalIndex = i;
+				}
+			}
+		}
+
+		BranchRootsToConnect.Add({
+			.BranchIndex = Branch.BranchIndex,
+			.NewRootPointIndex = ParentBranch.PointIndices[ClosestApicalIndex],
+		});
+	}
+
+	// Connect disconnected branches
+	for (const FNewBranchRootIndex& NewBranchRootIndex : BranchRootsToConnect)
+	{
+		InOutBranchHierarchy.Branches[NewBranchRootIndex.BranchIndex].PointIndices.Insert(NewBranchRootIndex.NewRootPointIndex, 0);
+	}
+
+	// Fuse branch roots
+	{
+		TArray<int32> PointsToRemove;
+		PointsToRemove.Reserve(BranchRootsToFuse.Num());
+		for (const FNewBranchRootIndex& NewBranchRootIndex : BranchRootsToFuse)
+		{
+			const int32 OldBranchRootPointIndex = InOutBranchHierarchy.Branches[NewBranchRootIndex.BranchIndex].PointIndices[0];
+			PointsToRemove.Add(OldBranchRootPointIndex);
+
+			InOutBranchHierarchy.Branches[NewBranchRootIndex.BranchIndex].PointIndices[0] = NewBranchRootIndex.NewRootPointIndex;
+		}
+
+		struct FPointRef
+		{
+			int32 BranchIndex;
+			int32 PointIndicesIndex;
+		};
+		TArray<TArray<FPointRef>> PointsRefs;
+		PointsRefs.SetNum(InOutBranchHierarchy.Points.Num());
+
+		for (int32 i = 0; i < InOutBranchHierarchy.Branches.Num(); ++i)
+		{
+			const TArray<int32>& PointIndices = InOutBranchHierarchy.Branches[i].PointIndices;
+			for (int32 j = 0; j < PointIndices.Num(); ++j)
+			{
+				const int32 PointIndex = PointIndices[j];
+				PointsRefs[PointIndex].Add({ i, j });
+			}
+		}
+
+		PointsToRemove.Sort(); // Remove points from back to front to avoid having to continuously re-map the PointsToRemove array
+		for (int32 i = PointsToRemove.Num() - 1; i >= 0; --i)
+		{
+			const int32 PointIndex = PointsToRemove[i];
+
+			check(PointsRefs[PointIndex].Num() == 0);
+			const int32 SwappedPointIndex = InOutBranchHierarchy.Points.Num() - 1;
+			InOutBranchHierarchy.Points.RemoveAtSwap(PointIndex);
+
+			for (const FPointRef& PointRef : PointsRefs[SwappedPointIndex])
+			{
+				InOutBranchHierarchy.Branches[PointRef.BranchIndex].PointIndices[PointRef.PointIndicesIndex] = PointIndex;
+			}
+
+			PointsRefs[PointIndex] = MoveTemp(PointsRefs[SwappedPointIndex]);
+		}
+	}
+}
+
+void PV::Import::Helper::FuseCollinearBranches(FPVBranchHierarchyDescription& InOutBranchHierarchy)
+{
+	TArray<int32> BranchesToFuseWithParent;
+	BranchesToFuseWithParent.Reserve(InOutBranchHierarchy.Branches.Num());
+
+	const TFunction<void(int32)> SearchForBranchesToFuse_Recursive = [&](int32 BranchIndex)
+	{
+		const FPVBranchDescription& Branch = InOutBranchHierarchy.Branches[BranchIndex];
+
+		if (Branch.ChildBranchIndices.Num() == 0)
+		{
+			return;
+		}
+
+		const int32 LastParentPointIndex = Branch.PointIndices[Branch.PointIndices.Num() - 1];
+		const int32 SecondToLastParentPointIndex = Branch.PointIndices[Branch.PointIndices.Num() - 2];
+
+		const FVector3f& LastParentPointLocation = InOutBranchHierarchy.Points[LastParentPointIndex];
+		const FVector3f& SecondToLastParentPointLocation = InOutBranchHierarchy.Points[SecondToLastParentPointIndex];
+		const FVector3f ParentApicalDirection = (LastParentPointLocation - SecondToLastParentPointLocation).GetSafeNormal();
+
+		int32 ChildBranchToFuse = INDEX_NONE;
+		float BestApicalDotProduct = -1;
+		for (int32 ChildBranchIndex : Branch.ChildBranchIndices)
+		{
+			const FPVBranchDescription& ChildBranch = InOutBranchHierarchy.Branches[ChildBranchIndex];
+
+			const int32 FistChildBranchPointIndex = ChildBranch.PointIndices[0];
+			if (FistChildBranchPointIndex != LastParentPointIndex)
+			{
+				continue;
+			}
+
+			const FVector3f& FirstChildPointLocation = InOutBranchHierarchy.Points[FistChildBranchPointIndex];
+			const int32 SecondChildBranchPointIndex = ChildBranch.PointIndices[1];
+			const FVector3f& SecondChildPointLocation = InOutBranchHierarchy.Points[SecondChildBranchPointIndex];
+
+			const FVector3f ChildApicalDirection = (SecondChildPointLocation - FirstChildPointLocation).GetSafeNormal();
+			const float ApicalDotProduct = ChildApicalDirection.Dot(ParentApicalDirection);
+
+			constexpr float MinDotForFusingChild = 0.90f;
+			if (ApicalDotProduct < MinDotForFusingChild)
+			{
+				continue;
+			}
+					
+			if (ApicalDotProduct < BestApicalDotProduct)
+			{
+				continue;
+			}
+
+			ChildBranchToFuse = ChildBranchIndex;
+			BestApicalDotProduct = ApicalDotProduct;
+		}
+
+		if (ChildBranchToFuse != INDEX_NONE)
+		{
+			BranchesToFuseWithParent.Add(ChildBranchToFuse);
+		}
+
+		for (int32 ChildBranchIndex : Branch.ChildBranchIndices)
+		{
+			SearchForBranchesToFuse_Recursive(ChildBranchIndex);
+		}
+	};
+
+	SearchForBranchesToFuse_Recursive(InOutBranchHierarchy.RootBranchIndex);
+
+	TArray<TArray<int32>> ParentBranchRefs;
+	ParentBranchRefs.SetNum(InOutBranchHierarchy.Branches.Num());
+
+	BranchesToFuseWithParent.Sort(); // Remove points from back to front to avoid having to continuously re-map the BranchesToFuseWithParent array
+	for (int32 i = BranchesToFuseWithParent.Num() - 1; i >= 0; --i)
+	{
+		const int32 OrphanedBranchIndex = BranchesToFuseWithParent[i];
+
+		// Fuse branch with parent
+		{
+			const FPVBranchDescription& OrphanedBranch = InOutBranchHierarchy.Branches[OrphanedBranchIndex];
+			FPVBranchDescription& BranchToFuseWith = InOutBranchHierarchy.Branches[OrphanedBranch.ParentBranchIndex];
+
+			check(BranchToFuseWith.PointIndices.Last() == OrphanedBranch.PointIndices[0]);
+
+			BranchToFuseWith.PointIndices.Reserve(BranchToFuseWith.PointIndices.Num() + OrphanedBranch.PointIndices.Num() - 1);
+			for (int32 j = 1; j < OrphanedBranch.PointIndices.Num(); ++j)
+			{
+				BranchToFuseWith.PointIndices.Add(OrphanedBranch.PointIndices[j]);
+			}
+
+			check(BranchToFuseWith.ChildBranchIndices.Contains(OrphanedBranchIndex));
+
+			BranchToFuseWith.ChildBranchIndices.Remove(OrphanedBranchIndex);
+			BranchToFuseWith.ChildBranchIndices.Reserve(BranchToFuseWith.ChildBranchIndices.Num() + OrphanedBranch.ChildBranchIndices.Num());
+			for (int32 ChildBranchIndex : OrphanedBranch.ChildBranchIndices)
+			{
+				FPVBranchDescription& ChildBranch = InOutBranchHierarchy.Branches[ChildBranchIndex];
+				BranchToFuseWith.ChildBranchIndices.Add(ChildBranch.BranchIndex);
+				ChildBranch.ParentBranchIndex = BranchToFuseWith.BranchIndex;
+			}
+		}
+			
+		// Remove orphaned branch + fixup branch references
+		{
+			const int32 NewBranchIndex = OrphanedBranchIndex;
+			const int32 OldBranchIndex = InOutBranchHierarchy.Branches.Num() - 1;
+			InOutBranchHierarchy.Branches.RemoveAtSwap(OrphanedBranchIndex);
+
+			if (OldBranchIndex == NewBranchIndex)
+			{
+				continue;
+			}
+
+			FPVBranchDescription& SwappedBranch = InOutBranchHierarchy.Branches[NewBranchIndex];
+			SwappedBranch.BranchIndex = NewBranchIndex;
+
+			for (int32 ChildBranchIndex : SwappedBranch.ChildBranchIndices)
+			{
+				InOutBranchHierarchy.Branches[ChildBranchIndex].ParentBranchIndex = NewBranchIndex;
+			}
+
+			FPVBranchDescription& SwappedBranchParentBranch = InOutBranchHierarchy.Branches[SwappedBranch.ParentBranchIndex];
+			const int32 OldParentChildIndicesIndex = SwappedBranchParentBranch.ChildBranchIndices.Find(OldBranchIndex);
+			check(OldParentChildIndicesIndex != INDEX_NONE);
+
+			SwappedBranchParentBranch.ChildBranchIndices[OldParentChildIndicesIndex] = NewBranchIndex;
+		}
+	}
+}
+
+bool PV::Import::Helper::GenerateGrowthDataFromPlantDescription(FManagedArrayCollection& OutCollection, const FPVPlantDescription& InPlantDescription)
+{
+	const FPVBranchHierarchyDescription& InBranchHierarchy = InPlantDescription.BranchHierarchy;
+	const FPVFoliageDescription& InFoliageDescription = InPlantDescription.Foliage;
+
+	if (InBranchHierarchy.RootBranchIndex == INDEX_NONE)
+	{
+		return false;
+	}
+
+	const int32 NumBranchPoints = InPlantDescription.GetNumBranchPoints();
+	const int32 NumFoliagePoints = InPlantDescription.GetNumFoliagePoints();
+	const int32 NumBranches = InBranchHierarchy.Branches.Num();
+	const int32 NumFoliageMeshes = InFoliageDescription.FoliageMeshes.Num();
+
+	CreateEmptyGrowthData(OutCollection, NumBranchPoints, NumFoliagePoints, NumBranches, NumFoliageMeshes);
+	FillGrowthDataFromBranchHierarchyDescription(OutCollection, InBranchHierarchy);
+	FillGrowthDataFromFoliageDescription(OutCollection, InFoliageDescription, InBranchHierarchy);
+
+	PV::Facades::FPointFacade PointFacade(OutCollection);
+	const PV::Facades::FBranchFacade BranchFacade(OutCollection);
+	const PV::Facades::FPlantFacade PlantFacade(OutCollection);
+
+	ComputeLengthFromRoot(PointFacade, BranchFacade, PlantFacade);
+	ComputeBudDevelopment(PointFacade, BranchFacade, PlantFacade);
+	ComputeBudDirections(PointFacade, BranchFacade, PlantFacade);
+	InterpolatePointsScales(PointFacade, BranchFacade, PlantFacade);
+
+	SplitFusedPoints(OutCollection); // NOTE: Important that this is done last as ComputeBudDevelopment and ComputeBudDirections relies on fused points to traverse the tree
+
+	return true;
+}
+
+bool PV::Import::Helper::GenerateProceduralVegetationPresetFromPlantDescription(const FPVPlantDescription& InPlantDescription, UProceduralVegetationPreset* InOutProceduralVegetationPreset)
+{
+	FManagedArrayCollection Collection;
+	if (!GenerateGrowthDataFromPlantDescription(Collection, InPlantDescription))
+	{
+		return false;
+	}
+
+	InOutProceduralVegetationPreset->Modify();
+	InOutProceduralVegetationPreset->Variants.Reset();
+	InOutProceduralVegetationPreset->Variants.Add(TEXT("Default"), Collection);
+	InOutProceduralVegetationPreset->FillVariationInfo();
+	InOutProceduralVegetationPreset->PostEditChange();
+
+	return true;
+}
+
+void PV::Import::Helper::DistToBranch(const FVector3f& Point, const TArray<FVector3f>& Points, const TArray<int32>& BranchPointIndices, FVector3f& OutClosestPoint, float& OutDistSquared, int32& OutPointIndex)
+{
+	OutClosestPoint = FVector3f::ZeroVector;
+	OutDistSquared = MAX_FLT;
+	OutPointIndex = INDEX_NONE;
+
+	for (int32 i = 0; i < BranchPointIndices.Num() - 1; ++i)
+	{
+		const int32 PointIndex0 = BranchPointIndices[i];
+		const int32 PointIndex1 = BranchPointIndices[i + 1];
+		if (!Points.IsValidIndex(PointIndex0) || !Points.IsValidIndex(PointIndex1))
+		{
+			continue;
+		}
+
+		const FVector3f& BranchPoint1 = Points[PointIndex0];
+		const FVector3f& BranchPoint2 = Points[PointIndex1];
+
+		const FVector3f ClosestPointOnLine = FMath::ClosestPointOnSegment(Point, BranchPoint1, BranchPoint2);
+		const float DistToLineSquared = (Point - ClosestPointOnLine).SizeSquared();
+		if (DistToLineSquared < OutDistSquared)
+		{
+			OutDistSquared = DistToLineSquared;
+			OutClosestPoint = ClosestPointOnLine;
+			OutPointIndex = PointIndex0;
+		}
+	}
+}
+
+float PV::Import::Helper::DistToBranch(const FVector3f& Point, const TArray<FVector3f>& Points, const TArray<int32>& BranchPointIndices)
+{
+	FVector3f ClosestPoint;
+	float ClosestDistSquared = MAX_FLT;
+	int32 PointIndex;
+	DistToBranch(Point, Points, BranchPointIndices, ClosestPoint, ClosestDistSquared, PointIndex);
+
+	return ClosestDistSquared > 0 ? FMath::Sqrt(ClosestDistSquared) : 0;
+}
+
+TArray<int32> PV::Import::Helper::EstimateBranchHierarchy(const TArray<FVector3f>& Points, const TArray<TArray<int32>>& Branches, int32 RootBranchIndex)
+{
+	TRACE_CPUPROFILER_EVENT_SCOPE(PV::Import::Helper::EstimateBranchHierarchy);
+
+	TArray<int32> OutBranchParents;
+	OutBranchParents.SetNum(Branches.Num());
+
+	for (const TArray<int32>& Branch : Branches)
+	{
+		for (int32 PointIndex : Branch)
+		{
+			if (!Points.IsValidIndex(PointIndex))
+			{
+				return TArray<int32>();
+			}
+		}
+	}
+
+	for (int32 i = 0; i < Branches.Num(); ++i)
+	{
+		if (i == RootBranchIndex)
+		{
+			OutBranchParents[i] = INDEX_NONE;
+			continue;
+		}
+
+		const int32 BranchRootPointIndex = Branches[i][0];
+		const FVector3f& BranchRootPoint = Points[BranchRootPointIndex];
+		float ClosestBranchDistSquared = MAX_FLT;
+		int32 ClosestBranchIndex = INDEX_NONE;
+		bool ClosestBranchPointIsRoot = true;
+
+		for (int32 j = 0; j < Branches.Num(); ++j)
+		{
+			if (i == j)
+			{
+				continue;
+			}
+
+			FVector3f ClosestPointOnBranch;
+			float BranchDistSquared = MAX_FLT;
+			int32 ClosestPointIndex;
+			Helper::DistToBranch(BranchRootPoint, Points, Branches[j], ClosestPointOnBranch, BranchDistSquared, ClosestPointIndex);
+
+			const bool bBranchPointIsRoot = ClosestPointIndex == BranchRootPointIndex;
+
+			// If there are multiple branches growing from the same point then we may select the wrong point as the distance to the branch will be the same.
+			// We therefore prefer the branch where the closest point is not root.
+			if ((FMath::IsNearlyEqual(BranchDistSquared, ClosestBranchDistSquared) && ClosestBranchPointIsRoot && !bBranchPointIsRoot)
+				|| BranchDistSquared < ClosestBranchDistSquared)
+			{
+				ClosestBranchIndex = j;
+				ClosestBranchDistSquared = BranchDistSquared;
+				ClosestBranchPointIsRoot = bBranchPointIsRoot;
+			}
+		}
+
+		OutBranchParents[i] = ClosestBranchIndex;
+	}
+
+	return OutBranchParents;
+}
+
+int32 PV::Import::Helper::FindRootBranch(const TArray<FVector3f>& InPoints, const TArray<TArray<int32>>& BranchIndices)
+{
+	float LowestRootPoint = MAX_FLT;
+	int32 LowestRootPointIdx = INDEX_NONE;
+	for (int32 i = 0; i < BranchIndices.Num(); ++i)
+	{
+		const TArray<int32>& Branch = BranchIndices[i];
+		if (InPoints[Branch[0]].Z < LowestRootPoint)
+		{
+			LowestRootPoint = InPoints[Branch[0]].Z;
+			LowestRootPointIdx = i;
+		}
+	}
+	return LowestRootPointIdx;
+}
+
+FPVBranchHierarchyDescription PV::Import::Helper::CreateBranchHierarchyFromPoints(const TArray<FVector3f>& InPoints, const TArray<float>& InPointsRadii, const TArray<TArray<int32>>& InBranchPointIndices)
+{
+	TRACE_CPUPROFILER_EVENT_SCOPE(PV::Import::Helper::CreateBranchHierarchyFromPoints);
+
+	if (InBranchPointIndices.Num() == 0 || InPoints.Num() == 0)
+	{
+		return FPVBranchHierarchyDescription();
+	}
+	
+	if (InPointsRadii.Num() != InPoints.Num())
+	{
+		return FPVBranchHierarchyDescription();
+	}
+
+	if (InBranchPointIndices.Num() == 1)
+	{
+		FPVBranchHierarchyDescription OutBranchHierarchy;
+		OutBranchHierarchy.RootBranchIndex = 0;
+		OutBranchHierarchy.Points = InPoints;
+		OutBranchHierarchy.PointsRadii = InPointsRadii;
+		FPVBranchDescription& Branch = OutBranchHierarchy.Branches.AddDefaulted_GetRef();
+		Branch.PointIndices = InBranchPointIndices[0];
+		Branch.ParentBranchIndex = INDEX_NONE;
+		Branch.BranchIndex = 0;
+		return OutBranchHierarchy;
+	}
+
+	const int32 RootBranchIndex = FindRootBranch(InPoints, InBranchPointIndices);
+
+	FPVBranchHierarchyDescription OutBranchHierarchy;
+	OutBranchHierarchy.RootBranchIndex = RootBranchIndex;
+	OutBranchHierarchy.Points = InPoints;
+	OutBranchHierarchy.PointsRadii = InPointsRadii;
+
+	// Fill branch points
+	for (const TArray<int32>& BranchPointIndices : InBranchPointIndices)
+	{
+		FPVBranchDescription& Branch = OutBranchHierarchy.Branches.AddDefaulted_GetRef();
+		Branch.PointIndices = BranchPointIndices;
+	}
+
+	// Set up branch hierarchy
+	const TArray<int32> BranchParents = EstimateBranchHierarchy(InPoints, InBranchPointIndices, RootBranchIndex);
+	for (int32 i = 0; i < BranchParents.Num(); ++i)
+	{
+		OutBranchHierarchy.Branches[i].BranchIndex = i;
+		OutBranchHierarchy.Branches[i].ParentBranchIndex = BranchParents[i];
+
+		if (BranchParents[i] != INDEX_NONE)
+		{
+			OutBranchHierarchy.Branches[BranchParents[i]].ChildBranchIndices.Add(i);
+		}
+	}
+
+	if (OutBranchHierarchy.ValidateHierarchy() != FPVBranchHierarchyDescription::EValidateHierarchyResult::Success)
+	{
+		return FPVBranchHierarchyDescription();
+	}
+
+	return OutBranchHierarchy;
+}
\ No newline at end of file
diff --git a/Engine/Plugins/Experimental/ProceduralVegetationEditor/Source/ProceduralVegetationEditor/Private/Importer/PVImportHelpers.h b/Engine/Plugins/Experimental/ProceduralVegetationEditor/Source/ProceduralVegetationEditor/Private/Importer/PVImportHelpers.h
new file mode 100644
--- /dev/null
+++ b/Engine/Plugins/Experimental/ProceduralVegetationEditor/Source/ProceduralVegetationEditor/Private/Importer/PVImportHelpers.h
@@ -0,0 +1,53 @@
+// Copyright Epic Games, Inc. All Rights Reserved.
+#pragma once
+
+#include "PVPlantDescription.h"
+
+namespace PV::Facades
+{
+	class FPointFacade;
+	class FBranchFacade;
+	class FPlantFacade;
+};
+
+struct FManagedArrayCollection;
+class UProceduralVegetationPreset;
+class UPackage;
+
+namespace PV::Import::Helper
+{
+	void CreateEmptyGrowthData(FManagedArrayCollection& OutCollection, int32 NumBranchPoints, int32 NumFoliagePoints, int32 NumBranches, int32 NumFoliageMeshes);
+
+	void FillGrowthDataFromBranchHierarchyDescription(FManagedArrayCollection& Collection, const FPVBranchHierarchyDescription& InBranchHierarchy);
+
+	void FillGrowthDataFromFoliageDescription(FManagedArrayCollection& Collection, const FPVFoliageDescription& InFoliageDescription, const FPVBranchHierarchyDescription& InBranchHierarchy);
+
+	void ComputeLengthFromRoot(PV::Facades::FPointFacade& PointFacade, const PV::Facades::FBranchFacade& BranchFacade, const PV::Facades::FPlantFacade& PlantFacade);
+
+	void ComputeBudDevelopment(PV::Facades::FPointFacade& PointFacade, const PV::Facades::FBranchFacade& BranchFacade, const PV::Facades::FPlantFacade& PlantFacade);
+
+	void ComputeBudDirections(PV::Facades::FPointFacade& PointFacade, const PV::Facades::FBranchFacade& BranchFacade, const PV::Facades::FPlantFacade& PlantFacade);
+
+	void InterpolatePointsScales(PV::Facades::FPointFacade& PointFacade, const PV::Facades::FBranchFacade& BranchFacade, const PV::Facades::FPlantFacade& PlantFacade);
+
+	bool SplitFusedPoints(FManagedArrayCollection& Collection);
+
+	void ConnectDisconnectedBranches(FPVBranchHierarchyDescription& InOutBranchHierarchy);
+
+	void FuseCollinearBranches(FPVBranchHierarchyDescription& InOutBranchHierarchy);
+
+	bool GenerateGrowthDataFromPlantDescription(FManagedArrayCollection& InOutCollection, const FPVPlantDescription& InPlantDescription);
+
+	bool GenerateProceduralVegetationPresetFromPlantDescription(const FPVPlantDescription& InPlantDescription, UProceduralVegetationPreset* InOutProceduralVegetationPreset);
+
+	void DistToBranch(const FVector3f& Point, const TArray<FVector3f>& Points, const TArray<int32>& BranchPointIndices, FVector3f& OutClosestPoint, float& OutDistSquared, int32& OutPointIndex);
+
+	float DistToBranch(const FVector3f& Point, const TArray<FVector3f>& Points, const TArray<int32>& BranchPointIndices);
+
+	// Returns array of parent branch index for each supplied branch. Root branch will have parent index INDEX_NONE.
+	TArray<int32> EstimateBranchHierarchy(const TArray<FVector3f>& Points, const TArray<TArray<int32>>& Branches, int32 RootBranchIndex);
+
+	int32 FindRootBranch(const TArray<FVector3f>& InPoints, const TArray<TArray<int32>>& BranchIndices);
+
+	FPVBranchHierarchyDescription CreateBranchHierarchyFromPoints(const TArray<FVector3f>& InPoints, const TArray<float>& InPointsRadii, const TArray<TArray<int32>>& InBranchPointIndices);
+};
\ No newline at end of file
diff --git a/Engine/Plugins/Experimental/ProceduralVegetationEditor/Source/ProceduralVegetationEditor/Private/Importer/PVImportTests.cpp b/Engine/Plugins/Experimental/ProceduralVegetationEditor/Source/ProceduralVegetationEditor/Private/Importer/PVImportTests.cpp
new file mode 100644
--- /dev/null
+++ b/Engine/Plugins/Experimental/ProceduralVegetationEditor/Source/ProceduralVegetationEditor/Private/Importer/PVImportTests.cpp
@@ -0,0 +1,697 @@
+// Copyright Epic Games, Inc. All Rights Reserved.
+
+#if WITH_DEV_AUTOMATION_TESTS
+
+#include "CoreTypes.h"
+#include "Misc/AutomationTest.h"
+#include "HAL/PreprocessorHelpers.h"
+#include "GeometryCollection/ManagedArrayCollection.h"
+
+#include "PVMockTree.h"
+#include "PVImportHelpers.h"
+#include "PVImporter_MockTree.h"
+
+#include "Facades/PVAttributesNames.h"
+#include "Facades/PVPointFacade.h"
+#include "Facades/PVBranchFacade.h"
+#include "Facades/PVPlantFacade.h"
+
+DEFINE_LOG_CATEGORY_STATIC(LogPVTest, Log, All);
+
+#define PV_IMPORTER_TEST_PRIVATE(TClass, PrettyName, TFlags) \
+	IMPLEMENT_SIMPLE_AUTOMATION_TEST(TClass, PrettyName, TFlags) \
+	bool TClass::RunTest(const FString& Parameters)
+
+#define PV_IMPORTER_TEST(TestName) \
+	PV_IMPORTER_TEST_PRIVATE(UE_JOIN(FPVImporter_,TestName), "PVE.Importer." UE_STRINGIZE(TestName), EAutomationTestFlags::EditorContext | EAutomationTestFlags::EngineFilter)
+
+PV_IMPORTER_TEST(FillGrowthDataFromBranchHierarchyDescription)
+{
+	const TSharedRef<FPVBranchHierarchyDescription> BranchHierarchy = PV::MockTree::CreateBranchHierarchyDescription();
+
+	const int32 NumBranchPoints = BranchHierarchy->GetNumBranchPoints();
+	const int32 NumFoliagePoints = 0;
+	const int32 NumBranches = BranchHierarchy->GetNumBranches();
+	const int32 NumFoliageMeshes = 0;
+
+	FManagedArrayCollection ExpectedCollection;
+	PV::Import::Helper::CreateEmptyGrowthData(ExpectedCollection, NumBranchPoints, NumFoliagePoints, NumBranches, NumFoliageMeshes);
+
+	PV::Facades::FPointFacade ExpectedPointFacade(ExpectedCollection);
+	PV::Facades::FBranchFacade ExpectedBranchFacade(ExpectedCollection);
+	
+	// Init ExpectedPointFacade
+	{
+		TManagedArray<FVector3f>& ExpectedPositions = ExpectedPointFacade.ModifyPositions();
+
+		ExpectedPositions[0] = PV::MockTree::Point0;
+		ExpectedPositions[1] = PV::MockTree::Point1;
+		ExpectedPositions[2] = PV::MockTree::Point2;
+		ExpectedPositions[3] = PV::MockTree::Point3;
+		ExpectedPositions[4] = PV::MockTree::Point4;
+		ExpectedPositions[5] = PV::MockTree::Point5;
+		ExpectedPositions[6] = PV::MockTree::Point6;
+		ExpectedPositions[7] = PV::MockTree::Point7;
+		ExpectedPositions[8] = PV::MockTree::Point8;
+		ExpectedPositions[9] = PV::MockTree::Point9;
+		ExpectedPositions[10] = PV::MockTree::Point10;
+		ExpectedPositions[11] = PV::MockTree::Point11;
+	}
+
+	// Init ExpectedBranchFacade
+	{
+		ExpectedBranchFacade.SetPoints(0, PV::MockTree::Branch0_PointIndices);
+		ExpectedBranchFacade.SetPoints(1, PV::MockTree::Branch1_PointIndices);
+		ExpectedBranchFacade.SetPoints(2, PV::MockTree::Branch2_PointIndices);
+		ExpectedBranchFacade.SetPoints(3, PV::MockTree::Branch3_PointIndices);
+
+		ExpectedBranchFacade.SetParentBranchNumber(0, 0);
+		ExpectedBranchFacade.SetParentBranchNumber(1, PV::MockTree::Branch0_BranchNumber);
+		ExpectedBranchFacade.SetParentBranchNumber(2, PV::MockTree::Branch1_BranchNumber);
+		ExpectedBranchFacade.SetParentBranchNumber(3, PV::MockTree::Branch0_BranchNumber);
+
+		ExpectedBranchFacade.SetParents(0, { });
+		ExpectedBranchFacade.SetParents(1, { PV::MockTree::Branch0_BranchNumber });
+		ExpectedBranchFacade.SetParents(2, { PV::MockTree::Branch0_BranchNumber, PV::MockTree::Branch1_BranchNumber });
+		ExpectedBranchFacade.SetParents(3, { PV::MockTree::Branch0_BranchNumber });
+
+		ExpectedBranchFacade.SetChildren(0, { PV::MockTree::Branch1_BranchNumber, PV::MockTree::Branch3_BranchNumber });
+		ExpectedBranchFacade.SetChildren(1, { PV::MockTree::Branch2_BranchNumber });
+		ExpectedBranchFacade.SetChildren(2, { });
+		ExpectedBranchFacade.SetChildren(3, { });
+	}
+
+	FManagedArrayCollection ActualCollection;
+	PV::Import::Helper::CreateEmptyGrowthData(ActualCollection, NumBranchPoints, NumFoliagePoints, NumBranches, NumFoliageMeshes);
+
+	UTEST_EQUAL("CorrectNrOfPointsInGroup", ActualCollection.NumElements(PV::GroupNames::PointGroup), NumBranchPoints);
+	UTEST_EQUAL("CorrectNrOfBranchesInGroup", ActualCollection.NumElements(PV::GroupNames::BranchGroup), NumBranches);
+
+	PV::Import::Helper::FillGrowthDataFromBranchHierarchyDescription(ActualCollection, *BranchHierarchy);
+
+	const PV::Facades::FPointFacade ActualPointFacade(ActualCollection);
+	const PV::Facades::FBranchFacade ActualBranchFacade(ActualCollection);
+
+	// Validate that the points are correctly indexed and in the correct location
+	UTEST_TRUE("PointPosition0", ActualPointFacade.GetPosition(0).Equals(ExpectedPointFacade.GetPosition(0)));
+	UTEST_TRUE("PointPosition1", ActualPointFacade.GetPosition(1).Equals(ExpectedPointFacade.GetPosition(1)));
+	UTEST_TRUE("PointPosition2", ActualPointFacade.GetPosition(2).Equals(ExpectedPointFacade.GetPosition(2)));
+	UTEST_TRUE("PointPosition3", ActualPointFacade.GetPosition(3).Equals(ExpectedPointFacade.GetPosition(3)));
+	UTEST_TRUE("PointPosition4", ActualPointFacade.GetPosition(4).Equals(ExpectedPointFacade.GetPosition(4)));
+	UTEST_TRUE("PointPosition5", ActualPointFacade.GetPosition(5).Equals(ExpectedPointFacade.GetPosition(5)));
+	UTEST_TRUE("PointPosition6", ActualPointFacade.GetPosition(6).Equals(ExpectedPointFacade.GetPosition(6)));
+	UTEST_TRUE("PointPosition7", ActualPointFacade.GetPosition(7).Equals(ExpectedPointFacade.GetPosition(7)));
+	UTEST_TRUE("PointPosition8", ActualPointFacade.GetPosition(8).Equals(ExpectedPointFacade.GetPosition(8)));
+	UTEST_TRUE("PointPosition9", ActualPointFacade.GetPosition(9).Equals(ExpectedPointFacade.GetPosition(9)));
+	UTEST_TRUE("PointPosition10", ActualPointFacade.GetPosition(10).Equals(ExpectedPointFacade.GetPosition(10)));
+	UTEST_TRUE("PointPosition11", ActualPointFacade.GetPosition(11).Equals(ExpectedPointFacade.GetPosition(11)));
+
+	// Validate that the branch point indices are correct
+	UTEST_EQUAL("BranchPoints0", ActualBranchFacade.GetPoints(0), ExpectedBranchFacade.GetPoints(0));
+	UTEST_EQUAL("BranchPoints1", ActualBranchFacade.GetPoints(1), ExpectedBranchFacade.GetPoints(1));
+	UTEST_EQUAL("BranchPoints2", ActualBranchFacade.GetPoints(2), ExpectedBranchFacade.GetPoints(2));
+	UTEST_EQUAL("BranchPoints3", ActualBranchFacade.GetPoints(3), ExpectedBranchFacade.GetPoints(3));
+
+	// Validate that the branch parents are correct
+	UTEST_EQUAL("ParentBranchNumber0", ActualBranchFacade.GetParentBranchNumber(0), ExpectedBranchFacade.GetParentBranchNumber(0));
+	UTEST_EQUAL("ParentBranchNumber1", ActualBranchFacade.GetParentBranchNumber(1), ExpectedBranchFacade.GetParentBranchNumber(1));
+	UTEST_EQUAL("ParentBranchNumber2", ActualBranchFacade.GetParentBranchNumber(2), ExpectedBranchFacade.GetParentBranchNumber(2));
+	UTEST_EQUAL("ParentBranchNumber3", ActualBranchFacade.GetParentBranchNumber(3), ExpectedBranchFacade.GetParentBranchNumber(3));
+	UTEST_EQUAL("BranchParents0", ActualBranchFacade.GetParents(0), ExpectedBranchFacade.GetParents(0));
+	UTEST_EQUAL("BranchParents1", ActualBranchFacade.GetParents(1), ExpectedBranchFacade.GetParents(1));
+	UTEST_EQUAL("BranchParents2", ActualBranchFacade.GetParents(2), ExpectedBranchFacade.GetParents(2));
+	UTEST_EQUAL("BranchParents3", ActualBranchFacade.GetParents(3), ExpectedBranchFacade.GetParents(3));
+
+	return true;
+}
+
+PV_IMPORTER_TEST(ComputeLengthFromRoot)
+{
+	const TSharedRef<FPVBranchHierarchyDescription> BranchHierarchy = PV::MockTree::CreateBranchHierarchyDescription();
+
+	const int32 NumBranchPoints = BranchHierarchy->GetNumBranchPoints();
+	const int32 NumFoliagePoints = 0;
+	const int32 NumBranches = BranchHierarchy->GetNumBranches();
+	const int32 NumFoliageMeshes = 0;
+
+	FManagedArrayCollection Collection;
+	PV::Import::Helper::CreateEmptyGrowthData(Collection, NumBranchPoints, NumFoliagePoints, NumBranches, NumFoliageMeshes);
+	PV::Import::Helper::FillGrowthDataFromBranchHierarchyDescription(Collection, *BranchHierarchy);
+
+	PV::Facades::FPointFacade PointFacade(Collection);
+	const PV::Facades::FBranchFacade BranchFacade(Collection);
+	const PV::Facades::FPlantFacade PlantFacade(Collection);
+	PV::Import::Helper::ComputeLengthFromRoot(PointFacade, BranchFacade, PlantFacade);
+
+	UTEST_NEARLY_EQUAL("Point0_LengthFromRoot", PointFacade.GetLengthFromRoot(0), PV::MockTree::Point0_LengthFromRoot, 0.01f);
+	UTEST_NEARLY_EQUAL("Point1_LengthFromRoot", PointFacade.GetLengthFromRoot(1), PV::MockTree::Point1_LengthFromRoot, 0.01f);
+	UTEST_NEARLY_EQUAL("Point2_LengthFromRoot", PointFacade.GetLengthFromRoot(2), PV::MockTree::Point2_LengthFromRoot, 0.01f);
+	UTEST_NEARLY_EQUAL("Point3_LengthFromRoot", PointFacade.GetLengthFromRoot(3), PV::MockTree::Point3_LengthFromRoot, 0.01f);
+	UTEST_NEARLY_EQUAL("Point4_LengthFromRoot", PointFacade.GetLengthFromRoot(4), PV::MockTree::Point4_LengthFromRoot, 0.01f);
+	UTEST_NEARLY_EQUAL("Point5_LengthFromRoot", PointFacade.GetLengthFromRoot(5), PV::MockTree::Point5_LengthFromRoot, 0.01f);
+	UTEST_NEARLY_EQUAL("Point6_LengthFromRoot", PointFacade.GetLengthFromRoot(6), PV::MockTree::Point6_LengthFromRoot, 0.01f);
+	UTEST_NEARLY_EQUAL("Point7_LengthFromRoot", PointFacade.GetLengthFromRoot(7), PV::MockTree::Point7_LengthFromRoot, 0.01f);
+	UTEST_NEARLY_EQUAL("Point8_LengthFromRoot", PointFacade.GetLengthFromRoot(8), PV::MockTree::Point8_LengthFromRoot, 0.01f);
+	UTEST_NEARLY_EQUAL("Point9_LengthFromRoot", PointFacade.GetLengthFromRoot(9), PV::MockTree::Point9_LengthFromRoot, 0.01f);
+	UTEST_NEARLY_EQUAL("Point10_LengthFromRoot", PointFacade.GetLengthFromRoot(10), PV::MockTree::Point10_LengthFromRoot, 0.01f);
+	UTEST_NEARLY_EQUAL("Point11_LengthFromRoot", PointFacade.GetLengthFromRoot(11), PV::MockTree::Point11_LengthFromRoot, 0.01f);
+
+	return true;
+}
+
+PV_IMPORTER_TEST(ComputeBudDevelopment)
+{
+	struct FBudDevelopmentDataLayout
+	{
+		int32 Generatation;
+		int32 Age;
+		int32 Branch_Age;
+		int32 AgeSenescence;
+		int32 Light_Senescence;
+		int32 Relative_Age;
+	};
+
+	const FBudDevelopmentDataLayout Point0_ExpectedBudDevelopment = {
+		.Generatation = 1, 
+		.Age = 1,
+		.Branch_Age = 1,
+		.AgeSenescence = 0,
+		.Light_Senescence = 0,
+		.Relative_Age = 1
+	};
+	const FBudDevelopmentDataLayout Point1_ExpectedBudDevelopment = {
+		.Generatation = 1,
+		.Age = 2,
+		.Branch_Age = 1,
+		.AgeSenescence = 0,
+		.Light_Senescence = 0,
+		.Relative_Age = 2
+	};
+	const FBudDevelopmentDataLayout Point2_ExpectedBudDevelopment = {
+		.Generatation = 1,
+		.Age = 3,
+		.Branch_Age = 1,
+		.AgeSenescence = 0,
+		.Light_Senescence = 0,
+		.Relative_Age = 3
+	};
+	const FBudDevelopmentDataLayout Point3_ExpectedBudDevelopment = {
+		.Generatation = 1,
+		.Age = 4,
+		.Branch_Age = 1,
+		.AgeSenescence = 0,
+		.Light_Senescence = 0,
+		.Relative_Age = 4
+	};
+	const FBudDevelopmentDataLayout Point4_ExpectedBudDevelopment = {
+		.Generatation = 1,
+		.Age = 5,
+		.Branch_Age = 1,
+		.AgeSenescence = 0,
+		.Light_Senescence = 0,
+		.Relative_Age = 5
+	};
+	const FBudDevelopmentDataLayout Point5_ExpectedBudDevelopment = {
+		.Generatation = 1,
+		.Age = 6,
+		.Branch_Age = 1,
+		.AgeSenescence = 0,
+		.Light_Senescence = 0,
+		.Relative_Age = 6
+	};
+	const FBudDevelopmentDataLayout Point6_ExpectedBudDevelopment = {
+		.Generatation = 2,
+		.Age = 5,
+		.Branch_Age = 4,
+		.AgeSenescence = 0,
+		.Light_Senescence = 0,
+		.Relative_Age = 5
+	};
+	const FBudDevelopmentDataLayout Point7_ExpectedBudDevelopment = {
+		.Generatation = 2,
+		.Age = 4,
+		.Branch_Age = 3,
+		.AgeSenescence = 0,
+		.Light_Senescence = 0,
+		.Relative_Age = 4
+	};
+	const FBudDevelopmentDataLayout Point8_ExpectedBudDevelopment = {
+		.Generatation = 2,
+		.Age = 5,
+		.Branch_Age = 3,
+		.AgeSenescence = 0,
+		.Light_Senescence = 0,
+		.Relative_Age = 5
+	};
+	const FBudDevelopmentDataLayout Point9_ExpectedBudDevelopment = {
+		.Generatation = 2,
+		.Age = 6,
+		.Branch_Age = 3,
+		.AgeSenescence = 0,
+		.Light_Senescence = 0,
+		.Relative_Age = 6
+	};
+	const FBudDevelopmentDataLayout Point10_ExpectedBudDevelopment = {
+		.Generatation = 3,
+		.Age = 6,
+		.Branch_Age = 5,
+		.AgeSenescence = 0,
+		.Light_Senescence = 0,
+		.Relative_Age = 6
+	};
+	const FBudDevelopmentDataLayout Point11_ExpectedBudDevelopment = {
+		.Generatation = 3,
+		.Age = 7,
+		.Branch_Age = 5,
+		.AgeSenescence = 0,
+		.Light_Senescence = 0,
+		.Relative_Age = 7
+	};
+
+	const int32 ExpectedNumBudDevelopmentValues = 6;
+
+	const TSharedRef<FPVBranchHierarchyDescription> BranchHierarchy = PV::MockTree::CreateBranchHierarchyDescription();
+	const int32 NumBranchPoints = BranchHierarchy->GetNumBranchPoints();
+	const int32 NumFoliagePoints = 0;
+	const int32 NumBranches = BranchHierarchy->GetNumBranches();
+	const int32 NumFoliageMeshes = 0;
+
+	FManagedArrayCollection Collection;
+	PV::Import::Helper::CreateEmptyGrowthData(Collection, NumBranchPoints, NumFoliagePoints, NumBranches, NumFoliageMeshes);
+	PV::Import::Helper::FillGrowthDataFromBranchHierarchyDescription(Collection, *BranchHierarchy);
+
+	PV::Facades::FPointFacade PointFacade(Collection);
+	const PV::Facades::FBranchFacade BranchFacade(Collection);
+	const PV::Facades::FPlantFacade PlantFacade(Collection);
+	PV::Import::Helper::ComputeBudDevelopment(PointFacade, BranchFacade, PlantFacade);
+
+	UTEST_TRUE("Point0_BudDevelopment_Num", PointFacade.GetBudDevelopment(0).Num() == ExpectedNumBudDevelopmentValues);
+	UTEST_TRUE("Point1_BudDevelopment_Num", PointFacade.GetBudDevelopment(1).Num() == ExpectedNumBudDevelopmentValues);
+	UTEST_TRUE("Point2_BudDevelopment_Num", PointFacade.GetBudDevelopment(2).Num() == ExpectedNumBudDevelopmentValues);
+	UTEST_TRUE("Point3_BudDevelopment_Num", PointFacade.GetBudDevelopment(3).Num() == ExpectedNumBudDevelopmentValues);
+	UTEST_TRUE("Point4_BudDevelopment_Num", PointFacade.GetBudDevelopment(4).Num() == ExpectedNumBudDevelopmentValues);
+	UTEST_TRUE("Point5_BudDevelopment_Num", PointFacade.GetBudDevelopment(5).Num() == ExpectedNumBudDevelopmentValues);
+	UTEST_TRUE("Point6_BudDevelopment_Num", PointFacade.GetBudDevelopment(6).Num() == ExpectedNumBudDevelopmentValues);
+	UTEST_TRUE("Point7_BudDevelopment_Num", PointFacade.GetBudDevelopment(7).Num() == ExpectedNumBudDevelopmentValues);
+	UTEST_TRUE("Point8_BudDevelopment_Num", PointFacade.GetBudDevelopment(8).Num() == ExpectedNumBudDevelopmentValues);
+	UTEST_TRUE("Point9_BudDevelopment_Num", PointFacade.GetBudDevelopment(9).Num() == ExpectedNumBudDevelopmentValues);
+	UTEST_TRUE("Point10_BudDevelopment_Num", PointFacade.GetBudDevelopment(10).Num() == ExpectedNumBudDevelopmentValues);
+	UTEST_TRUE("Point11_BudDevelopment_Num", PointFacade.GetBudDevelopment(11).Num() == ExpectedNumBudDevelopmentValues);
+
+	UTEST_TRUE(
+		"Point0_BudDevelopment", 
+		FMemory::Memcmp((void*)PointFacade.GetBudDevelopment(0).GetData(), (void*)&Point0_ExpectedBudDevelopment, ExpectedNumBudDevelopmentValues * sizeof(int32)) == 0
+	);
+	UTEST_TRUE(
+		"Point1_BudDevelopment",
+		FMemory::Memcmp((void*)PointFacade.GetBudDevelopment(1).GetData(), (void*)&Point1_ExpectedBudDevelopment, ExpectedNumBudDevelopmentValues * sizeof(int32)) == 0
+	);
+	UTEST_TRUE(
+		"Point2_BudDevelopment",
+		FMemory::Memcmp((void*)PointFacade.GetBudDevelopment(2).GetData(), (void*)&Point2_ExpectedBudDevelopment, ExpectedNumBudDevelopmentValues * sizeof(int32)) == 0
+	);
+	UTEST_TRUE(
+		"Point3_BudDevelopment",
+		FMemory::Memcmp((void*)PointFacade.GetBudDevelopment(3).GetData(), (void*)&Point3_ExpectedBudDevelopment, ExpectedNumBudDevelopmentValues * sizeof(int32)) == 0
+	);
+	UTEST_TRUE(
+		"Point4_BudDevelopment",
+		FMemory::Memcmp((void*)PointFacade.GetBudDevelopment(4).GetData(), (void*)&Point4_ExpectedBudDevelopment, ExpectedNumBudDevelopmentValues * sizeof(int32)) == 0
+	);
+	UTEST_TRUE(
+		"Point5_BudDevelopment",
+		FMemory::Memcmp((void*)PointFacade.GetBudDevelopment(5).GetData(), (void*)&Point5_ExpectedBudDevelopment, ExpectedNumBudDevelopmentValues * sizeof(int32)) == 0
+	);
+	UTEST_TRUE(
+		"Point6_BudDevelopment",
+		FMemory::Memcmp((void*)PointFacade.GetBudDevelopment(6).GetData(), (void*)&Point6_ExpectedBudDevelopment, ExpectedNumBudDevelopmentValues * sizeof(int32)) == 0
+	);
+	UTEST_TRUE(
+		"Point7_BudDevelopment",
+		FMemory::Memcmp((void*)PointFacade.GetBudDevelopment(7).GetData(), (void*)&Point7_ExpectedBudDevelopment, ExpectedNumBudDevelopmentValues * sizeof(int32)) == 0
+	);
+	UTEST_TRUE(
+		"Point8_BudDevelopment",
+		FMemory::Memcmp((void*)PointFacade.GetBudDevelopment(8).GetData(), (void*)&Point8_ExpectedBudDevelopment, ExpectedNumBudDevelopmentValues * sizeof(int32)) == 0
+	);
+	UTEST_TRUE(
+		"Point9_BudDevelopment",
+		FMemory::Memcmp((void*)PointFacade.GetBudDevelopment(9).GetData(), (void*)&Point9_ExpectedBudDevelopment, ExpectedNumBudDevelopmentValues * sizeof(int32)) == 0
+	);
+	UTEST_TRUE(
+		"Point10_BudDevelopment",
+		FMemory::Memcmp((void*)PointFacade.GetBudDevelopment(10).GetData(), (void*)&Point10_ExpectedBudDevelopment, ExpectedNumBudDevelopmentValues * sizeof(int32)) == 0
+	);
+
+	UTEST_TRUE(
+		"Point11_BudDevelopment",
+		FMemory::Memcmp((void*)PointFacade.GetBudDevelopment(11).GetData(), (void*)&Point11_ExpectedBudDevelopment, ExpectedNumBudDevelopmentValues * sizeof(int32)) == 0
+	);
+
+	return true;
+}
+
+PV_IMPORTER_TEST(ComputeBudDirections)
+{
+	const FVector3f ExpectedApical0 = (PV::MockTree::Point1 - PV::MockTree::Point0).GetSafeNormal();
+	const FVector3f ExpectedApical1 = (PV::MockTree::Point2 - PV::MockTree::Point1).GetSafeNormal();
+	const FVector3f ExpectedApical2 = (PV::MockTree::Point3 - PV::MockTree::Point2).GetSafeNormal();
+	const FVector3f ExpectedApical3 = (PV::MockTree::Point4 - PV::MockTree::Point3).GetSafeNormal();
+	const FVector3f ExpectedApical4 = (PV::MockTree::Point5 - PV::MockTree::Point4).GetSafeNormal();
+	const FVector3f ExpectedApical5 = (PV::MockTree::Point5 - PV::MockTree::Point4).GetSafeNormal();
+	const FVector3f ExpectedApical6 = (PV::MockTree::Point6 - PV::MockTree::Point3).GetSafeNormal();
+	const FVector3f ExpectedApical7 = (PV::MockTree::Point8 - PV::MockTree::Point7).GetSafeNormal();
+	const FVector3f ExpectedApical8 = (PV::MockTree::Point9 - PV::MockTree::Point8).GetSafeNormal();
+	const FVector3f ExpectedApical9 = (PV::MockTree::Point9 - PV::MockTree::Point8).GetSafeNormal();
+	const FVector3f ExpectedApical10 = (PV::MockTree::Point11 - PV::MockTree::Point10).GetSafeNormal();
+	const FVector3f ExpectedApical11 = (PV::MockTree::Point11 - PV::MockTree::Point10).GetSafeNormal();
+
+	const int32 ExpectedNumBudDirectionValues = 6;
+
+	const TSharedRef<FPVBranchHierarchyDescription> BranchHierarchy = PV::MockTree::CreateBranchHierarchyDescription();
+	const int32 NumBranchPoints = BranchHierarchy->GetNumBranchPoints();
+	const int32 NumFoliagePoints = 0;
+	const int32 NumBranches = BranchHierarchy->GetNumBranches();
+	const int32 NumFoliageMeshes = 0;
+
+	FManagedArrayCollection Collection;
+	PV::Import::Helper::CreateEmptyGrowthData(Collection, NumBranchPoints, NumFoliagePoints, NumBranches, NumFoliageMeshes);
+	PV::Import::Helper::FillGrowthDataFromBranchHierarchyDescription(Collection, *BranchHierarchy);
+
+	PV::Facades::FPointFacade PointFacade(Collection);
+	const PV::Facades::FBranchFacade BranchFacade(Collection);
+	const PV::Facades::FPlantFacade PlantFacade(Collection);
+	PV::Import::Helper::ComputeBudDirections(PointFacade, BranchFacade, PlantFacade);
+
+	UTEST_TRUE("Point0_BudDirections_Num", PointFacade.GetBudDirection(0).Num() == ExpectedNumBudDirectionValues);
+	UTEST_TRUE("Point1_BudDirections_Num", PointFacade.GetBudDirection(1).Num() == ExpectedNumBudDirectionValues);
+	UTEST_TRUE("Point2_BudDirections_Num", PointFacade.GetBudDirection(2).Num() == ExpectedNumBudDirectionValues);
+	UTEST_TRUE("Point3_BudDirections_Num", PointFacade.GetBudDirection(3).Num() == ExpectedNumBudDirectionValues);
+	UTEST_TRUE("Point4_BudDirections_Num", PointFacade.GetBudDirection(4).Num() == ExpectedNumBudDirectionValues);
+	UTEST_TRUE("Point5_BudDirections_Num", PointFacade.GetBudDirection(5).Num() == ExpectedNumBudDirectionValues);
+	UTEST_TRUE("Point6_BudDirections_Num", PointFacade.GetBudDirection(6).Num() == ExpectedNumBudDirectionValues);
+	UTEST_TRUE("Point7_BudDirections_Num", PointFacade.GetBudDirection(7).Num() == ExpectedNumBudDirectionValues);
+	UTEST_TRUE("Point8_BudDirections_Num", PointFacade.GetBudDirection(8).Num() == ExpectedNumBudDirectionValues);
+	UTEST_TRUE("Point9_BudDirections_Num", PointFacade.GetBudDirection(9).Num() == ExpectedNumBudDirectionValues);
+	UTEST_TRUE("Point10_BudDirections_Num", PointFacade.GetBudDirection(10).Num() == ExpectedNumBudDirectionValues);
+	UTEST_TRUE("Point11_BudDirections_Num", PointFacade.GetBudDirection(11).Num() == ExpectedNumBudDirectionValues);
+
+	UTEST_NEARLY_EQUAL("Point0_Apical", FVector(PointFacade.GetBudDirection(0)[0]), FVector(ExpectedApical0), UE_KINDA_SMALL_NUMBER);
+	UTEST_NEARLY_EQUAL("Point1_Apical", FVector(PointFacade.GetBudDirection(1)[0]), FVector(ExpectedApical1), UE_KINDA_SMALL_NUMBER);
+	UTEST_NEARLY_EQUAL("Point2_Apical", FVector(PointFacade.GetBudDirection(2)[0]), FVector(ExpectedApical2), UE_KINDA_SMALL_NUMBER);
+	UTEST_NEARLY_EQUAL("Point3_Apical", FVector(PointFacade.GetBudDirection(3)[0]), FVector(ExpectedApical3), UE_KINDA_SMALL_NUMBER);
+	UTEST_NEARLY_EQUAL("Point4_Apical", FVector(PointFacade.GetBudDirection(4)[0]), FVector(ExpectedApical4), UE_KINDA_SMALL_NUMBER);
+	UTEST_NEARLY_EQUAL("Point5_Apical", FVector(PointFacade.GetBudDirection(5)[0]), FVector(ExpectedApical5), UE_KINDA_SMALL_NUMBER);
+	UTEST_NEARLY_EQUAL("Point6_Apical", FVector(PointFacade.GetBudDirection(6)[0]), FVector(ExpectedApical6), UE_KINDA_SMALL_NUMBER);
+	UTEST_NEARLY_EQUAL("Point7_Apical", FVector(PointFacade.GetBudDirection(7)[0]), FVector(ExpectedApical7), UE_KINDA_SMALL_NUMBER);
+	UTEST_NEARLY_EQUAL("Point8_Apical", FVector(PointFacade.GetBudDirection(8)[0]), FVector(ExpectedApical8), UE_KINDA_SMALL_NUMBER);
+	UTEST_NEARLY_EQUAL("Point9_Apical", FVector(PointFacade.GetBudDirection(9)[0]), FVector(ExpectedApical9), UE_KINDA_SMALL_NUMBER);
+	UTEST_NEARLY_EQUAL("Point10_Apical", FVector(PointFacade.GetBudDirection(10)[0]), FVector(ExpectedApical10), UE_KINDA_SMALL_NUMBER);
+	UTEST_NEARLY_EQUAL("Point11_Apical", FVector(PointFacade.GetBudDirection(11)[0]), FVector(ExpectedApical11), UE_KINDA_SMALL_NUMBER);
+
+	// TODO: Test UpVector and Axillary
+
+	return true;
+}
+
+PV_IMPORTER_TEST(SplitFusedPoints)
+{
+	const TSharedRef<FPVBranchHierarchyDescription> BranchHierarchy = PV::MockTree::CreateBranchHierarchyDescription();
+	const int32 NumBranchPoints = BranchHierarchy->GetNumBranchPoints();
+	const int32 NumFoliagePoints = 0;
+	const int32 NumBranches = BranchHierarchy->GetNumBranches();
+	const int32 NumFoliageMeshes = 0;
+
+	FManagedArrayCollection Collection;
+	PV::Import::Helper::CreateEmptyGrowthData(Collection, NumBranchPoints, NumFoliagePoints, NumBranches, NumFoliageMeshes);
+	PV::Import::Helper::FillGrowthDataFromBranchHierarchyDescription(Collection, *BranchHierarchy);
+
+	PV::Facades::FPointFacade PointFacade(Collection);
+	const PV::Facades::FBranchFacade BranchFacade(Collection);
+	const PV::Facades::FPlantFacade PlantFacade(Collection);
+	PV::Import::Helper::ComputeBudDirections(PointFacade, BranchFacade, PlantFacade);
+	PV::Import::Helper::ComputeLengthFromRoot(PointFacade, BranchFacade, PlantFacade);
+	PV::Import::Helper::ComputeBudDevelopment(PointFacade, BranchFacade, PlantFacade);
+	PV::Import::Helper::ComputeBudDirections(PointFacade, BranchFacade, PlantFacade);
+
+	const bool bSplitFusedPointsResult = PV::Import::Helper::SplitFusedPoints(Collection);
+	UTEST_TRUE("SplitFusedPointsSucceeded", bSplitFusedPointsResult);
+
+	const int32 Branch0_OldRootPointIndex = 0;
+	const int32 Branch1_OldRootPointIndex = 2;
+	const int32 Branch2_OldRootPointIndex = 8;
+	const int32 Branch3_OldRootPointIndex = 3;
+
+	const int32 Branch0_NewRootPointIndex = BranchFacade.GetPoints(0)[0];
+	const int32 Branch1_NewRootPointIndex = BranchFacade.GetPoints(1)[0];
+	const int32 Branch2_NewRootPointIndex = BranchFacade.GetPoints(2)[0];
+	const int32 Branch3_NewRootPointIndex = BranchFacade.GetPoints(3)[0];
+
+	UTEST_TRUE("Branch0_HasSamePointIndex",  Branch0_NewRootPointIndex == Branch0_OldRootPointIndex);
+	UTEST_TRUE("Branch1_HasUniqePointIndex", Branch1_NewRootPointIndex != Branch1_OldRootPointIndex && Branch1_NewRootPointIndex > 11);
+	UTEST_TRUE("Branch2_HasUniqePointIndex", Branch2_NewRootPointIndex != Branch2_OldRootPointIndex && Branch2_NewRootPointIndex > 11);
+	UTEST_TRUE("Branch3_HasUniqePointIndex", Branch3_NewRootPointIndex != Branch3_OldRootPointIndex && Branch3_NewRootPointIndex > 11);
+
+	UTEST_EQUAL("Branch0_NewPointPositionIsSame", PointFacade.GetPosition(Branch0_NewRootPointIndex), PointFacade.GetPosition(Branch0_OldRootPointIndex));
+	UTEST_EQUAL("Branch1_NewPointPositionIsSame", PointFacade.GetPosition(Branch1_NewRootPointIndex), PointFacade.GetPosition(Branch1_OldRootPointIndex));
+	UTEST_EQUAL("Branch2_NewPointPositionIsSame", PointFacade.GetPosition(Branch2_NewRootPointIndex), PointFacade.GetPosition(Branch2_OldRootPointIndex));
+	UTEST_EQUAL("Branch3_NewPointPositionIsSame", PointFacade.GetPosition(Branch3_NewRootPointIndex), PointFacade.GetPosition(Branch3_OldRootPointIndex));
+
+	UTEST_EQUAL("Branch0_NewPointLengthFromRootIsSame", PointFacade.GetLengthFromRoot(Branch0_NewRootPointIndex), PointFacade.GetLengthFromRoot(Branch0_OldRootPointIndex));
+	UTEST_EQUAL("Branch1_NewPointLengthFromRootIsSame", PointFacade.GetLengthFromRoot(Branch1_NewRootPointIndex), PointFacade.GetLengthFromRoot(Branch1_OldRootPointIndex));
+	UTEST_EQUAL("Branch2_NewPointLengthFromRootIsSame", PointFacade.GetLengthFromRoot(Branch2_NewRootPointIndex), PointFacade.GetLengthFromRoot(Branch2_OldRootPointIndex));
+	UTEST_EQUAL("Branch3_NewPointLengthFromRootIsSame", PointFacade.GetLengthFromRoot(Branch3_NewRootPointIndex), PointFacade.GetLengthFromRoot(Branch3_OldRootPointIndex));
+
+	UTEST_EQUAL("Branch0_NewPointLengthFromSeedIsSame", PointFacade.GetLengthFromSeed(Branch0_NewRootPointIndex), PointFacade.GetLengthFromSeed(Branch0_OldRootPointIndex));
+	UTEST_EQUAL("Branch1_NewPointLengthFromSeedIsSame", PointFacade.GetLengthFromSeed(Branch1_NewRootPointIndex), PointFacade.GetLengthFromSeed(Branch1_OldRootPointIndex));
+	UTEST_EQUAL("Branch2_NewPointLengthFromSeedIsSame", PointFacade.GetLengthFromSeed(Branch2_NewRootPointIndex), PointFacade.GetLengthFromSeed(Branch2_OldRootPointIndex));
+	UTEST_EQUAL("Branch3_NewPointLengthFromSeedIsSame", PointFacade.GetLengthFromSeed(Branch3_NewRootPointIndex), PointFacade.GetLengthFromSeed(Branch3_OldRootPointIndex));
+
+	return true;
+}
+
+PV_IMPORTER_TEST(FindRootBranch)
+{
+	const TArray<TArray<int32>>& Branches = PV::MockTree::BranchPointIndices;
+	const int32 RootBranchIndex = PV::Import::Helper::FindRootBranch(PV::MockTree::Points, Branches);
+
+	UTEST_NOT_EQUAL("ValidRootBranch", RootBranchIndex, int32(INDEX_NONE));
+
+	const TArray<int32> RootBranchIndices = Branches[RootBranchIndex];
+	const FVector3f RootVertex = PV::MockTree::Points[RootBranchIndices[0]];
+
+	UTEST_NEARLY_EQUAL("CorrectRootBranchLocation", FVector(RootVertex), FVector::ZeroVector, UE_KINDA_SMALL_NUMBER);
+
+	return true;
+}
+
+PV_IMPORTER_TEST(EstimateBranchHierarchy)
+{
+	// Test standard MockTree
+	{
+		TArray<TArray<int32>> Branches = {
+			PV::MockTree::Branch0_PointIndices,
+			PV::MockTree::Branch1_PointIndices,
+			PV::MockTree::Branch2_PointIndices,
+			PV::MockTree::Branch3_PointIndices
+		};
+
+		TArray<int32> ExpectedBranchHierarchy = { -1, 0, 1, 0 };
+
+		// Naive testing of hierarchy
+		{
+			const int32 RootBranchIndex = PV::Import::Helper::FindRootBranch(PV::MockTree::Points, Branches);
+			const TArray<int32> BranchHierarchy = PV::Import::Helper::EstimateBranchHierarchy(PV::MockTree::Points, Branches, RootBranchIndex);
+
+			UTEST_EQUAL("CorrectNrOfBranches", BranchHierarchy.Num(), ExpectedBranchHierarchy.Num());
+			for (int32 i = 0; i < Branches.Num(); ++i)
+			{
+				UTEST_EQUAL(FString::Printf(TEXT("Branch%d_CorrectParent"), i), BranchHierarchy[i], ExpectedBranchHierarchy[i]);
+			}
+		}
+
+		// Test with branch list reversed
+		{
+			Algo::Reverse(Branches);
+			Algo::Reverse(ExpectedBranchHierarchy);
+			for (int32& Hierarchy : ExpectedBranchHierarchy)
+			{
+				if (Hierarchy != -1)
+				{
+					Hierarchy = Branches.Num() - 1 - Hierarchy;
+				}
+			}
+
+			const int32 RootBranchIndex = PV::Import::Helper::FindRootBranch(PV::MockTree::Points, Branches);
+			const TArray<int32> BranchHierarchy = PV::Import::Helper::EstimateBranchHierarchy(PV::MockTree::Points, Branches, RootBranchIndex);
+			
+			UTEST_EQUAL("CorrectNrOfBranches_Reverse", BranchHierarchy.Num(), ExpectedBranchHierarchy.Num());
+			for (int32 i = 0; i < Branches.Num(); ++i)
+			{
+				UTEST_EQUAL(FString::Printf(TEXT("Branch%d_CorrectParent_Reverse"), i), BranchHierarchy[i], ExpectedBranchHierarchy[i]);
+			}
+		}
+	}
+
+	// Add additional branches to MockTree and ensure hierarchy stays correct
+	{
+		// Add a few branches starting from the same point as another as we want to make sure the algorithm can handle multiple 
+		// branches being at the same distance from the start point of another.
+		TArray<FVector3f> PointsCpy = PV::MockTree::Points;
+		const int32 Point2_1_Index = PointsCpy.Add(PV::MockTree::Point2 + FVector3f(-10, 0, 0));
+		const int32 Point2_2_Index = PointsCpy.Add(PV::MockTree::Point2 + FVector3f(10, 0, 0));
+		const int32 Point2_3_Index = PointsCpy.Add(PV::MockTree::Point2 + FVector3f(-10, 10, 0));
+		const int32 Point2_4_Index = PointsCpy.Add(PV::MockTree::Point2 + FVector3f(10, 10, 0));
+		const int32 Point2_5_Index = PointsCpy.Add(PV::MockTree::Point2 + FVector3f(-10, 10, -10));
+		const int32 Point2_6_Index = PointsCpy.Add(PV::MockTree::Point2 + FVector3f(10, 8, 1));
+
+		const int32 Point8_1_Index = PointsCpy.Add(PV::MockTree::Point8 + FVector3f(10, 0, 0));
+		const int32 Point8_2_Index = PointsCpy.Add(PV::MockTree::Point8 + FVector3f(-10, 0, 0));
+		const int32 Point8_3_Index = PointsCpy.Add(PV::MockTree::Point8 + FVector3f(10, 10, 0));
+		const int32 Point8_4_Index = PointsCpy.Add(PV::MockTree::Point8 + FVector3f(-10, 10, 0));
+		const int32 Point8_5_Index = PointsCpy.Add(PV::MockTree::Point8 + FVector3f(10, 10, 10));
+		const int32 Point8_6_Index = PointsCpy.Add(PV::MockTree::Point8 + FVector3f(-10, 10, 10));
+
+		const int32 Point3_0_Index = PointsCpy.Add(PV::MockTree::Point3 + FVector3f(10, 0, 0));
+		const int32 Point3_1_Index = PointsCpy.Add(PV::MockTree::Point3 + FVector3f(-10, 0, 0));
+
+		TArray<TArray<int32>> Branches = {
+			PV::MockTree::Branch0_PointIndices,
+			PV::MockTree::Branch1_PointIndices,
+			PV::MockTree::Branch2_PointIndices,
+			PV::MockTree::Branch3_PointIndices,
+			{ 2, Point2_1_Index },
+			{ 2, Point2_2_Index },
+			{ 2, Point2_3_Index },
+			{ 2, Point2_4_Index },
+			{ 2, Point2_5_Index },
+			{ 2, Point2_6_Index },
+
+			{ 8, Point8_1_Index },
+			{ 8, Point8_2_Index },
+			{ 8, Point8_3_Index },
+			{ 8, Point8_4_Index },
+			{ 8, Point8_5_Index },
+			{ 8, Point8_6_Index },
+
+			{ 3, Point3_0_Index },
+			{ 3, Point3_1_Index },
+		};
+
+		TArray<int32> ExpectedBranchHierarchy = { 
+			-1, 0, 1, 0, 
+			0, 0, 0, 0, 0, 0,
+			1, 1, 1, 1, 1, 1,
+			0, 0
+		};
+
+		// Naive testing of hierarchy
+		{
+			const int32 RootBranchIndex = PV::Import::Helper::FindRootBranch(PointsCpy, Branches);
+			const TArray<int32> BranchHierarchy = PV::Import::Helper::EstimateBranchHierarchy(PointsCpy, Branches, RootBranchIndex);
+
+			UTEST_EQUAL("CorrectNrOfBranches", BranchHierarchy.Num(), ExpectedBranchHierarchy.Num());
+			for (int32 i = 0; i < Branches.Num(); ++i)
+			{
+				UTEST_EQUAL(FString::Printf(TEXT("Branch%d_CorrectParent"), i), BranchHierarchy[i], ExpectedBranchHierarchy[i]);
+			}
+		}
+
+		// Test with branch list reversed
+		{
+			Algo::Reverse(Branches);
+			Algo::Reverse(ExpectedBranchHierarchy);
+			for (int32& Hierarchy : ExpectedBranchHierarchy)
+			{
+				if (Hierarchy != -1)
+				{
+					Hierarchy = Branches.Num() - 1 - Hierarchy;
+				}
+			}
+
+			const int32 RootBranchIndex = PV::Import::Helper::FindRootBranch(PointsCpy, Branches);
+			const TArray<int32> BranchHierarchy = PV::Import::Helper::EstimateBranchHierarchy(PointsCpy, Branches, RootBranchIndex);
+
+			UTEST_EQUAL("CorrectNrOfBranches_Reverse", BranchHierarchy.Num(), ExpectedBranchHierarchy.Num());
+			for (int32 i = 0; i < Branches.Num(); ++i)
+			{
+				UTEST_EQUAL(FString::Printf(TEXT("Branch%d_CorrectParent_Reverse"), i), BranchHierarchy[i], ExpectedBranchHierarchy[i]);
+			}
+		}
+	}
+
+	return true;
+}
+
+PV_IMPORTER_TEST(CreateBranchHierarchyFromEdges)
+{
+	FPVBranchHierarchyDescription ExpectedBranchHierarchy;
+	ExpectedBranchHierarchy.RootBranchIndex = 0;
+	ExpectedBranchHierarchy.Points = PV::MockTree::Points;
+
+	FPVBranchDescription& B0 = ExpectedBranchHierarchy.Branches.AddDefaulted_GetRef();
+	B0.BranchIndex = 0;
+	B0.PointIndices = PV::MockTree::Branch0_PointIndices;
+	B0.ParentBranchIndex = -1;
+	B0.ChildBranchIndices = { 1, 3 };
+
+	FPVBranchDescription& B1 = ExpectedBranchHierarchy.Branches.AddDefaulted_GetRef();
+	B1.BranchIndex = 1;
+	B1.PointIndices = PV::MockTree::Branch1_PointIndices;
+	B1.ParentBranchIndex = 0;
+	B1.ChildBranchIndices = { 2 };
+
+	FPVBranchDescription& B2 = ExpectedBranchHierarchy.Branches.AddDefaulted_GetRef();
+	B2.BranchIndex = 2;
+	B2.PointIndices = PV::MockTree::Branch2_PointIndices;
+	B2.ParentBranchIndex = 1;
+	B2.ChildBranchIndices = { };
+
+	FPVBranchDescription& B3 = ExpectedBranchHierarchy.Branches.AddDefaulted_GetRef();
+	B3.BranchIndex = 3;
+	B3.PointIndices = PV::MockTree::Branch3_PointIndices;
+	B3.ParentBranchIndex = 0;
+	B3.ChildBranchIndices = { };
+
+	const TSharedRef<FPVBranchHierarchyDescription> BranchHierarchy = PV::MockTree::CreateBranchHierarchyDescription();
+
+	UTEST_EQUAL("CorrectRootIndex", BranchHierarchy->RootBranchIndex, ExpectedBranchHierarchy.RootBranchIndex);
+	UTEST_EQUAL("CorrectNrOfBranches", BranchHierarchy->Branches.Num(), ExpectedBranchHierarchy.Branches.Num());
+
+	for (int32 i = 0; i < BranchHierarchy->Branches.Num(); ++i)
+	{
+		const FPVBranchDescription& Branch = BranchHierarchy->Branches[i];
+		const FPVBranchDescription* ExpectedBranchPtr = ExpectedBranchHierarchy.Branches.FindByPredicate([&](const FPVBranchDescription& InBranch) { return InBranch.BranchIndex == Branch.BranchIndex; });
+
+		UTEST_NOT_NULL("ContainsBranchIndex", ExpectedBranchPtr);
+
+		const FPVBranchDescription& ExpectedBranch = *ExpectedBranchPtr;
+
+		const FString BranchName = FString::Printf(TEXT("Branch%d: "), i);
+
+		UTEST_EQUAL(BranchName + "_CorrectBranchIndex", Branch.BranchIndex, ExpectedBranch.BranchIndex);
+		UTEST_EQUAL(BranchName + "_CorrectParent", Branch.ParentBranchIndex, ExpectedBranch.ParentBranchIndex);
+		UTEST_EQUAL(BranchName + "_CorrectNrOfChildren", Branch.ChildBranchIndices.Num(), ExpectedBranch.ChildBranchIndices.Num());
+		UTEST_EQUAL(BranchName + "_CorrectNrOfPoints", Branch.PointIndices.Num(), ExpectedBranch.PointIndices.Num());
+
+		for (int32 j = 0; j < Branch.PointIndices.Num(); ++j)
+		{
+			const int32 PointIndex = Branch.PointIndices[j];
+			const int32 ExpectedPointIndex = ExpectedBranch.PointIndices[j];
+			UTEST_EQUAL(BranchName + FString::Printf(TEXT("_CorrectPointIndex %d"), j), PointIndex, ExpectedPointIndex);
+
+			const FVector Point = FVector(BranchHierarchy->Points[PointIndex]);
+			const FVector ExpectedPoint = FVector(ExpectedBranchHierarchy.Points[PointIndex]);
+			UTEST_NEARLY_EQUAL(BranchName + FString::Printf(TEXT("_CorrectPoint %d"), j), Point, ExpectedPoint, UE_KINDA_SMALL_NUMBER);
+		}
+
+		for (int32 j = 0; j < Branch.ChildBranchIndices.Num(); ++j)
+		{
+			UTEST_TRUE(BranchName + FString::Printf(TEXT("_CorrectChildIndex %d"), j), ExpectedBranch.ChildBranchIndices.Contains(Branch.ChildBranchIndices[j]));
+		}
+	}
+
+	return true;
+}
+
+#endif
\ No newline at end of file
diff --git a/Engine/Plugins/Experimental/ProceduralVegetationEditor/Source/ProceduralVegetationEditor/Private/Importer/PVImporter_MockTree.cpp b/Engine/Plugins/Experimental/ProceduralVegetationEditor/Source/ProceduralVegetationEditor/Private/Importer/PVImporter_MockTree.cpp
new file mode 100644
--- /dev/null
+++ b/Engine/Plugins/Experimental/ProceduralVegetationEditor/Source/ProceduralVegetationEditor/Private/Importer/PVImporter_MockTree.cpp
@@ -0,0 +1,94 @@
+// Copyright Epic Games, Inc. All Rights Reserved.
+
+#include "PVImporter_MockTree.h"
+#include "ProceduralVegetationEditorModule.h"
+#include "ProceduralVegetationPreset.h"
+
+#include "PVMockTree.h"
+#include "PVImportHelpers.h"
+
+#include "UObject/Package.h"
+#include "Misc/PackageName.h"
+#include "Helpers/PVUtilities.h"
+#include "AssetRegistry/IAssetRegistry.h"
+
+namespace PV::Import::MockTree::Internal
+{
+	static FAutoConsoleCommand ImportMockTreeCMD(
+		TEXT("PVE.Import.MockTree"),
+		TEXT("PVE.Import.MockTree {OUTPUT_PATH}"),
+		FConsoleCommandWithArgsDelegate::CreateLambda([](const TArray<FString>& Args) {
+			if (Args.Num() != 1)
+			{
+				UE_LOG(LogProceduralVegetationEditor, Error, TEXT("Invalid arguments"));
+				return;
+			}
+
+			const bool Result = ImportMockTreeAsProceduralVegetation(Args[0]);
+			if (Result != true)
+			{
+				UE_LOG(LogProceduralVegetationEditor, Error, TEXT("Import Failed"));
+			}
+			else
+			{
+				UE_LOG(LogProceduralVegetationEditor, Log, TEXT("Import Succeeded"));
+			}
+		}),
+		ECVF_Default
+	);
+}
+
+bool PV::Import::MockTree::ImportMockTreeAsProceduralVegetation(const FString& OutputPath)
+{
+	const FString ProceduralVegetationPresetPath = OutputPath;
+	const FString PlantDescriptionDataAssetPath = FString::Printf(TEXT("%s_PlantDescription"), *OutputPath);
+
+	if (!FPackageName::IsValidTextForLongPackageName(ProceduralVegetationPresetPath)
+		|| !FPackageName::IsValidTextForLongPackageName(PlantDescriptionDataAssetPath))
+	{
+		return false;
+	}
+
+	if (Utilities::DoesConflictingPackageExist(ProceduralVegetationPresetPath, UProceduralVegetationPreset::StaticClass())
+		|| Utilities::DoesConflictingPackageExist(PlantDescriptionDataAssetPath, UPVPlantDescriptionDataAsset::StaticClass()))
+	{
+		return false;
+	}
+
+	TSharedRef<FPVBranchHierarchyDescription> BranchHierarchy = PV::MockTree::CreateBranchHierarchyDescription();
+
+	UPackage* PlantDescriptionDataAssetPackage = CreatePackage(*PlantDescriptionDataAssetPath);
+	PlantDescriptionDataAssetPackage->FullyLoad();
+	TObjectPtr<UPVPlantDescriptionDataAsset> PlantDescriptionDataAsset = Cast<UPVPlantDescriptionDataAsset>(PlantDescriptionDataAssetPackage->FindAssetInPackage());
+	if (PlantDescriptionDataAsset)
+	{
+		PlantDescriptionDataAsset->SetExternalPackage(nullptr);
+		PlantDescriptionDataAsset->MarkAsGarbage();
+	}
+
+	PlantDescriptionDataAsset = NewObject<UPVPlantDescriptionDataAsset>(PlantDescriptionDataAssetPackage, FPackageName::GetShortFName(PlantDescriptionDataAssetPath), RF_Standalone | RF_Public);
+
+	PlantDescriptionDataAsset->Modify();
+	PlantDescriptionDataAsset->PlantDescription.BranchHierarchy = MoveTemp(*BranchHierarchy);
+	PlantDescriptionDataAsset->PlantDescription.Foliage = {};
+	PlantDescriptionDataAsset->PostEditChange();
+
+	UPackage* ProceduralVegetationPresetPackage = CreatePackage(*ProceduralVegetationPresetPath);
+	ProceduralVegetationPresetPackage->FullyLoad();
+	TObjectPtr<UProceduralVegetationPreset> ProceduralVegetationPreset = Cast<UProceduralVegetationPreset>(ProceduralVegetationPresetPackage->FindAssetInPackage());
+	if (ProceduralVegetationPreset)
+	{
+		ProceduralVegetationPreset->SetExternalPackage(nullptr);
+		ProceduralVegetationPreset->MarkAsGarbage();
+	}
+
+	ProceduralVegetationPreset = NewObject<UProceduralVegetationPreset>(ProceduralVegetationPresetPackage, FPackageName::GetShortFName(ProceduralVegetationPresetPath), RF_Standalone | RF_Public);
+
+	Helper::GenerateProceduralVegetationPresetFromPlantDescription(PlantDescriptionDataAsset->PlantDescription, ProceduralVegetationPreset);
+
+	IAssetRegistry::Get()->AssetCreated(PlantDescriptionDataAsset);
+	IAssetRegistry::Get()->AssetCreated(ProceduralVegetationPreset);
+
+	return true;
+}
+
diff --git a/Engine/Plugins/Experimental/ProceduralVegetationEditor/Source/ProceduralVegetationEditor/Private/Importer/PVImporter_MockTree.h b/Engine/Plugins/Experimental/ProceduralVegetationEditor/Source/ProceduralVegetationEditor/Private/Importer/PVImporter_MockTree.h
new file mode 100644
--- /dev/null
+++ b/Engine/Plugins/Experimental/ProceduralVegetationEditor/Source/ProceduralVegetationEditor/Private/Importer/PVImporter_MockTree.h
@@ -0,0 +1,10 @@
+// Copyright Epic Games, Inc. All Rights Reserved.
+#pragma once
+
+#include "CoreMinimal.h"
+#include "PVPlantDescription.h"
+
+namespace PV::Import::MockTree
+{
+	bool ImportMockTreeAsProceduralVegetation(const FString& OutputPath);
+};
\ No newline at end of file
diff --git a/Engine/Plugins/Experimental/ProceduralVegetationEditor/Source/ProceduralVegetationEditor/Private/Importer/PVImporter_StaticMesh.cpp b/Engine/Plugins/Experimental/ProceduralVegetationEditor/Source/ProceduralVegetationEditor/Private/Importer/PVImporter_StaticMesh.cpp
new file mode 100644
--- /dev/null
+++ b/Engine/Plugins/Experimental/ProceduralVegetationEditor/Source/ProceduralVegetationEditor/Private/Importer/PVImporter_StaticMesh.cpp
@@ -0,0 +1,1103 @@
+// Copyright Epic Games, Inc. All Rights Reserved.
+
+#include "PVImporter_StaticMesh.h"
+#include "ProceduralVegetationEditorModule.h"
+
+#include "ProceduralVegetationPreset.h"
+#include "PVPlantDescription.h"
+#include "Helpers/PVUtilities.h"
+#include "PVImportHelpers.h"
+
+#include "AssetRegistry/IAssetRegistry.h"
+#include "DynamicMesh/DynamicMesh3.h"
+#include "DynamicMesh/DynamicMeshAABBTree3.h"
+#include "DynamicMesh/MeshNormals.h"
+#include "DynamicMesh/DynamicMeshAttributeSet.h"
+#include "Engine/StaticMesh.h"
+#include "Parameterization/MeshDijkstra.h"
+#include "SphereTypes.h"
+#include "ConversionUtils/SceneComponentToDynamicMesh.h"
+#include "GeometryCollectionToDynamicMesh.h"
+#include "GeometryCollection/GeometryCollection.h"
+#include "Selections/MeshConnectedComponents.h"
+
+namespace PV::Import::StaticMesh::Internal
+{
+	const static FName PScaleAttributeName = TEXT("pscale");
+	const static FName FuseToAttributeName = TEXT("fuseTo");
+	const static FName StartGroupAttributeName = TEXT("End_Grp");
+	const static FName EndGroupAttributeName = TEXT("Start_Grp");
+	const static FName SeedAttributeName = TEXT("seed");
+	const static FName SeedZPosAttributeName = TEXT("seed_Z_pos");
+	const static FName DistAttributeName = TEXT("dist");
+
+	typedef UE::Geometry::TDynamicMeshVertexAttribute<float, 1> FPScaleAttribute;
+	typedef UE::Geometry::TDynamicMeshVertexAttribute<int32, 1> FFuseToAttribute;
+	typedef UE::Geometry::TDynamicMeshVertexAttribute<bool, 1> FStartGroupAttribute;
+	typedef UE::Geometry::TDynamicMeshVertexAttribute<bool, 1> FEndGroupAttribute;
+	typedef UE::Geometry::TDynamicMeshVertexAttribute<bool, 1> FSeedAttribute;
+	typedef UE::Geometry::TDynamicMeshVertexAttribute<float, 1> FSeedZPosAttribute;
+	typedef UE::Geometry::TDynamicMeshVertexAttribute<float, 1> FDistAttribute;
+
+	template<typename T> 
+	T* GetAttachedAttribute(UE::Geometry::FDynamicMesh3& DynamicMesh, FName AttributeName)
+	{
+		if (DynamicMesh.HasAttributes() && DynamicMesh.Attributes()->HasAttachedAttribute(AttributeName))
+		{
+			return static_cast<T*>(DynamicMesh.Attributes()->GetAttachedAttribute(AttributeName));
+		}
+		return nullptr;
+	}
+
+	template<typename T>
+	const T* GetAttachedAttribute(const UE::Geometry::FDynamicMesh3& DynamicMesh, FName AttributeName)
+	{
+		if (DynamicMesh.HasAttributes() && DynamicMesh.Attributes()->HasAttachedAttribute(AttributeName))
+		{
+			return static_cast<const T*>(DynamicMesh.Attributes()->GetAttachedAttribute(AttributeName));
+		}
+		return nullptr;
+	}
+
+	template<typename AttribType>
+	void AttachVertexAttribute(UE::Geometry::FDynamicMesh3& DynamicMesh, FName AttributeName)
+	{
+		using namespace UE::Geometry;
+
+		if (!DynamicMesh.HasAttributes())
+		{
+			DynamicMesh.EnableAttributes();
+		}
+
+		check(!DynamicMesh.Attributes()->HasAttachedAttribute(AttributeName));
+
+		auto* Attribute = new AttribType(&DynamicMesh);
+		Attribute->Initialize();
+		DynamicMesh.Attributes()->AttachAttribute(AttributeName, Attribute);
+	}
+
+	template<typename T>
+	void CopyAttachedVertexAttributeToCollection(const UE::Geometry::FDynamicMesh3& DynamicMesh, FManagedArrayCollection& OutCollection, FName AttributeName)
+	{
+		if (!DynamicMesh.HasAttributes() || !DynamicMesh.Attributes()->HasAttachedAttribute(AttributeName))
+		{
+			return;
+		}
+
+		const static FName DebugVertexGroup = *FString::Printf(TEXT("%s_Debug"), *FGeometryCollection::VerticesGroup.ToString());
+
+		if (!OutCollection.HasGroup(DebugVertexGroup))
+		{
+			OutCollection.AddGroup(DebugVertexGroup);
+			OutCollection.AddElements(DynamicMesh.VertexCount(), DebugVertexGroup);
+
+			auto& DebugVertexPositionAttribute = OutCollection.AddAttribute<FVector3f>(FGeometryCollection::VertexPositionAttribute, DebugVertexGroup);
+			for (int32 VertexIndex : DynamicMesh.VertexIndicesItr())
+			{
+				DebugVertexPositionAttribute[VertexIndex] = FVector3f(DynamicMesh.GetVertex(VertexIndex));
+			}
+		}
+
+		auto* AttachedAttribute = Internal::GetAttachedAttribute<UE::Geometry::TDynamicMeshVertexAttribute<T, 1>>(DynamicMesh, AttributeName);
+
+		TManagedArray<T>& CollectionAttribute = OutCollection.AddAttribute<T>(AttributeName, DebugVertexGroup);
+		for (int32 VertexIndex : DynamicMesh.VertexIndicesItr())
+		{
+			T Value;
+			AttachedAttribute->GetValue(VertexIndex, &Value);
+			CollectionAttribute[VertexIndex] = Value;
+		}
+	}
+
+	FPScaleAttribute* GetPScaleAttribute(UE::Geometry::FDynamicMesh3& DynamicMesh) { return GetAttachedAttribute<FPScaleAttribute>(DynamicMesh, PScaleAttributeName); }
+	const FPScaleAttribute* GetPScaleAttribute(const UE::Geometry::FDynamicMesh3& DynamicMesh) { return GetAttachedAttribute<const FPScaleAttribute>(DynamicMesh, PScaleAttributeName); }
+	
+	FFuseToAttribute* GetFuseToAttribute(UE::Geometry::FDynamicMesh3& DynamicMesh) { return GetAttachedAttribute<FFuseToAttribute>(DynamicMesh, FuseToAttributeName); }
+	const FFuseToAttribute* GetFuseToAttribute(const UE::Geometry::FDynamicMesh3& DynamicMesh) { return GetAttachedAttribute<const FFuseToAttribute>(DynamicMesh, FuseToAttributeName); }
+	
+	FStartGroupAttribute* GetStartGroupAttribute(UE::Geometry::FDynamicMesh3& DynamicMesh) { return GetAttachedAttribute<FStartGroupAttribute>(DynamicMesh, StartGroupAttributeName); }
+	const FStartGroupAttribute* GetStartGroupAttribute(const UE::Geometry::FDynamicMesh3& DynamicMesh) { return GetAttachedAttribute<const FStartGroupAttribute>(DynamicMesh, StartGroupAttributeName); }
+	
+	FEndGroupAttribute* GetEndGroupAttribute(UE::Geometry::FDynamicMesh3& DynamicMesh) { return GetAttachedAttribute<FEndGroupAttribute>(DynamicMesh, EndGroupAttributeName); }
+	const FEndGroupAttribute* GetEndGroupAttribute(const UE::Geometry::FDynamicMesh3& DynamicMesh) { return GetAttachedAttribute<const FEndGroupAttribute>(DynamicMesh, EndGroupAttributeName); }
+	
+	FSeedAttribute* GetSeedAttribute(UE::Geometry::FDynamicMesh3& DynamicMesh) { return GetAttachedAttribute<FSeedAttribute>(DynamicMesh, SeedAttributeName); }
+	const FSeedAttribute* GetSeedAttribute(const UE::Geometry::FDynamicMesh3& DynamicMesh) { return GetAttachedAttribute<const FSeedAttribute>(DynamicMesh, SeedAttributeName); }
+	
+	FSeedZPosAttribute* GetSeedZPosAttribute(UE::Geometry::FDynamicMesh3& DynamicMesh) { return GetAttachedAttribute<FSeedZPosAttribute>(DynamicMesh, SeedZPosAttributeName); }
+	const FSeedZPosAttribute* GetSeedZPosAttribute(const UE::Geometry::FDynamicMesh3& DynamicMesh) { return GetAttachedAttribute<const FSeedZPosAttribute>(DynamicMesh, SeedZPosAttributeName); }
+	
+	FDistAttribute* GetDistAttribute(UE::Geometry::FDynamicMesh3& DynamicMesh) { return GetAttachedAttribute<FDistAttribute>(DynamicMesh, DistAttributeName); }
+	const FDistAttribute* GetDistAttribute(const UE::Geometry::FDynamicMesh3& DynamicMesh) { return GetAttachedAttribute<const FDistAttribute>(DynamicMesh, DistAttributeName); }
+
+	void AttachPScaleAttribute(UE::Geometry::FDynamicMesh3& DynamicMesh) { AttachVertexAttribute<FPScaleAttribute>(DynamicMesh, PScaleAttributeName); }
+	void AttachFuseToAttribute(UE::Geometry::FDynamicMesh3& DynamicMesh) { AttachVertexAttribute<FFuseToAttribute>(DynamicMesh, FuseToAttributeName); }
+	void AttachStartGroupAttribute(UE::Geometry::FDynamicMesh3& DynamicMesh) { AttachVertexAttribute<FStartGroupAttribute>(DynamicMesh, StartGroupAttributeName); }
+	void AttachEndGroupAttribute(UE::Geometry::FDynamicMesh3& DynamicMesh) { AttachVertexAttribute<FEndGroupAttribute>(DynamicMesh, EndGroupAttributeName); }
+	void AttachSeedAttribute(UE::Geometry::FDynamicMesh3& DynamicMesh) { AttachVertexAttribute<FSeedAttribute>(DynamicMesh, SeedAttributeName); }
+	void AttachSeedZPosAttribute(UE::Geometry::FDynamicMesh3& DynamicMesh) { AttachVertexAttribute<FSeedZPosAttribute>(DynamicMesh, SeedZPosAttributeName); }
+	void AttachDistAttribute(UE::Geometry::FDynamicMesh3& DynamicMesh) { AttachVertexAttribute<FDistAttribute>(DynamicMesh, DistAttributeName); }
+
+	void RemoveUnusedVertices(UE::Geometry::FDynamicMesh3& DynamicMesh)
+	{
+		TRACE_CPUPROFILER_EVENT_SCOPE(PV::Import::StaticMesh::OrphanRemoval);
+
+		DynamicMesh.RemoveUnusedVertices();
+		DynamicMesh.CompactInPlace();
+	}
+
+	void KeepMaterial(UE::Geometry::FDynamicMesh3& DynamicMesh, int32 MaterialIndexToKeep)
+	{
+		TRACE_CPUPROFILER_EVENT_SCOPE(PV::Import::StaticMesh::KeepMaterial);
+
+		if (MaterialIndexToKeep == INDEX_NONE)
+		{
+			return;
+		}
+
+		UE::Geometry::FDynamicMeshAttributeSet* MeshAttributeSet = DynamicMesh.Attributes();
+		UE::Geometry::FDynamicMeshMaterialAttribute* MaterialIDAttribute = MeshAttributeSet ? MeshAttributeSet->GetMaterialID() : nullptr;
+		
+		if (!MaterialIDAttribute)
+		{
+			return;
+		}
+
+		for (int32 TriangleID : DynamicMesh.TriangleIndicesItr())
+		{
+			const int32 MaterialID = MaterialIDAttribute->GetValue(TriangleID);
+			if (MaterialID != MaterialIndexToKeep)
+			{
+				DynamicMesh.RemoveTriangle(TriangleID);
+			}
+		}
+
+		RemoveUnusedVertices(DynamicMesh);
+	}
+
+	void ComputeNormals(UE::Geometry::FDynamicMesh3& DynamicMesh)
+	{
+		TRACE_CPUPROFILER_EVENT_SCOPE(PV::Import::StaticMesh::ComputeNormals);
+
+		if (!DynamicMesh.HasVertexNormals())
+		{
+			DynamicMesh.EnableVertexNormals(FVector3f::UnitX());
+		}
+
+		if (DynamicMesh.Attributes() && DynamicMesh.Attributes()->NumNormalLayers() > 0)
+		{
+			UE::Geometry::FDynamicMeshNormalOverlay* Overlay = DynamicMesh.Attributes()->GetNormalLayer(0);
+			for (int32 EID : Overlay->ElementIndicesItr())
+			{
+				int32 ParentVID = Overlay->GetParentVertex(EID);
+				if (ParentVID != INDEX_NONE)
+				{
+					DynamicMesh.SetVertexNormal(ParentVID, Overlay->GetElement(EID));
+				}
+			}
+		}
+		else
+		{
+			UE::Geometry::FMeshNormals::QuickComputeVertexNormals(DynamicMesh);
+		}
+	}
+
+	void AverageNormals(UE::Geometry::FDynamicMesh3& DynamicMesh)
+	{
+		// This takes the normals of each vertex instance and averages them onto the vertex (or point as it's called in Houdini)
+		// This one seems more important than the rest for getting any good results
+
+		// TODO: Not super important for now, do later
+	}
+
+	void SmoothNormalsWithLengthPreservation(UE::Geometry::FDynamicMesh3& DynamicMesh)
+	{
+		// TODO: Has minimal impact on the final result, can wait with this
+	}
+
+	void GetSpherical_Internal(UE::Geometry::FDynamicMesh3& DynamicMesh, const TArray<int32>& TriangleIds)
+	{
+		TRACE_CPUPROFILER_EVENT_SCOPE(PV::Import::StaticMesh::GetSpherical_Internal);
+
+		using namespace UE::Geometry;
+
+		TSet<int32> UniqueVertexIndices;
+		UniqueVertexIndices.Reserve(TriangleIds.Num() * 3);
+		for (int32 tid : TriangleIds)
+		{
+			const FIndex3i VertexIndices = DynamicMesh.GetTriangle(tid);
+			UniqueVertexIndices.Add(VertexIndices[0]);
+			UniqueVertexIndices.Add(VertexIndices[1]);
+			UniqueVertexIndices.Add(VertexIndices[2]);
+		}
+
+		const TArray<int32> ComponentVertexIndices = UniqueVertexIndices.Array();
+
+		FDynamicMeshAABBTree3 DynamicMeshAABBTree(&DynamicMesh, false);
+		DynamicMeshAABBTree.Build(TriangleIds);
+
+		const static auto ComputeRadius = [](const FVector& Point, const FVector& Normal, const FVector& QueryPosition)->double
+		{
+			// Compute radius of the ball that touches points p and q and whose center falls on the normal n from p
+			const FVector Delta = Point - QueryPosition;
+			const double Dist = Delta.Size();
+			const double CosTheta = Normal.Dot(Delta) / Dist;
+			return Dist / (2 * CosTheta);
+		};
+
+		double NearestDistSqr = 0.f;
+
+		TArray<FVector> NewVertexPositions;
+		NewVertexPositions.SetNum(ComponentVertexIndices.Num());
+
+		TArray<float> PScales;
+		PScales.SetNum(ComponentVertexIndices.Num());
+
+		for (int32 i = 0; i < ComponentVertexIndices.Num(); ++i)
+		{
+			const int32 VertexIndex = ComponentVertexIndices[i];
+			const FVector& VertexPosition = DynamicMesh.GetVertex(VertexIndex);
+			const FVector VertexNormal = FVector(DynamicMesh.GetVertexNormal(VertexIndex));
+
+			NewVertexPositions[i] = VertexPosition;
+			PScales[i] = 0.5;
+
+			/*if (VertexNormal.Size() < 0.5) // TODO: Re-visit once we avg normals
+			{
+				continue;
+			}
+			
+			VertexNormal.Normalize();
+			*/
+
+			float Radius = 50; // 50cm
+			int32 OldPoint = INDEX_NONE;
+			FVector SphereCenter = VertexPosition - (FVector(VertexNormal) * Radius);
+			
+			int32 NearPoint = DynamicMeshAABBTree.FindNearestVertex(SphereCenter, NearestDistSqr, Radius);
+			while (NearPoint != INDEX_NONE)
+			{
+				// Get the position of the nearest point of the previous near point.
+				const FVector point_position = DynamicMesh.GetVertex(NearPoint);
+
+				// Computes the radius of a sphere that would fit spherically between these two positions
+				const float NewRadius = ComputeRadius(VertexPosition, VertexNormal, point_position);
+
+				Radius = FMath::Min(Radius, FMath::Max(NewRadius, 0));
+				// Sets the center of the lookup shere to be at the positions of the orignal position, minus the offset of the new found radius
+				SphereCenter = VertexPosition - VertexNormal * Radius;
+
+				// Stores the current near point as the old point for comparison
+				OldPoint = NearPoint;
+				NearPoint = DynamicMeshAABBTree.FindNearestVertex(SphereCenter, NearestDistSqr, Radius);
+
+				if (NearPoint == VertexIndex || OldPoint == NearPoint)
+				{ 
+					break; 
+				}
+			}
+
+			const int32 NearTriID = DynamicMeshAABBTree.FindNearestTriangle(SphereCenter, NearestDistSqr);
+			if (NearTriID == INDEX_NONE)
+			{
+				continue;
+			}
+
+			const FDistPoint3Triangle3d Query = TMeshQueries<FDynamicMesh3>::TriangleDistance(DynamicMesh, NearTriID, SphereCenter);
+			const FVector& ClosestTrianglePoint = Query.ClosestTrianglePoint;
+			const FVector ClosestTrianglePointNormal = DynamicMesh.GetTriBaryNormal(NearTriID, Query.TriangleBaryCoords[0], Query.TriangleBaryCoords[1], Query.TriangleBaryCoords[2]);
+
+			const FVector CenterDir = (SphereCenter - ClosestTrianglePoint).GetSafeNormal();
+			const double CenterDot = ClosestTrianglePointNormal.Dot(CenterDir);
+			if (CenterDot > 0) 
+			{ 
+				continue; 
+			}
+
+			PScales[i] = FMath::Max(Radius, 0.01);
+			NewVertexPositions[i] = SphereCenter;
+		}
+
+		FPScaleAttribute* PScaleAttribute = GetPScaleAttribute(DynamicMesh);
+		check(PScaleAttribute != nullptr);
+
+		for (int32 i = 0; i < ComponentVertexIndices.Num(); ++i)
+		{
+			const int32 VertexIndex = ComponentVertexIndices[i];
+			DynamicMesh.SetVertex(VertexIndex, NewVertexPositions[i]);
+			PScaleAttribute->SetValue(VertexIndex, &PScales[i]);
+		}
+	}
+
+	void GetSpherical(UE::Geometry::FDynamicMesh3& DynamicMesh)
+	{
+		TRACE_CPUPROFILER_EVENT_SCOPE(PV::Import::StaticMesh::GetSpherical);
+
+		UE::Geometry::FMeshConnectedComponents ConnectedComponents(&DynamicMesh);
+		ConnectedComponents.FindConnectedTriangles();
+
+		ParallelFor(ConnectedComponents.Num(), [&](int32 Index)
+		{
+			Internal::GetSpherical_Internal(DynamicMesh, ConnectedComponents.GetComponent(Index).Indices);
+		});
+	}
+
+	void BlurRadius(UE::Geometry::FDynamicMesh3& DynamicMesh)
+	{
+		// TODO: Has smaller impact on the final result, can wait with this
+	}
+	
+	void ExtractReductionCloud(UE::Geometry::FDynamicMesh3& DynamicMesh)
+	{
+		TRACE_CPUPROFILER_EVENT_SCOPE(PV::Import::StaticMesh::ExtractReductionCloud);
+
+		using namespace UE::Geometry;
+
+		FPScaleAttribute* PScaleAttribute = GetPScaleAttribute(DynamicMesh);
+		FFuseToAttribute* FuseToAttribute = GetFuseToAttribute(DynamicMesh);
+		check(PScaleAttribute != nullptr);
+		check(FuseToAttribute != nullptr);
+		
+		FDynamicMeshAABBTree3 DynamicMeshAABBTree(&DynamicMesh, false);
+		DynamicMeshAABBTree.Build();
+
+		const auto FindOverlappingVertices = [&](const FVector& Pos, float Radius)->TArray<int32>
+		{
+			TArray<int32> OutVertexIndices;
+			OutVertexIndices.Reserve(DynamicMesh.VertexCount());
+
+			const double RadiusSqr = Radius * Radius;
+
+			FDynamicMeshAABBTree3::FTreeTraversal TreeTraversal;
+
+			TreeTraversal.NextBoxF = [&](const FAxisAlignedBox3d& Box, int Depth)->bool
+			{
+				return Box.DistanceSquared(Pos) <= RadiusSqr;
+			};
+
+			TreeTraversal.NextTriangleF = [&](int32 TID)
+			{
+				const FIndex3i Triangle = DynamicMesh.GetTriangle(TID);
+				for (int j = 0; j < 3; ++j)
+				{
+					double VertexDistSqr = DistanceSquared(DynamicMesh.GetVertex(Triangle[j]), Pos);
+					if (VertexDistSqr <= RadiusSqr)
+					{
+						OutVertexIndices.AddUnique(Triangle[j]);
+					}
+				}
+			};
+
+			DynamicMeshAABBTree.DoTraversal(TreeTraversal);
+
+			return OutVertexIndices;
+		};
+
+		TArray<int32> VertexIndices;
+		VertexIndices.Reserve(DynamicMesh.VertexCount());
+		for (int32 VertexIndex : DynamicMesh.VertexIndicesItr())
+		{
+			VertexIndices.Add(VertexIndex);
+		}
+		
+		TArray<bool> Removed; 
+		Removed.SetNumZeroed(DynamicMesh.MaxVertexID());
+		
+		TArray<bool> Checked;
+		Checked.SetNumZeroed(DynamicMesh.MaxVertexID());
+		
+		TArray<float> Radii;
+		Radii.SetNumUninitialized(DynamicMesh.MaxVertexID());
+		
+		for (int32 i = 0; i < VertexIndices.Num(); i++)
+		{
+			PScaleAttribute->GetValue(VertexIndices[i], &Radii[VertexIndices[i]]);
+			Radii[VertexIndices[i]] *= 2.0;
+		}
+
+		{
+			TRACE_CPUPROFILER_EVENT_SCOPE(SortByRadius);
+			VertexIndices.Sort([&](int32 A, int32 B) { return Radii[A] > Radii[B]; });
+		}
+
+		for (int i = 0; i < VertexIndices.Num(); i++)
+		{
+			const int32 VertexIndex = VertexIndices[i];
+
+			if (Removed[VertexIndex])
+			{ 
+				continue; 
+			}
+
+			Checked[VertexIndex] = true;
+
+			const FVector VertexPosition = DynamicMesh.GetVertex(VertexIndex);
+			const float VertexRadius = Radii[VertexIndex];
+
+			const TArray<int32> OverlappingVertices = FindOverlappingVertices(VertexPosition, VertexRadius);
+			for (int j = 0; j < OverlappingVertices.Num(); j++)
+			{
+				const int32 OverlappingVertexIndex = OverlappingVertices[j];
+				if (OverlappingVertexIndex == VertexIndex || Removed[OverlappingVertexIndex] || Checked[OverlappingVertexIndex])
+				{
+					continue;
+				}
+
+				const FVector OverlappingVertexPosition = DynamicMesh.GetVertex(OverlappingVertexIndex);
+				const float OverlappingVertexRadius = Radii[OverlappingVertexIndex];
+				const float OverBounds = OverlappingVertexRadius + VertexRadius;
+
+				const FVector Dist = (OverlappingVertexPosition - VertexPosition).GetAbs() - OverBounds;
+				if (Dist.X > 0 || Dist.Y > 0 || Dist.Z > 0)
+				{
+					continue;
+				}
+
+				Removed[OverlappingVertexIndex] = true;
+				FuseToAttribute->SetValue(OverlappingVertexIndex, &VertexIndex);
+			}
+
+			PScaleAttribute->SetValue(VertexIndex, &VertexRadius);
+			FuseToAttribute->SetValue(VertexIndex, &VertexIndex);
+		}
+	}
+	
+	void RadialFuse(UE::Geometry::FDynamicMesh3& DynamicMesh)
+	{
+		TRACE_CPUPROFILER_EVENT_SCOPE(PV::Import::StaticMesh::RadialFuse);
+
+		FFuseToAttribute* FuseToAttribute = GetFuseToAttribute(DynamicMesh);
+		check(FuseToAttribute != nullptr);
+
+		TArray<int32> TriangleIndices;
+		TriangleIndices.Reserve(DynamicMesh.TriangleCount());
+		for (int32 TID : DynamicMesh.TriangleIndicesItr())
+		{
+			TriangleIndices.Add(TID);
+		}
+
+		for (int i = 0; i < TriangleIndices.Num(); i++)
+		{
+			const int32 TriangleIndex = TriangleIndices[i];
+
+			const UE::Geometry::FIndex3i VertexIndices = DynamicMesh.GetTriangle(TriangleIndex);
+			
+			TArray<int, TInlineAllocator<3>> NewTriangleIndices;
+			for (int j = 0; j < 3; j++) 
+			{
+				int fuse_to;
+				FuseToAttribute->GetValue(VertexIndices[j], &fuse_to);
+				NewTriangleIndices.AddUnique(fuse_to);
+			}
+
+			DynamicMesh.RemoveTriangle(TriangleIndex, false);
+			
+			if (NewTriangleIndices.Num() == 2)
+			{ 
+				const FVector P1 = DynamicMesh.GetVertex(NewTriangleIndices[0]);
+				const FVector P2 = DynamicMesh.GetVertex(NewTriangleIndices[1]);
+
+				// Since we don't have edges for the Dynamic mesh we add a new vertex at the center of the edge, offset by a random vector
+				const FVector P3 = ((P1 + P2) / 2) + (P1 - P2).GetSafeNormal().Cross(FVector::ForwardVector).Cross(FVector::UpVector);
+				
+				const int32 P3Index = DynamicMesh.AppendVertex(P3);
+				DynamicMesh.AppendTriangle(NewTriangleIndices[0], NewTriangleIndices[1], P3Index);
+			}
+			else if (NewTriangleIndices.Num() == 3)
+			{
+				DynamicMesh.AppendTriangle(NewTriangleIndices[0], NewTriangleIndices[1], NewTriangleIndices[2]);
+			}
+			else if (NewTriangleIndices.Num() > 3)
+			{
+				check(false); // This should not be able to happen
+			}
+		}
+
+		Internal::RemoveUnusedVertices(DynamicMesh);
+	}
+	
+	void GroupPathPoints(UE::Geometry::FDynamicMesh3& DynamicMesh)
+	{
+		TRACE_CPUPROFILER_EVENT_SCOPE(PV::Import::StaticMesh::GroupPathPoints);
+
+		FStartGroupAttribute* StartGroupAttribute = GetStartGroupAttribute(DynamicMesh);
+		FEndGroupAttribute* EndGroupAttribute = GetEndGroupAttribute(DynamicMesh);
+		FSeedAttribute* SeedAttribute = GetSeedAttribute(DynamicMesh);
+		FSeedZPosAttribute* SeedZPosAttribute = GetSeedZPosAttribute(DynamicMesh);
+		check(StartGroupAttribute != nullptr);
+		check(EndGroupAttribute != nullptr);
+		check(SeedAttribute != nullptr);
+		check(SeedZPosAttribute != nullptr);
+
+		const UE::Geometry::FAxisAlignedBox3d BoxBounds = DynamicMesh.GetBounds();
+
+		TArray<double> xComps;
+		TArray<double> yComps;
+		TArray<double> zComps;
+		for (int32 VertexIndex : DynamicMesh.VertexIndicesItr())
+		{
+			const FVector VertexPosition = DynamicMesh.GetVertex(VertexIndex);
+			
+			const int32 EdgeCount = DynamicMesh.GetVtxEdgeCount(VertexIndex);
+			check(EdgeCount > 0);
+
+			xComps.Reset(EdgeCount);
+			yComps.Reset(EdgeCount);
+			zComps.Reset(EdgeCount);
+
+			for (int32 EdgeID : DynamicMesh.VtxEdgesItr(VertexIndex))
+			{
+				const UE::Geometry::FIndex2i VertexIDs = DynamicMesh.GetEdgeV(EdgeID);
+				const int32 OtherVertexID = VertexIDs.A == VertexIndex ? VertexIDs.B : VertexIDs.A;
+
+				const FVector NeighbourVertexPosition = DynamicMesh.GetVertex(OtherVertexID);
+				const FVector Dir = (NeighbourVertexPosition - VertexPosition).GetSafeNormal();
+				xComps.Add(Dir.X);
+				yComps.Add(Dir.Y);
+				zComps.Add(Dir.Z);
+			}
+
+			const double* MaxXComp = Algo::MaxElement(xComps);
+			const double* MinXComp = Algo::MinElement(xComps);
+			const double* MaxYComp = Algo::MaxElement(yComps);
+			const double* MinYComp = Algo::MinElement(yComps);
+			const double* MaxZComp = Algo::MaxElement(zComps);
+			const double* MinZComp = Algo::MinElement(zComps);
+
+			const double xDiff = *MaxXComp - *MinXComp;
+			const double yDiff = *MaxYComp - *MinYComp;
+			const double zDiff = *MaxZComp - *MinZComp;
+			const double MaxDiff = FMath::Max3(xDiff, yDiff, zDiff);
+
+			if (MaxDiff > 1.0)
+			{
+				continue;
+			}
+
+			const float Z = FMath::GetMappedRangeValueClamped(FVector2D(BoxBounds.Min.Z, BoxBounds.Max.Z), FVector2D(0, 1), VertexPosition.Z);
+			const bool bIsEndGroup = Z > 0.1;
+			const bool bIsStartGroup = !bIsEndGroup;
+
+			StartGroupAttribute->SetValue(VertexIndex, &bIsStartGroup);
+			EndGroupAttribute->SetValue(VertexIndex, &bIsEndGroup);
+			SeedAttribute->SetValue(VertexIndex, &bIsStartGroup);
+			SeedZPosAttribute->SetValue(VertexIndex, &Z);
+		}
+	}
+	
+	void ComputeDist(UE::Geometry::FDynamicMesh3& DynamicMesh)
+	{
+		TRACE_CPUPROFILER_EVENT_SCOPE(PV::Import::StaticMesh::ComputeDist);
+
+		FStartGroupAttribute* StartGroupAttribute = GetStartGroupAttribute(DynamicMesh);
+		FEndGroupAttribute* EndGroupAttribute = GetEndGroupAttribute(DynamicMesh);
+		FDistAttribute* DistAttribute = GetDistAttribute(DynamicMesh);
+		check(StartGroupAttribute != nullptr);
+		check(EndGroupAttribute != nullptr);
+		check(DistAttribute != nullptr);
+
+		typedef UE::Geometry::TMeshDijkstra<UE::Geometry::FDynamicMesh3> FMeshDijkstra;
+		FMeshDijkstra MeshDijkstra(&DynamicMesh);
+
+		TArray<FMeshDijkstra::FSeedPoint> SeedPoints;
+
+		for (int32 VertexIndex : DynamicMesh.VertexIndicesItr())
+		{
+			bool bIsStart = false;
+			StartGroupAttribute->GetValue(VertexIndex, &bIsStart);
+
+			if (bIsStart)
+			{
+				FMeshDijkstra::FSeedPoint& Seed = SeedPoints.AddDefaulted_GetRef();
+				Seed.PointID = VertexIndex;
+			}
+		}
+
+		for (int32 VertexIndex : DynamicMesh.VertexIndicesItr())
+		{
+			float Dist = 0;
+
+			MeshDijkstra.ComputeToTargetPoint(SeedPoints, VertexIndex);
+
+			TArray<int32> VertexPath;
+			if (!MeshDijkstra.FindPathToNearestSeed(VertexIndex, VertexPath))
+			{
+				const bool FalseValue = false;
+				EndGroupAttribute->SetValue(VertexIndex, &FalseValue);
+			}
+			else
+			{
+				// TODO: Is there a built-in function in TMeshDijkstra for this? Feels expensive to have to compute this for each point...
+				for (int32 i = 0; i < VertexPath.Num() - 1; ++i)
+				{
+					int32 Vertex0ID = VertexPath[i];
+					int32 Vertex1ID = VertexPath[i + 1];
+					Dist += (DynamicMesh.GetVertex(Vertex0ID) - DynamicMesh.GetVertex(Vertex1ID)).Size();
+				}
+			}
+
+			DistAttribute->SetValue(VertexIndex, &Dist);
+		}
+	}
+
+	void DrawBranches(const UE::Geometry::FDynamicMesh3& DynamicMesh, TArray<TArray<int32>>& OutBranchCurves)
+	{
+		TRACE_CPUPROFILER_EVENT_SCOPE(PV::Import::StaticMesh::DrawBranches);
+
+		const FStartGroupAttribute* StartGroupAttribute = GetStartGroupAttribute(DynamicMesh);
+		const FEndGroupAttribute* EndGroupAttribute = GetEndGroupAttribute(DynamicMesh);
+		const FDistAttribute* DistAttribute = GetDistAttribute(DynamicMesh);
+		check(StartGroupAttribute != nullptr);
+		check(EndGroupAttribute != nullptr);
+		check(DistAttribute != nullptr);
+
+		const auto GetVertexNeighbours = [&](int32 VertexID)->TArray<int32>
+		{
+			TArray<int32> OutNeighbours;
+			OutNeighbours.Reserve(DynamicMesh.GetVtxEdgeCount(VertexID));
+			for (int32 EdgeID : DynamicMesh.VtxEdgesItr(VertexID))
+			{
+				const UE::Geometry::FIndex2i VertexIDs = DynamicMesh.GetEdgeV(EdgeID);
+				const int32 OtherVertexID = VertexIDs.A == VertexID ? VertexIDs.B : VertexIDs.A;
+				OutNeighbours.Add(OtherVertexID);
+			}
+			return OutNeighbours;
+		};
+
+		const auto FindNextPoint = [&](int32 PointIndex, float MinDist)->int32
+		{
+			int32 OutNextPointIndex = INDEX_NONE;
+			const TArray<int32> Neighbours = GetVertexNeighbours(PointIndex);
+			for (int32 Neighbour : Neighbours)
+			{
+				float Dist = 0;
+				DistAttribute->GetValue(Neighbour, &Dist);
+				if (Dist > MinDist)
+				{ 
+					continue; 
+				}
+				MinDist = Dist;
+				OutNextPointIndex = Neighbour;
+			}
+			return OutNextPointIndex;
+		};
+
+		const auto ExpandPointGroup = [&](const UE::Geometry::TDynamicMeshVertexAttribute<bool, 1>* GroupAttribute)->TArray<int32>
+		{
+			TArray<int32> OutVertexIndices;
+			OutVertexIndices.Reserve(DynamicMesh.VertexCount());
+			if (GroupAttribute == nullptr)
+			{
+				for (int32 VertexIndex : DynamicMesh.VertexIndicesItr())
+				{
+					OutVertexIndices.Add(VertexIndex);
+				}
+			}
+			else
+			{
+				for (int32 VertexIndex : DynamicMesh.VertexIndicesItr())
+				{
+					bool bIsInGroup = false;
+					GroupAttribute->GetValue(VertexIndex, &bIsInGroup);
+
+					if (bIsInGroup)
+					{
+						OutVertexIndices.Add(VertexIndex);
+					}
+				}
+			}
+
+			return OutVertexIndices;
+		};
+
+		const TArray<int32> EndPoints = ExpandPointGroup(EndGroupAttribute);
+		const TArray<int32> EtartPoints = ExpandPointGroup(StartGroupAttribute);
+		
+		for (int32 i = 0; i < EndPoints.Num(); i++)
+		{
+			int32 CurrentPointIndex = EndPoints[i];
+			
+			float MinDist = 0;
+			DistAttribute->GetValue(CurrentPointIndex, &MinDist);
+			
+			TArray<int32> CurvePoints = { CurrentPointIndex };
+
+			for (int j = 0; j < DynamicMesh.VertexCount(); j++)
+			{
+				const int32 NextPointIndex = FindNextPoint(CurrentPointIndex, MinDist);
+				if (NextPointIndex == INDEX_NONE)
+				{ 
+					break; 
+				}
+
+				CurvePoints.Add(NextPointIndex);
+
+				bool bIsStart = false;
+				StartGroupAttribute->GetValue(NextPointIndex, &bIsStart);
+				if (bIsStart)
+				{ 
+					break; 
+				}
+
+				CurrentPointIndex = NextPointIndex;
+			}
+
+			if (CurvePoints.Num() > 1)
+			{
+				OutBranchCurves.Add(MoveTemp(CurvePoints));
+			}
+		}
+	}
+
+	void RemoveOverlaps(const UE::Geometry::FDynamicMesh3& DynamicMesh, TArray<TArray<int32>>& InOutBranchCurves)
+	{
+		TRACE_CPUPROFILER_EVENT_SCOPE(PV::Import::StaticMesh::RemoveOverlaps);
+
+		TArray<float> BranchLengths;
+		BranchLengths.SetNum(InOutBranchCurves.Num());
+
+		int32 MaxNumVertices = 0;
+		for (int32 i = 0; i < InOutBranchCurves.Num(); ++i)
+		{
+			const TArray<int32>& BranchCurve = InOutBranchCurves[i];
+			MaxNumVertices += BranchCurve.Num();
+			for (int32 j = 0; j < BranchCurve.Num() - 1; ++j)
+			{
+				const FVector& Vertex0 = DynamicMesh.GetVertex(BranchCurve[j]);
+				const FVector& Vertex1 = DynamicMesh.GetVertex(BranchCurve[j + 1]);
+				BranchLengths[i] += (Vertex0 - Vertex1).Size();
+			}
+		}
+
+		struct FVertexLongestBranch
+		{
+			int32 BranchIndex = INDEX_NONE;
+			float BranchLength = 0;
+		};
+
+		TMap<int32, FVertexLongestBranch> VerticesLongestBranch;
+		VerticesLongestBranch.Reserve(MaxNumVertices);
+
+		for (int32 i = 0; i < InOutBranchCurves.Num(); ++i)
+		{
+			const TArray<int32>& BranchCurve = InOutBranchCurves[i];
+			const float BranchLength = BranchLengths[i];
+			
+			for (int32 j = 0; j < BranchCurve.Num(); ++j)
+			{
+				const int32 VertexIndex = BranchCurve[j];
+				FVertexLongestBranch& VertexLongestBranch = VerticesLongestBranch.FindOrAdd(VertexIndex);
+
+				if (BranchLength > VertexLongestBranch.BranchLength)
+				{
+					VertexLongestBranch.BranchLength = BranchLength;
+					VertexLongestBranch.BranchIndex = i;
+				}
+			}
+		}
+
+		for (int32 i = InOutBranchCurves.Num() - 1; i >= 0; --i)
+		{
+			TArray<int32>& BranchCurve = InOutBranchCurves[i];
+			
+			for (int32 j = BranchCurve.Num() - 1; j >= 0; --j)
+			{
+				const int32 ThisVertexIndex = BranchCurve[j];
+				const int32 ThisLongestBranchIndex = VerticesLongestBranch[ThisVertexIndex].BranchIndex;
+				
+				if (ThisLongestBranchIndex != i)
+				{
+					if (j == 0)
+					{
+						BranchCurve.RemoveAt(j);
+					}
+					else
+					{
+						const int32 NextVertexIndex = BranchCurve[j - 1];
+						const int32 NextLongestBranchIndex = VerticesLongestBranch[NextVertexIndex].BranchIndex;
+						if (NextLongestBranchIndex != i)
+						{
+							BranchCurve.RemoveAt(j);
+						}
+					}
+				}
+			}
+
+			if (BranchCurve.Num() < 2)
+			{
+				InOutBranchCurves.RemoveAtSwap(i);
+			}
+		}
+	}
+};
+
+namespace PV::Import::StaticMesh::Internal2
+{
+	void ExtractBranchesFromBranchCurves(const UE::Geometry::FDynamicMesh3& DynamicMesh, 
+		const TArray<TArray<int32>>& InBranchCurves, 
+		TArray<FVector3f>& OutPoints, 
+		TArray<float>& OutPointsRadii, 
+		TArray<TArray<int32>>& OutBranchPointIndices
+	)
+	{
+		TRACE_CPUPROFILER_EVENT_SCOPE(PV::Import::StaticMesh::ExtractBranchesFromBranchCurves);
+
+		const PV::Import::StaticMesh::Internal::FPScaleAttribute* PScaleAttribute = PV::Import::StaticMesh::Internal::GetPScaleAttribute(DynamicMesh);
+		const PV::Import::StaticMesh::Internal::FStartGroupAttribute* EndGroupAttribute = PV::Import::StaticMesh::Internal::GetEndGroupAttribute(DynamicMesh);
+		check(PScaleAttribute != nullptr);
+		check(EndGroupAttribute != nullptr);
+
+		int32 MaxNrOfPoints = 0;
+		for (const TArray<int32>& BranchCurve : InBranchCurves)
+		{
+			MaxNrOfPoints += BranchCurve.Num();
+		}
+
+		OutPoints.Reserve(MaxNrOfPoints);
+		OutPointsRadii.Reserve(MaxNrOfPoints);
+		OutBranchPointIndices.Reserve(InBranchCurves.Num());
+
+		TMap<int32, int32> VertexIndexToPointIndex;
+		VertexIndexToPointIndex.Reserve(MaxNrOfPoints);
+
+		for (const TArray<int32>& BranchCurve : InBranchCurves)
+		{
+			TArray<int32>& Branch = OutBranchPointIndices.AddDefaulted_GetRef();
+			Branch.Reserve(BranchCurve.Num());
+
+			for (int32 i = BranchCurve.Num() - 1; i >= 0; --i) // Reverse since curves go from tip to base while branches to from base to tip
+			{
+				const int32 VertexIndex = BranchCurve[i];
+				int32 PointIndex = INDEX_NONE;
+
+				int32* ExistingPointIndex = VertexIndexToPointIndex.Find(VertexIndex);
+				if (!ExistingPointIndex)
+				{
+					PointIndex = OutPoints.Add(FVector3f(DynamicMesh.GetVertex(VertexIndex)));
+					
+					bool bIsInEndGroup = false;
+					EndGroupAttribute->GetValue(VertexIndex, &bIsInEndGroup);
+
+					if (bIsInEndGroup)
+					{
+						OutPointsRadii.Add(0.01f);
+					}
+					else
+					{
+						float VertexRadius = 0;
+						PScaleAttribute->GetValue(VertexIndex, &VertexRadius);
+						OutPointsRadii.Add(VertexRadius);
+					}
+
+					VertexIndexToPointIndex.Add(VertexIndex, PointIndex);
+				}
+				else
+				{
+					PointIndex = *ExistingPointIndex;
+				}
+
+				Branch.Add(PointIndex);
+			}
+		}
+	}
+};
+
+PV::Import::StaticMesh::EImportResult PV::Import::StaticMesh::ImportGrowthDataFromStaticMesh(const TSoftObjectPtr<UStaticMesh>& InSourceStaticMesh, FManagedArrayCollection& OutCollection, const FImportMeshOptions& InOptions)
+{
+	TRACE_CPUPROFILER_EVENT_SCOPE(PV::Import::StaticMesh::ImportGrowthDataFromStaticMesh);
+
+	UStaticMesh* StaticMesh = InSourceStaticMesh.LoadSynchronous();
+	if (!StaticMesh)
+	{
+		return EImportResult::InvalidSourceMesh;
+	}
+
+	if (StaticMesh->IsCompiling())
+	{
+		return EImportResult::MeshNotReady;
+	}
+
+	UE::Geometry::FDynamicMesh3 DynamicMesh;
+	{
+		TRACE_CPUPROFILER_EVENT_SCOPE(StaticMeshToDynamicMesh);
+
+		FText OutErrorMessage;
+		UE::Conversion::FStaticMeshConversionOptions Options;
+		if (!UE::Conversion::StaticMeshToDynamicMesh(StaticMesh, DynamicMesh, OutErrorMessage, Options))
+		{
+			return EImportResult::InvalidMeshData;
+		}
+	}
+
+	const auto ReturnWithDebugMesh = [&OutCollection](const UE::Geometry::FDynamicMesh3& DynamicMesh)->EImportResult
+	{
+		FGeometryCollection MeshGeomCollection;
+		UE::Geometry::FGeometryCollectionToDynamicMeshes::FToCollectionOptions Options2;
+		UE::Geometry::FGeometryCollectionToDynamicMeshes::AppendMeshToCollection(MeshGeomCollection, DynamicMesh, FTransform::Identity, Options2);
+
+		FManagedArrayCollection MeshArrayCollection;
+		MeshGeomCollection.CopyTo(&OutCollection);
+
+		return EImportResult::Success;
+	};
+
+	if (InOptions.DebugState == FImportMeshOptions::EDebugState::InitialMesh)
+	{
+		return ReturnWithDebugMesh(DynamicMesh);
+	}
+
+	const int32 MaterialIndexToKeep = StaticMesh->GetMaterialIndex(InOptions.MaterialFilter);
+	Internal::KeepMaterial(DynamicMesh, MaterialIndexToKeep); // For now we only allow filtering of one material. Could easily be extended to allow for multiple materials
+	if (InOptions.DebugState == FImportMeshOptions::EDebugState::KeepMaterial)
+	{
+		return ReturnWithDebugMesh(DynamicMesh);
+	}
+
+	Internal::ComputeNormals(DynamicMesh);
+	if (InOptions.DebugState == FImportMeshOptions::EDebugState::ComputeNormals)
+	{
+		return ReturnWithDebugMesh(DynamicMesh);
+	}
+
+	const UE::Geometry::FDynamicMesh3 ReferenceMesh = DynamicMesh;
+	Internal::AttachPScaleAttribute(DynamicMesh);
+	Internal::AttachFuseToAttribute(DynamicMesh);
+	Internal::AttachStartGroupAttribute(DynamicMesh);
+	Internal::AttachEndGroupAttribute(DynamicMesh);
+	Internal::AttachSeedAttribute(DynamicMesh);
+	Internal::AttachSeedZPosAttribute(DynamicMesh);
+	Internal::AttachDistAttribute(DynamicMesh);
+
+	/*
+	Internal::AverageNormals(DynamicMesh);
+	if (InOptions.DebugState == FImportMeshOptions::EDebugState::AverageNormals)
+	{
+		return ReturnWithDebugMesh();
+	}
+
+	Internal::SmoothNormalsWithLengthPreservation(DynamicMesh);
+	if (InOptions.DebugState == FImportMeshOptions::EDebugState::SmoothNormalsWithLengthPreservation)
+	{
+		return ReturnWithDebugMesh();
+	}*/
+
+	Internal::GetSpherical(DynamicMesh);
+	if (InOptions.DebugState == FImportMeshOptions::EDebugState::GetSpherical)
+	{
+		Internal::CopyAttachedVertexAttributeToCollection<float>(DynamicMesh, OutCollection, Internal::PScaleAttributeName);
+		return ReturnWithDebugMesh(DynamicMesh);
+	}
+
+	/*Internal::BlurRadius(DynamicMesh);
+	if (InOptions.DebugState == FImportMeshOptions::EDebugState::BlurRadius)
+	{
+		return ReturnWithDebugMesh();
+	}*/
+
+	Internal::ExtractReductionCloud(DynamicMesh);
+	if (InOptions.DebugState == FImportMeshOptions::EDebugState::ExtractReductionCloud)
+	{
+		Internal::CopyAttachedVertexAttributeToCollection<float>(DynamicMesh, OutCollection, Internal::PScaleAttributeName);
+		Internal::CopyAttachedVertexAttributeToCollection<int32>(DynamicMesh, OutCollection, Internal::FuseToAttributeName);
+		return ReturnWithDebugMesh(DynamicMesh);
+	}
+
+	Internal::RadialFuse(DynamicMesh);
+	if (InOptions.DebugState == FImportMeshOptions::EDebugState::RadialFuse)
+	{
+		return ReturnWithDebugMesh(DynamicMesh);
+	}
+
+	Internal::GroupPathPoints(DynamicMesh);
+	if (InOptions.DebugState == FImportMeshOptions::EDebugState::GroupPathPoints)
+	{
+		Internal::CopyAttachedVertexAttributeToCollection<bool>(DynamicMesh, OutCollection, Internal::StartGroupAttributeName);
+		Internal::CopyAttachedVertexAttributeToCollection<bool>(DynamicMesh, OutCollection, Internal::EndGroupAttributeName);
+		Internal::CopyAttachedVertexAttributeToCollection<bool>(DynamicMesh, OutCollection, Internal::SeedAttributeName);
+		Internal::CopyAttachedVertexAttributeToCollection<float>(DynamicMesh, OutCollection, Internal::SeedZPosAttributeName);
+
+		return ReturnWithDebugMesh(DynamicMesh);
+	}
+	
+	Internal::ComputeDist(DynamicMesh);
+	if (InOptions.DebugState == FImportMeshOptions::EDebugState::ComputeDist)
+	{
+		Internal::CopyAttachedVertexAttributeToCollection<float>(DynamicMesh, OutCollection, Internal::DistAttributeName);
+
+		return ReturnWithDebugMesh(DynamicMesh);
+	}
+
+	const static auto AddBranchCurvesToCollecton = [](
+		FManagedArrayCollection& InCollection, 
+		const TArray<TArray<int32>>& BranchCurves, 
+		const TFunction<FVector3f(int32)>& GetVertexPosFunc)
+	{
+		const static FName BranchCurveGroupName = TEXT("BranchCurves");
+		const static FName CurveAttributeName = TEXT("Curve");
+
+		InCollection.AddGroup(BranchCurveGroupName);
+		InCollection.AddElements(BranchCurves.Num(), BranchCurveGroupName);
+
+		TManagedArray<TArray<FVector3f>>& CurveAttribute = InCollection.AddAttribute<TArray<FVector3f>>(CurveAttributeName, BranchCurveGroupName);
+		for (int32 i = 0; i < BranchCurves.Num(); ++i)
+		{
+			const TArray<int32>& BranchCurve = BranchCurves[i];
+			CurveAttribute[i].Reserve(BranchCurve.Num());
+			for (int32 VertexIndex : BranchCurve)
+			{
+				CurveAttribute[i].Add(GetVertexPosFunc(VertexIndex));
+			}
+		}
+	};
+
+	TArray<TArray<int32>> BranchCurves;
+	Internal::DrawBranches(DynamicMesh, BranchCurves);
+	if (InOptions.DebugState == FImportMeshOptions::EDebugState::DrawBranches)
+	{
+		AddBranchCurvesToCollecton(OutCollection, BranchCurves, [&](int32 VertexIndex) { return FVector3f(DynamicMesh.GetVertex(VertexIndex)); });
+		return ReturnWithDebugMesh(ReferenceMesh);
+	}
+
+	Internal::RemoveOverlaps(DynamicMesh, BranchCurves);
+	if (InOptions.DebugState == FImportMeshOptions::EDebugState::RemoveOverlaps)
+	{
+		AddBranchCurvesToCollecton(OutCollection, BranchCurves, [&](int32 VertexIndex) { return FVector3f(DynamicMesh.GetVertex(VertexIndex)); });
+		return ReturnWithDebugMesh(ReferenceMesh);
+	}
+
+	TArray<FVector3f> Points;
+	TArray<float> PointsRadii;
+	TArray<TArray<int32>> BranchPointIndices;
+	Internal2::ExtractBranchesFromBranchCurves(DynamicMesh, BranchCurves, Points, PointsRadii, BranchPointIndices);
+	if (InOptions.DebugState == FImportMeshOptions::EDebugState::ExtractBranches)
+	{
+		const static FName PointsGroupName = TEXT("Points");
+		const static FName PositionAttributeName = TEXT("Position");
+		const static FName ScaleAttributeName = TEXT("Scale");
+
+		OutCollection.AddGroup(PointsGroupName);
+		OutCollection.AddElements(Points.Num(), PointsGroupName);
+
+		TManagedArray<FVector3f>& PositionAttribute = OutCollection.AddAttribute<FVector3f>(PositionAttributeName, PointsGroupName);
+		TManagedArray<float>& ScaleAttribute = OutCollection.AddAttribute<float>(ScaleAttributeName, PointsGroupName);
+		for (int32 i = 0; i < Points.Num(); ++i)
+		{
+			PositionAttribute[i] = Points[i];
+			ScaleAttribute[i] = PointsRadii[i];
+		}
+
+		AddBranchCurvesToCollecton(OutCollection, BranchPointIndices, [&](int32 VertexIndex) { return FVector3f(Points[VertexIndex]); });
+
+		return ReturnWithDebugMesh(ReferenceMesh);
+	}
+
+	FPVPlantDescription PlantDescription;
+	PlantDescription.BranchHierarchy = PV::Import::Helper::CreateBranchHierarchyFromPoints(Points, PointsRadii, BranchPointIndices);
+
+	if (!PV::Import::Helper::GenerateGrowthDataFromPlantDescription(OutCollection, PlantDescription))
+	{
+		return PV::Import::StaticMesh::EImportResult::FailedToGenerateGrowthData;
+	}
+
+	return PV::Import::StaticMesh::EImportResult::Success;
+}
diff --git a/Engine/Plugins/Experimental/ProceduralVegetationEditor/Source/ProceduralVegetationEditor/Private/Importer/PVImporter_StaticMesh.h b/Engine/Plugins/Experimental/ProceduralVegetationEditor/Source/ProceduralVegetationEditor/Private/Importer/PVImporter_StaticMesh.h
new file mode 100644
--- /dev/null
+++ b/Engine/Plugins/Experimental/ProceduralVegetationEditor/Source/ProceduralVegetationEditor/Private/Importer/PVImporter_StaticMesh.h
@@ -0,0 +1,50 @@
+// Copyright Epic Games, Inc. All Rights Reserved.
+#pragma once
+
+#include "CoreMinimal.h"
+
+#include "GeometryCollection/ManagedArrayCollection.h"
+
+class UStaticMesh;
+
+namespace PV::Import::StaticMesh
+{
+	enum class EImportResult
+	{
+		Success,
+		InvalidSourceMesh,
+		InvalidOutputPath,
+		InvalidMeshData,
+		MeshNotReady,
+		FailedToGenerateGrowthData,
+		Canceled
+	};
+
+	struct FImportMeshOptions
+	{
+		FName MaterialFilter;
+
+		enum class EDebugState : uint8
+		{
+			None,
+			InitialMesh,
+			KeepMaterial,
+			ComputeNormals,
+			AverageNormals,
+			SmoothNormalsWithLengthPreservation,
+			GetSpherical,
+			BlurRadius,
+			ExtractReductionCloud,
+			RadialFuse,
+			GroupPathPoints,
+			ComputeDist,
+			DrawBranches,
+			RemoveOverlaps,
+			ExtractBranches,
+			MAX
+		};
+		EDebugState DebugState = EDebugState::None;
+	};
+
+	EImportResult ImportGrowthDataFromStaticMesh(const TSoftObjectPtr<UStaticMesh>& InSourceStaticMesh, FManagedArrayCollection& OutCollection, const FImportMeshOptions& InOptions);
+};
diff --git a/Engine/Plugins/Experimental/ProceduralVegetationEditor/Source/ProceduralVegetationEditor/Private/Importer/PVMockTree.cpp b/Engine/Plugins/Experimental/ProceduralVegetationEditor/Source/ProceduralVegetationEditor/Private/Importer/PVMockTree.cpp
new file mode 100644
--- /dev/null
+++ b/Engine/Plugins/Experimental/ProceduralVegetationEditor/Source/ProceduralVegetationEditor/Private/Importer/PVMockTree.cpp
@@ -0,0 +1,11 @@
+// Copyright Epic Games, Inc. All Rights Reserved.
+
+#include "PVMockTree.h"
+#include "PVImporter_MockTree.h"
+#include "PVImportHelpers.h"
+#include "PVPlantDescription.h"
+
+TSharedRef<FPVBranchHierarchyDescription> PV::MockTree::CreateBranchHierarchyDescription()
+{
+	return MakeShared<FPVBranchHierarchyDescription>(PV::Import::Helper::CreateBranchHierarchyFromPoints(MockTree::Points, MockTree::PointsRadii, MockTree::BranchPointIndices));
+}
\ No newline at end of file
diff --git a/Engine/Plugins/Experimental/ProceduralVegetationEditor/Source/ProceduralVegetationEditor/Private/Importer/PVMockTree.h b/Engine/Plugins/Experimental/ProceduralVegetationEditor/Source/ProceduralVegetationEditor/Private/Importer/PVMockTree.h
new file mode 100644
--- /dev/null
+++ b/Engine/Plugins/Experimental/ProceduralVegetationEditor/Source/ProceduralVegetationEditor/Private/Importer/PVMockTree.h
@@ -0,0 +1,107 @@
+// Copyright Epic Games, Inc. All Rights Reserved.
+#pragma once
+
+#include "CoreMinimal.h"
+
+struct FPVBranchHierarchyDescription;
+
+namespace PV::MockTree
+{
+	/*
+	* Branch Indices:     Branch Numbers:
+	*           *B2		            *BN3
+	*   *B3 |  / 		    *BN4|  /
+	*    \  | / / *B1	     \  | / / *BN2
+	*     \ | \/		      \ | \/ 
+	*      \| /			       \| / 
+	*       |/ 			        |/ 
+	*       |			        |
+	*       | *B0		        | *BN1
+	* 
+	* 
+	* Point Indices:
+	*      5*  11*
+	* 		|   /     *9
+	* *6    |  /     /
+	* \     | *10   /
+	*  \   4*  \   /
+	*   \   |   \ /
+	*    \  |    *8
+	*     \ |   /
+	*      3*  *7
+	*       | /
+	*       |/
+	*       *2
+	*       |
+	*       *1
+	*       |
+	*       *0
+	*/
+
+	const FVector3f Point0  = {  00.0000,  00.0000, 000.0000 };
+	const FVector3f Point1  = {  00.0394, -00.1783, 046.2954 };
+	const FVector3f Point2  = {  00.0788, -00.3566, 092.5909 };
+	const FVector3f Point3  = {  01.1927,  00.9499, 125.4070 };
+	const FVector3f Point4  = {  03.0192,  01.8644, 161.3479 };
+	const FVector3f Point5  = {  04.8458,  02.7790, 197.2887 };
+	const FVector3f Point6  = {  19.7899,  73.3726, 171.2636 };
+	const FVector3f Point7  = { -23.0204, -29.6581, 148.4679 };
+	const FVector3f Point8  = { -24.5249, -31.0789, 156.2391 };
+	const FVector3f Point9  = { -30.1656, -36.4065, 185.3766 };
+	const FVector3f Point10 = { -24.4424, -21.5258, 175.3671 };
+	const FVector3f Point11 = { -23.0204, -23.7218, 201.3998 };
+
+	const float Point0_Radius = 10.f;
+	const float Point1_Radius = 0.f;
+	const float Point2_Radius = 0.f;
+	const float Point3_Radius = 0.f;
+	const float Point4_Radius = 0.f;
+	const float Point5_Radius = 0.f;
+	const float Point6_Radius = 0.f;
+	const float Point7_Radius = 0.f;
+	const float Point8_Radius = 0.f;
+	const float Point9_Radius = 0.f;
+	const float Point10_Radius = 0.f;
+	const float Point11_Radius = 0.f;
+
+	const float Point0To1_Dist   = 46.2958;
+	const float Point1To2_Dist   = 46.2958;
+	const float Point2To3_Dist   = 32.861;
+	const float Point3To4_Dist   = 35.9989;
+	const float Point4To5_Dist   = 35.9989;
+	const float Point3To6_Dist   = 87.7139;
+	const float Point2To7_Dist   = 67.1893;
+	const float Point7To8_Dist   = 08.0419;
+	const float Point8To9_Dist   = 30.1529;
+	const float Point8To10_Dist  = 21.3811;
+	const float Point10To11_Dist = 26.1638;
+
+	const float Point0_LengthFromRoot = 0;
+	const float Point1_LengthFromRoot = Point0To1_Dist;
+	const float Point2_LengthFromRoot = Point1_LengthFromRoot + Point1To2_Dist;
+	const float Point3_LengthFromRoot = Point2_LengthFromRoot + Point2To3_Dist;
+	const float Point4_LengthFromRoot = Point3_LengthFromRoot + Point3To4_Dist;
+	const float Point5_LengthFromRoot = Point4_LengthFromRoot + Point4To5_Dist;
+	const float Point6_LengthFromRoot = Point3_LengthFromRoot + Point3To6_Dist;
+	const float Point7_LengthFromRoot = Point2_LengthFromRoot + Point2To7_Dist;
+	const float Point8_LengthFromRoot = Point7_LengthFromRoot + Point7To8_Dist;
+	const float Point9_LengthFromRoot = Point8_LengthFromRoot + Point8To9_Dist;
+	const float Point10_LengthFromRoot = Point8_LengthFromRoot + Point8To10_Dist;
+	const float Point11_LengthFromRoot = Point10_LengthFromRoot + Point10To11_Dist;
+
+	const TArray<int32> Branch0_PointIndices = { 0, 1, 2, 3, 4, 5 };
+	const TArray<int32> Branch1_PointIndices = { 2, 7, 8, 9 };
+	const TArray<int32> Branch2_PointIndices = { 8, 10, 11 };
+	const TArray<int32> Branch3_PointIndices = { 3, 6 };
+
+	const int32 Branch0_BranchNumber = 1;
+	const int32 Branch1_BranchNumber = 2;
+	const int32 Branch2_BranchNumber = 3;
+	const int32 Branch3_BranchNumber = 4;
+
+	const TArray<FVector3f> Points = { Point0, Point1, Point2, Point3, Point4, Point5, Point6, Point7, Point8, Point9, Point10, Point11 };
+	const TArray<float> PointsRadii = { Point0_Radius, Point1_Radius, Point2_Radius, Point3_Radius, Point4_Radius, Point5_Radius, Point6_Radius, Point7_Radius, Point8_Radius, Point9_Radius, Point10_Radius, Point11_Radius };
+	const TArray<TArray<int32>> BranchPointIndices = { Branch0_PointIndices, Branch1_PointIndices, Branch2_PointIndices, Branch3_PointIndices };
+
+	extern TSharedRef<FPVBranchHierarchyDescription> CreateBranchHierarchyDescription();
+}
\ No newline at end of file
diff --git a/Engine/Plugins/Experimental/ProceduralVegetationEditor/Source/ProceduralVegetationEditor/Private/Importer/PVPlantDescription.cpp b/Engine/Plugins/Experimental/ProceduralVegetationEditor/Source/ProceduralVegetationEditor/Private/Importer/PVPlantDescription.cpp
new file mode 100644
--- /dev/null
+++ b/Engine/Plugins/Experimental/ProceduralVegetationEditor/Source/ProceduralVegetationEditor/Private/Importer/PVPlantDescription.cpp
@@ -0,0 +1,259 @@
+// Copyright Epic Games, Inc. All Rights Reserved.
+
+#include "PVPlantDescription.h"
+
+#include "Engine/StaticMesh.h"
+#include "PrimitiveSceneProxy.h"
+#include "PrimitiveDrawInterface.h"
+#include "MeshDescription.h"
+#include "MeshElementCollector.h"
+
+FPVBranchHierarchyDescription::EValidateHierarchyResult FPVBranchHierarchyDescription::ValidateHierarchy() const
+{
+	TSet<int32> FoundBranches;
+	FoundBranches.Reserve(Branches.Num());
+
+	const TFunction<bool(int32)> WalkHierarchyRecursive = [&](int32 BranchIndex)->bool
+	{
+		bool bAlreadyAdded = false;
+		FoundBranches.Add(BranchIndex, &bAlreadyAdded);
+
+		if (bAlreadyAdded)
+		{
+			return false;
+		}
+
+		for (int32 ChildIndex : Branches[BranchIndex].ChildBranchIndices)
+		{
+			if (!WalkHierarchyRecursive(ChildIndex))
+			{
+				return false;
+			}
+		}
+
+		return true;
+	};
+
+	if (!WalkHierarchyRecursive(RootBranchIndex))
+	{
+		return EValidateHierarchyResult::CircularHiearachy;
+	}
+
+	if (FoundBranches.Num() != Branches.Num())
+	{
+		return EValidateHierarchyResult::DisjointedHiearachy;
+	}
+
+	return EValidateHierarchyResult::Success;
+}
+
+FPrimitiveSceneProxy* UPVPlantDescriptionVisualizationComponent::CreateSceneProxy()
+{
+	class FPVPlantDescriptionSceneProxy final : public FPrimitiveSceneProxy
+	{
+	private:
+		const FPVPlantDescription PlantDescription;
+		TArray<FMeshDescription> FoliageMeshDescriptions;
+	public:
+		FPVPlantDescriptionSceneProxy(const UPVPlantDescriptionVisualizationComponent* InComponent)
+			: FPrimitiveSceneProxy(InComponent)
+			, PlantDescription(InComponent->PlantDescriptionDataAsset ? InComponent->PlantDescriptionDataAsset->PlantDescription : FPVPlantDescription())
+		{
+			for (TObjectPtr<UStaticMesh> FoliageMesh : PlantDescription.Foliage.FoliageMeshes)
+			{
+				FMeshDescription& MeshDescription = FoliageMeshDescriptions.AddDefaulted_GetRef();
+
+				if (FoliageMesh)
+				{
+					MeshDescription = *FoliageMesh->GetMeshDescription(0);
+				}
+			}
+		}
+
+		virtual FPrimitiveViewRelevance GetViewRelevance(const FSceneView* View) const override
+		{
+			FPrimitiveViewRelevance ViewRelevance;
+			ViewRelevance.bDrawRelevance = IsShown(View);
+			ViewRelevance.bDynamicRelevance = true;
+			ViewRelevance.bSeparateTranslucency = ViewRelevance.bNormalTranslucency = true;
+			return ViewRelevance;
+		}
+
+		static void DrawDebugCircle(FPrimitiveDrawInterface* PDI, const FVector& Center, const FQuat& Rotation, float Radius, int32 Segments, const FColor& Color, uint8 DepthPriority, float Thickness = 0.f)
+		{
+			Segments = FMath::Max(Segments, 4);
+			const float AngleStep = 2.f * UE_PI / float(Segments);
+
+			const FVector AxisY = Rotation.GetRightVector();
+			const FVector AxisZ = Rotation.GetUpVector();
+
+			float Angle = 0.f;
+			while (Segments--)
+			{
+				const FVector Vertex1 = Center + Radius * Rotation.RotateVector(FVector(0, FMath::Cos(Angle), 0) + FVector(0, 0, FMath::Sin(Angle)));
+				Angle += AngleStep;
+				const FVector Vertex2 = Center + Radius * Rotation.RotateVector(FVector(0, FMath::Cos(Angle), 0) + FVector(0, 0, FMath::Sin(Angle)));
+				PDI->DrawLine(Vertex1, Vertex2, Color, Thickness, DepthPriority);
+			}
+		}
+
+		void DrawBranchHierarchy(FPrimitiveDrawInterface* PDI) const
+		{
+			const static TFunction<void(FPrimitiveDrawInterface*, const FPVBranchHierarchyDescription&, int32, int32)> DrawBranchesRecursive = 
+				[](FPrimitiveDrawInterface* PDI, const FPVBranchHierarchyDescription& BranchHierarchy, int32 BranchIndex, int32 Depth)
+			{
+				const FPVBranchDescription& Branch = BranchHierarchy.Branches[BranchIndex];
+
+				const FLinearColor DrawColor = FLinearColor::MakeRandomSeededColor(BranchIndex);
+
+				for (int32 i = 0; i < Branch.PointIndices.Num(); ++i)
+				{
+					const int32 PointIndex = Branch.PointIndices[i];
+					const FVector3f& Point = BranchHierarchy.Points[PointIndex];
+
+					PDI->DrawPoint(FVector(Point), DrawColor, 10.f, SDPG_Foreground);
+
+					if (i != 0)
+					{
+						if (BranchHierarchy.PointsRadii.IsValidIndex(PointIndex) && BranchHierarchy.PointsRadii[PointIndex] > 0)
+						{
+							const FVector3f ApicalDirection = [&]()
+								{
+									if (i < Branch.PointIndices.Num() - 1)
+									{
+										return (Point - BranchHierarchy.Points[Branch.PointIndices[i + 1]]).GetSafeNormal();
+									}
+									else if (i > 0)
+									{
+										return (BranchHierarchy.Points[Branch.PointIndices[i - 1]] - Point).GetSafeNormal();
+									}
+
+									return FVector3f::UpVector;
+								}();
+
+
+							DrawDebugCircle(PDI, FVector(Point), FRotationMatrix::MakeFromX(FVector(ApicalDirection)).ToQuat(), BranchHierarchy.PointsRadii[PointIndex], 8, FColor::Red, SDPG_Foreground, 2);
+						}
+					}
+
+					if (i != Branch.PointIndices.Num() - 1)
+					{
+						const int32 PointIndex2 = Branch.PointIndices[i + 1];
+						const FVector3f& Point2 = BranchHierarchy.Points[PointIndex2];
+
+						PDI->DrawLine(FVector(Point), FVector(Point2), DrawColor, SDPG_Foreground, 1.f);
+					}
+				}
+
+				for (int32 ChildIndex : Branch.ChildBranchIndices)
+				{
+					DrawBranchesRecursive(PDI, BranchHierarchy, ChildIndex, Depth + 1);
+				}
+			};
+
+			if (PlantDescription.BranchHierarchy.RootBranchIndex != INDEX_NONE)
+			{
+				DrawBranchesRecursive(PDI, PlantDescription.BranchHierarchy, PlantDescription.BranchHierarchy.RootBranchIndex, 0);
+			}
+		}
+		
+		void DrawFoliage(FPrimitiveDrawInterface* PDI) const
+		{
+			for (const FPVFoliagePointsDescription& FoliagePointsDescription : PlantDescription.Foliage.FoliagePoints)
+			{
+				for (int32 i = 0; i < FoliagePointsDescription.Positions.Num(); ++i)
+				{
+					const FVector3f Position = FoliagePointsDescription.Positions[i];
+					const FQuat4f Orientation = FoliagePointsDescription.Orientations[i];
+					const FVector3f Scale = FoliagePointsDescription.Scales[i];
+					const FMeshDescription& MeshDescription = FoliageMeshDescriptions[FoliagePointsDescription.FoliageIndices[i]];
+					
+					for (FTriangleID TriangleID : MeshDescription.Triangles().GetElementIDs())
+					{
+						const TConstArrayView<FVertexInstanceID> VertexInstances = MeshDescription.GetTriangleVertexInstances(TriangleID);
+
+						const FVertexID VertexID0 = MeshDescription.GetVertexInstanceVertex(VertexInstances[0]);
+						const FVertexID VertexID1 = MeshDescription.GetVertexInstanceVertex(VertexInstances[1]);
+						const FVertexID VertexID2 = MeshDescription.GetVertexInstanceVertex(VertexInstances[2]);
+
+						FVector3f Vertex0Position = MeshDescription.GetVertexPosition(VertexID0);
+						FVector3f Vertex1Position = MeshDescription.GetVertexPosition(VertexID1);
+						FVector3f Vertex2Position = MeshDescription.GetVertexPosition(VertexID2);
+
+						Vertex0Position *= Scale;
+						Vertex0Position = Orientation.RotateVector(Vertex0Position);
+						Vertex0Position += Position;
+
+						Vertex1Position *= Scale;
+						Vertex1Position = Orientation.RotateVector(Vertex1Position);
+						Vertex1Position += Position;
+
+						Vertex2Position *= Scale;
+						Vertex2Position = Orientation.RotateVector(Vertex2Position);
+						Vertex2Position += Position;
+
+						PDI->DrawLine(FVector(Vertex0Position), FVector(Vertex1Position), FLinearColor(FColor::Orange), SDPG_Foreground, 1);
+						PDI->DrawLine(FVector(Vertex1Position), FVector(Vertex2Position), FLinearColor(FColor::Orange), SDPG_Foreground, 1);
+						PDI->DrawLine(FVector(Vertex2Position), FVector(Vertex0Position), FLinearColor(FColor::Orange), SDPG_Foreground, 1);
+					}
+				}
+			}
+		}
+
+		virtual void GetDynamicMeshElements(
+			const TArray<const FSceneView*>& Views,
+			const FSceneViewFamily& ViewFamily,
+			uint32 VisibilityMap,
+			FMeshElementCollector& Collector
+		) const override
+		{
+			for (int32 ViewIndex = 0; ViewIndex < Views.Num(); ViewIndex++)
+			{
+				if (VisibilityMap & (1 << ViewIndex))
+				{
+					const FSceneView* View = Views[ViewIndex];
+
+#if UE_ENABLE_DEBUG_DRAWING
+					FPrimitiveDrawInterface* PDI = Collector.GetDebugPDI(ViewIndex);
+#else	
+					FPrimitiveDrawInterface* PDI = Collector.GetPDI(ViewIndex);
+#endif
+
+					DrawBranchHierarchy(PDI);
+					DrawFoliage(PDI);
+				}
+			}
+		}
+
+		virtual SIZE_T GetTypeHash() const override
+		{
+			static size_t UniquePointer;
+			return reinterpret_cast<size_t>(&UniquePointer);
+		}
+
+		virtual uint32 GetMemoryFootprint(void) const override
+		{
+			return sizeof * this + GetAllocatedSize();
+		}
+
+		virtual bool CanBeOccluded() const override { return false; }
+	};
+
+	FPVPlantDescriptionSceneProxy* Proxy = new FPVPlantDescriptionSceneProxy(this);
+	return Proxy;
+}
+
+FBoxSphereBounds UPVPlantDescriptionVisualizationComponent::CalcBounds(const FTransform& LocalToWorld) const
+{
+	FBox BBox = FBox(ForceInit);
+
+	if (IsValid(PlantDescriptionDataAsset))
+	{
+		for (const FVector3f& Point : PlantDescriptionDataAsset->PlantDescription.BranchHierarchy.Points)
+		{
+			BBox += FVector(Point);
+		}
+	}
+
+	return FBoxSphereBounds(BBox).TransformBy(LocalToWorld);
+}
\ No newline at end of file
diff --git a/Engine/Plugins/Experimental/ProceduralVegetationEditor/Source/ProceduralVegetationEditor/Private/Importer/PVPlantDescription.h b/Engine/Plugins/Experimental/ProceduralVegetationEditor/Source/ProceduralVegetationEditor/Private/Importer/PVPlantDescription.h
new file mode 100644
--- /dev/null
+++ b/Engine/Plugins/Experimental/ProceduralVegetationEditor/Source/ProceduralVegetationEditor/Private/Importer/PVPlantDescription.h
@@ -0,0 +1,147 @@
+// Copyright Epic Games, Inc. All Rights Reserved.
+#pragma once
+
+#include "CoreMinimal.h"
+#include "Engine/DataAsset.h"
+#include "Components/PrimitiveComponent.h"
+
+#include "PVPlantDescription.generated.h"
+
+USTRUCT()
+struct FPVBranchDescription
+{
+	GENERATED_BODY()
+
+	UPROPERTY(VisibleAnywhere, Category = "Branch Data")
+	int32 ParentBranchIndex = INDEX_NONE;
+
+	UPROPERTY(VisibleAnywhere, Category = "Branch Data")
+	int32 BranchIndex = INDEX_NONE;
+
+	UPROPERTY(VisibleAnywhere, Category = "Branch Data")
+	TArray<int32> ChildBranchIndices;
+
+	UPROPERTY(VisibleAnywhere, Category = "Branch Data")
+	TArray<int32> PointIndices;
+};
+
+USTRUCT()
+struct FPVBranchHierarchyDescription
+{
+	GENERATED_BODY()
+
+	UPROPERTY(VisibleAnywhere, Category = "Branch Data")
+	int32 RootBranchIndex = INDEX_NONE;
+
+	UPROPERTY(VisibleAnywhere, Category = "Branch Data")
+	TArray<FPVBranchDescription> Branches;
+
+	UPROPERTY(VisibleAnywhere, Category = "Branch Data")
+	TArray<FVector3f> Points;
+
+	UPROPERTY(VisibleAnywhere, Category = "Branch Data")
+	TArray<float> PointsRadii;
+
+	int32 GetNumBranches() const
+	{
+		return Branches.Num();
+	}
+
+	int32 GetNumBranchPoints() const
+	{
+		return Points.Num();
+	}
+
+	enum class EValidateHierarchyResult
+	{
+		DisjointedHiearachy,
+		CircularHiearachy,
+		Success
+	};
+	EValidateHierarchyResult ValidateHierarchy() const;
+};
+
+USTRUCT()
+struct FPVFoliagePointsDescription
+{
+	GENERATED_BODY()
+
+	UPROPERTY(VisibleAnywhere, Category = "Foliage Data")
+	TArray<FVector3f> Positions;
+
+	UPROPERTY(VisibleAnywhere, Category = "Foliage Data")
+	TArray<FQuat4f> Orientations;
+
+	UPROPERTY(VisibleAnywhere, Category = "Foliage Data")
+	TArray<FVector3f> Scales;
+
+	UPROPERTY(VisibleAnywhere, Category = "Foliage Data")
+	TArray<int32> FoliageIndices;
+};
+
+USTRUCT()
+struct FPVFoliageDescription
+{
+	GENERATED_BODY()
+
+	UPROPERTY(VisibleAnywhere, Category = "Foliage Data")
+	TArray<FPVFoliagePointsDescription> FoliagePoints;
+
+	UPROPERTY(VisibleAnywhere, Category = "Foliage Data")
+	TArray<TObjectPtr<UStaticMesh>> FoliageMeshes;
+
+	int32 GetNumFoliagePoints() const
+	{
+		int32 OutNumFoliagePoints = 0;
+		for (const FPVFoliagePointsDescription& FoliagePointsDescription : FoliagePoints)
+		{
+			OutNumFoliagePoints += FoliagePointsDescription.Positions.Num();
+		}
+		return OutNumFoliagePoints;
+	}
+};
+
+/*
+* A Plant Description represents the minimal data required to extrapolate PV growth data from an existing plant
+*/
+USTRUCT()
+struct FPVPlantDescription
+{
+	GENERATED_BODY()
+
+	UPROPERTY(VisibleAnywhere, Category = "Plant Data")
+	FPVBranchHierarchyDescription BranchHierarchy;
+
+	UPROPERTY(VisibleAnywhere, Category = "Plant Data")
+	FPVFoliageDescription Foliage;
+
+	int32 GetNumBranchPoints() const { return BranchHierarchy.GetNumBranchPoints(); }
+
+	int32 GetNumFoliagePoints() const { return Foliage.GetNumFoliagePoints(); }
+};
+
+// Internal, for debugging only
+UCLASS(BlueprintType)
+class PROCEDURALVEGETATIONEDITOR_API UPVPlantDescriptionDataAsset : public UDataAsset
+{
+	GENERATED_BODY()
+public:
+	UPROPERTY(VisibleAnywhere, Category = "Branch Hierarchy")
+	FPVPlantDescription PlantDescription;
+};
+
+// Internal, for debugging only
+UCLASS(Blueprintable, BlueprintType, meta = (BlueprintSpawnableComponent), MinimalAPI)
+class UPVPlantDescriptionVisualizationComponent : public UPrimitiveComponent
+{
+	GENERATED_BODY()
+public:
+
+	UPROPERTY(EditAnywhere, Category = "Branch Hierarchy Asset")
+	TObjectPtr<UPVPlantDescriptionDataAsset> PlantDescriptionDataAsset;
+
+public:
+
+	virtual FPrimitiveSceneProxy* CreateSceneProxy() override;
+	virtual FBoxSphereBounds CalcBounds(const FTransform& LocalToWorld) const override;
+};
\ No newline at end of file
diff --git a/Engine/Plugins/Experimental/ProceduralVegetationEditor/Source/ProceduralVegetationEditor/Private/PVEditorInterfaceImpl.cpp b/Engine/Plugins/Experimental/ProceduralVegetationEditor/Source/ProceduralVegetationEditor/Private/PVEditorInterfaceImpl.cpp
new file mode 100644
--- /dev/null
+++ b/Engine/Plugins/Experimental/ProceduralVegetationEditor/Source/ProceduralVegetationEditor/Private/PVEditorInterfaceImpl.cpp
@@ -0,0 +1,22 @@
+// Copyright Epic Games, Inc. All Rights Reserved.
+
+#pragma once
+
+#include "PVEditorInterfaceImpl.h"
+
+#include "Importer/PVImporter_StaticMesh.h"
+#include "PVImportStaticMeshParams.h"
+
+#include "Engine/StaticMesh.h"
+
+bool PV::FPVEditorInterfaceImpl::ImportGrowthDataFromStaticMesh(const FPVImportStaticMeshParams& InParams, FManagedArrayCollection& OutCollection)
+{
+	PV::Import::StaticMesh::FImportMeshOptions ImportMeshOptions;
+	ImportMeshOptions.MaterialFilter = InParams.MaterialFilter;
+	
+	static_assert(static_cast<int32>(EPVImportStaticMeshDebugState::MAX) == static_cast<int32>(PV::Import::StaticMesh::FImportMeshOptions::EDebugState::MAX));
+	ImportMeshOptions.DebugState = static_cast<PV::Import::StaticMesh::FImportMeshOptions::EDebugState>(static_cast<int32>(InParams.DebugState));
+
+	const PV::Import::StaticMesh::EImportResult ImportResult = PV::Import::StaticMesh::ImportGrowthDataFromStaticMesh(InParams.StaticMeshAsset, OutCollection, ImportMeshOptions);
+	return ImportResult == PV::Import::StaticMesh::EImportResult::Success;
+}
\ No newline at end of file
diff --git a/Engine/Plugins/Experimental/ProceduralVegetationEditor/Source/ProceduralVegetationEditor/Private/PVEditorInterfaceImpl.h b/Engine/Plugins/Experimental/ProceduralVegetationEditor/Source/ProceduralVegetationEditor/Private/PVEditorInterfaceImpl.h
new file mode 100644
--- /dev/null
+++ b/Engine/Plugins/Experimental/ProceduralVegetationEditor/Source/ProceduralVegetationEditor/Private/PVEditorInterfaceImpl.h
@@ -0,0 +1,20 @@
+// Copyright Epic Games, Inc. All Rights Reserved.
+
+#pragma once
+
+#include "IPVEditorInterface.h"
+
+#include "GeometryCollection/ManagedArrayCollection.h"
+
+struct FPVImportStaticMeshParams;
+
+#if WITH_EDITOR
+namespace PV
+{
+	class FPVEditorInterfaceImpl : public IPVEditorInterface
+	{
+	public:
+		bool ImportGrowthDataFromStaticMesh(const FPVImportStaticMeshParams& InParams, FManagedArrayCollection& OutCollection) override;
+	};
+};
+#endif
diff --git a/Engine/Plugins/Experimental/ProceduralVegetationEditor/Source/ProceduralVegetationEditor/Private/ProceduralVegetationEditorModule.cpp b/Engine/Plugins/Experimental/ProceduralVegetationEditor/Source/ProceduralVegetationEditor/Private/ProceduralVegetationEditorModule.cpp
--- a/Engine/Plugins/Experimental/ProceduralVegetationEditor/Source/ProceduralVegetationEditor/Private/ProceduralVegetationEditorModule.cpp
+++ b/Engine/Plugins/Experimental/ProceduralVegetationEditor/Source/ProceduralVegetationEditor/Private/ProceduralVegetationEditorModule.cpp
@@ -1,7 +1,6 @@
 // Copyright Epic Games, Inc. All Rights Reserved.
 
 #include "ProceduralVegetationEditorModule.h"
-
 #include "ContentBrowserMenuContexts.h"
 #include "LevelEditor.h"
 #include "PCGModule.h"
@@ -9,6 +8,7 @@
 #include "PropertyEditorModule.h"
 #include "PVEditorCommands.h"
 #include "PVEditorSettings.h"
+#include "PVEditorInterfaceImpl.h"
 #include "PVFloatRamp.h"
 #include "Selection.h"
 #include "ToolMenuContext.h"
@@ -73,6 +73,9 @@ void FProceduralVegetationEditorModule::StartupModule()
 
 	FMessageLogModule& MessageLogModule = FModuleManager::LoadModuleChecked<FMessageLogModule>("MessageLog");
 	MessageLogModule.RegisterLogListing(PVEditor::MessageLogName, LOCTEXT("PVEditorLogLabel", "Procedural Vegetation Editor"));
+	
+	check(PV::GPVEditorInterface == nullptr);
+	PV::GPVEditorInterface = new PV::FPVEditorInterfaceImpl();
 }
 
 void FProceduralVegetationEditorModule::ShutdownModule()
@@ -95,6 +98,10 @@ void FProceduralVegetationEditorModule::ShutdownModule()
 	{
 		MessageLogModule->UnregisterLogListing(PVEditor::MessageLogName);
 	}
+
+	check(PV::GPVEditorInterface != nullptr);
+	delete PV::GPVEditorInterface;
+	PV::GPVEditorInterface = nullptr;
 }
 
 void FProceduralVegetationEditorModule::RegisterMenus()
