Add PVE curve debug visualization mode

#Synced-CL 48813970

diff --git a/Engine/Plugins/Experimental/ProceduralVegetationEditor/Source/ProceduralVegetation/Public/DataTypes/PVData.h b/Engine/Plugins/Experimental/ProceduralVegetationEditor/Source/ProceduralVegetation/Public/DataTypes/PVData.h
--- a/Engine/Plugins/Experimental/ProceduralVegetationEditor/Source/ProceduralVegetation/Public/DataTypes/PVData.h
+++ b/Engine/Plugins/Experimental/ProceduralVegetationEditor/Source/ProceduralVegetation/Public/DataTypes/PVData.h
@@ -45,7 +45,8 @@ enum class EPVDebugValueVisualizationMode
 	Direction,
 	Sphere,
 	Vector,
-	Point
+	Point,
+	Curve
 };
 
 USTRUCT()
@@ -68,16 +69,19 @@ struct FPVVisualizationSettings
 	UPROPERTY(EditAnywhere, Category = DebugSettings, meta = (EditCondition = "DebugType==EPVDebugType::Custom"))
 	FName CustomGroupToFilter = TEXT("Custom");
 
-	UPROPERTY(EditAnywhere, Category = DebugSettings, meta = (EditCondition = "DebugType==EPVDebugType::Custom"))
+	UPROPERTY(EditAnywhere, Category = DebugSettings, meta = (EditCondition = "DebugType==EPVDebugType::Custom && VisualizationMode != EPVDebugValueVisualizationMode::Curve"))
 	FName CustomPivotPositionAttributeName = TEXT("Position");
 
-	UPROPERTY(EditAnywhere, Category = DebugSettings, meta = (EditCondition = "DebugType==EPVDebugType::Custom"))
+	UPROPERTY(EditAnywhere, Category = DebugSettings, meta = (EditCondition = "DebugType==EPVDebugType::Custom && VisualizationMode != EPVDebugValueVisualizationMode::Curve"))
 	FName CustomPivotScaleAttributeName = TEXT("Scale");
 
 	UPROPERTY(EditAnywhere, Category = DebugSettings)
 	EPVDebugValueVisualizationMode VisualizationMode = EPVDebugValueVisualizationMode::Text;
 
 	UPROPERTY(EditAnywhere, Category = DebugSettings)
+	bool bRandomizeColors = false;
+
+	UPROPERTY(EditAnywhere, Category = DebugSettings, meta = (EditCondition = "!bRandomizeColors"))
 	FColor Color = FColor::Black;
 
 	UPROPERTY(EditAnywhere, Category = DebugSettings, meta = (EditCondition = "bShowAnchorPoints"))
diff --git a/Engine/Plugins/Experimental/ProceduralVegetationEditor/Source/ProceduralVegetationEditor/Private/Visualizations/DebugVisualization/PVDebugVisualizationBase.cpp b/Engine/Plugins/Experimental/ProceduralVegetationEditor/Source/ProceduralVegetationEditor/Private/Visualizations/DebugVisualization/PVDebugVisualizationBase.cpp
--- a/Engine/Plugins/Experimental/ProceduralVegetationEditor/Source/ProceduralVegetationEditor/Private/Visualizations/DebugVisualization/PVDebugVisualizationBase.cpp
+++ b/Engine/Plugins/Experimental/ProceduralVegetationEditor/Source/ProceduralVegetationEditor/Private/Visualizations/DebugVisualization/PVDebugVisualizationBase.cpp
@@ -480,7 +480,8 @@ void FPVDebugVisualizationBase::DrawAttributes(FVisualizerDrawContext& InContext
 				float OutScale;
 				GetPivot(Collection, i, OutPos, OutScale);
 
-				DrawAsText(InContext.SceneSetupParams, FText::FromString(AttributeAsString[i]), OutPos, OutScale, VisualizationSettings.TextScale, VisualizationSettings.Color);
+				const FColor DrawColor = VisualizationSettings.bRandomizeColors ? FColor::MakeRandomColor() : VisualizationSettings.Color;
+				DrawAsText(InContext.SceneSetupParams, FText::FromString(AttributeAsString[i]), OutPos, OutScale, VisualizationSettings.TextScale, DrawColor);
 			}
 			break;
 		}
@@ -502,6 +503,8 @@ void FPVDebugVisualizationBase::DrawAttributes(FVisualizerDrawContext& InContext
 				float OutScale;
 				GetPivot(Collection, i, OutPos, OutScale);
 
+				const FColor DrawColor = VisualizationSettings.bRandomizeColors ? FColor::MakeRandomColor() : VisualizationSettings.Color;
+
 				for (int32 j = 0; j < AttributeAsVector[i].Num(); ++j)
 				{
 					const float BaseDirectionStrength = 8.f;
@@ -510,8 +513,8 @@ void FPVDebugVisualizationBase::DrawAttributes(FVisualizerDrawContext& InContext
 					const FVector StartPos(OutPos);
 					const FVector EndPos = StartPos + DirectionVector * BaseDirectionStrength * VisualizationSettings.GizmoScale;
 
-					const FColor DrawColor = FLinearColor::LerpUsingHSV(VisualizationSettings.Color.ReinterpretAsLinear(), FLinearColor::Black, j / (float)AttributeAsVector[i].Num()).ToFColor(true);
-					DrawAsDirection(LineBatchComponent, StartPos, EndPos, DrawColor);
+					const FColor LineColor = FLinearColor::LerpUsingHSV(DrawColor.ReinterpretAsLinear(), FLinearColor::Black, j / (float)AttributeAsVector[i].Num()).ToFColor(true);
+					DrawAsLine(LineBatchComponent, StartPos, EndPos, LineColor);
 
 					if (AttributeAsVector[i].Num() > 1)
 					{
@@ -556,7 +559,32 @@ void FPVDebugVisualizationBase::DrawAttributes(FVisualizerDrawContext& InContext
 				float OutScale;
 				GetPivot(Collection, i, OutPos, OutScale);
 
-				DrawAsSphere(LineBatchComponent, FVector(OutPos), AttributeAsScalar[i] * VisualizationSettings.GizmoScale, VisualizationSettings.Color);
+				const FColor DrawColor = VisualizationSettings.bRandomizeColors ? FColor::MakeRandomColor() : VisualizationSettings.Color;
+				DrawAsSphere(LineBatchComponent, FVector(OutPos), AttributeAsScalar[i] * VisualizationSettings.GizmoScale, DrawColor);
+			}
+			break;
+		}
+		case EPVDebugValueVisualizationMode::Curve:
+		{
+			using namespace PVE::AttributeToVector3;
+
+			const TArray<FVectorValue> AttributeAsVector = ConvertAttributeToVector3(Collection, AttributeName, GroupName);
+			if (AttributeAsVector.Num() != NumElements)
+			{
+				break;
+			}
+
+			TObjectPtr<UPVLineBatchComponent> LineBatchComponent = GetOrCreateLineComponent(InContext.SceneSetupParams);
+
+			for (int32 i = 0; i < NumElements; i++)
+			{
+				const FColor DrawColor = VisualizationSettings.bRandomizeColors ? FColor::MakeRandomColor() : VisualizationSettings.Color;
+				for (int32 j = 0; j < AttributeAsVector[i].Num() - 1; ++j)
+				{
+					const FVector StartPos = AttributeAsVector[i][j];
+					const FVector EndPos = AttributeAsVector[i][j + 1];
+					DrawAsLine(LineBatchComponent, StartPos, EndPos, DrawColor);
+				}
 			}
 			break;
 		}
@@ -604,7 +632,7 @@ void FPVDebugVisualizationBase::DrawAsText(FPCGSceneSetupParams& InOutParams, co
 	InOutParams.Scene->AddComponent(Component, FTransform(FVector(Pos.X, Pos.Y, Pos.Z)));
 }
 
-void FPVDebugVisualizationBase::DrawAsDirection(UPVLineBatchComponent* LineBatchComponent, const FVector& InStartPos, const FVector& InEndPos,
+void FPVDebugVisualizationBase::DrawAsLine(UPVLineBatchComponent* LineBatchComponent, const FVector& InStartPos, const FVector& InEndPos,
 	const FLinearColor& InColor)
 {
 	check(LineBatchComponent);
diff --git a/Engine/Plugins/Experimental/ProceduralVegetationEditor/Source/ProceduralVegetationEditor/Private/Visualizations/DebugVisualization/PVDebugVisualizationBase.h b/Engine/Plugins/Experimental/ProceduralVegetationEditor/Source/ProceduralVegetationEditor/Private/Visualizations/DebugVisualization/PVDebugVisualizationBase.h
--- a/Engine/Plugins/Experimental/ProceduralVegetationEditor/Source/ProceduralVegetationEditor/Private/Visualizations/DebugVisualization/PVDebugVisualizationBase.h
+++ b/Engine/Plugins/Experimental/ProceduralVegetationEditor/Source/ProceduralVegetationEditor/Private/Visualizations/DebugVisualization/PVDebugVisualizationBase.h
@@ -29,7 +29,7 @@ protected:
 	virtual void GetPivot(const FManagedArrayCollection& InCollection, const int InIndex, FVector3f& OutPos, float& OutScale) = 0;
 	
 	static void DrawAsText(FPCGSceneSetupParams& InOutParams, const FText& InTextToDraw, const FVector3f& InPos, float InScale, float TextSize, const FLinearColor& InColor = FLinearColor::White);
-	static void DrawAsDirection(UPVLineBatchComponent* LineBatchComponent, const FVector& InStartPos, const FVector& InEndPos, const FLinearColor& InColor = FLinearColor::White);
+	static void DrawAsLine(UPVLineBatchComponent* LineBatchComponent, const FVector& InStartPos, const FVector& InEndPos, const FLinearColor& InColor = FLinearColor::White);
 	static void DrawAsSphere(UPVLineBatchComponent* LineBatchComponent, const FVector& InPos, float Radius, const FLinearColor& InColor = FLinearColor::White);
 	static void DrawAsPoints(UPVLineBatchComponent* LineBatchComponent, const TArray<TPair<FVector, float>>& InPoints, const FLinearColor& InColor = FLinearColor::White);
 	static void DrawAsMeshPoints(UInstancedStaticMeshComponent* InstancedStaticMeshComponent, const TArray<FTransform>& InTransforms, const FLinearColor& InColor = FLinearColor::White);
