Fix linker error of PV::FacadeHelper::GetBranchPointChildPoints during monolithic compilation introduced in 48810753

#Synced-CL 48847314

diff --git a/Engine/Plugins/Experimental/ProceduralVegetationEditor/Source/ProceduralVegetation/Private/Helpers/PVFacadeHelper.cpp b/Engine/Plugins/Experimental/ProceduralVegetationEditor/Source/ProceduralVegetation/Private/Helpers/PVFacadeHelper.cpp
new file mode 100644
--- /dev/null
+++ b/Engine/Plugins/Experimental/ProceduralVegetationEditor/Source/ProceduralVegetation/Private/Helpers/PVFacadeHelper.cpp
@@ -0,0 +1,154 @@
+// Copyright Epic Games, Inc. All Rights Reserved.
+
+#include "Helpers/PVFacadeHelper.h"
+#include "Facades/PVFoliageFacade.h"
+
+TArray<int32> PV::FacadeHelper::GetBranchPointFoliageEntryIds(
+		const PV::Facades::FFoliageFacade& FoliageFacade,
+		const PV::Facades::FBranchFacade& BranchFacade,
+		const PV::Facades::FPointFacade& PointFacade,
+		int32 BranchIndex,
+		int32 PointIndex)
+{
+	const TArray<int32> BranchPoints = BranchFacade.GetPoints(BranchIndex);
+	const int32 PointBranchIndex = BranchPoints.Find(PointIndex);
+	if (PointBranchIndex == INDEX_NONE)
+	{
+		return TArray<int32>();
+	}
+
+	const int32 PreviousPointBranchIndex = PointBranchIndex > 0 ? PointBranchIndex - 1 : INDEX_NONE;
+	if (PreviousPointBranchIndex == INDEX_NONE)
+	{
+		return TArray<int32>();
+	}
+
+	const int32 PreviousPointIndex = BranchPoints[PreviousPointBranchIndex];
+
+	const float BranchPointLengthFromRoot = PointFacade.GetLengthFromRoot(PointIndex);
+	const float PreviousBranchPointLengthFromRoot = PointFacade.GetLengthFromRoot(PreviousPointIndex);
+
+	const TArray<int32> FoliageEntryIds = FoliageFacade.GetFoliageEntryIdsForBranch(BranchIndex);
+
+	TArray<int32> OutFoliageEntryIds;
+	OutFoliageEntryIds.Reserve(FoliageEntryIds.Num());
+
+	for (const int32 FoliageEntryId : FoliageEntryIds)
+	{
+		const float FoliageEntryLengthFromRoot = FoliageFacade.GetLengthFromRoot(FoliageEntryId);
+
+		if (FoliageEntryLengthFromRoot <= BranchPointLengthFromRoot && FoliageEntryLengthFromRoot > PreviousBranchPointLengthFromRoot)
+		{
+			OutFoliageEntryIds.Add(FoliageEntryId);
+		}
+	}
+
+	return OutFoliageEntryIds;
+}
+
+template<> TArray<int32> PV::FacadeHelper::GetBranchPointChildPoints<PV::FacadeHelper::EWalkBranchMode::FusedPoints>(
+	const PV::Facades::FBranchFacade& BranchFacade, 
+	const PV::Facades::FPointFacade& PointFacade, 
+	int32 BranchIndex,
+	int32 PointIndex
+)
+{
+	const TArray<int32>& BranchChildren = BranchFacade.GetChildren(BranchIndex);
+	const TManagedArray<int32>& BranchNumbers = BranchFacade.GetBranchNumbers();
+
+	TArray<int32> ChildrenIndices;
+	ChildrenIndices.Reserve(BranchChildren.Num());
+
+	for (const int32 BranchChild : BranchChildren)
+	{
+		const int32 ChildIndex = BranchNumbers.Find(BranchChild);
+		if (ChildIndex == INDEX_NONE)
+		{
+			continue;
+		}
+
+		const int32 ParentIndex = BranchFacade.GetParentIndex(ChildIndex);
+		if (ParentIndex != BranchIndex)
+		{
+			continue;
+		}
+
+		const TArray<int32>& ChildBranchPoints = BranchFacade.GetPoints(ChildIndex);
+		if (ChildBranchPoints.Num() == 0)
+		{
+			continue;
+		}
+
+		const int32 FirstChildPointIndex = ChildBranchPoints[0];
+		if (FirstChildPointIndex == PointIndex)
+		{
+			ChildrenIndices.Add(ChildIndex);
+		}
+	}
+
+	return ChildrenIndices;
+};
+
+template<> TArray<int32> PV::FacadeHelper::GetBranchPointChildPoints<PV::FacadeHelper::EWalkBranchMode::SplitPoints>(
+	const PV::Facades::FBranchFacade& BranchFacade, 
+	const PV::Facades::FPointFacade& PointFacade, 
+	int32 BranchIndex, 
+	int32 PointIndex
+)
+{
+	const TArray<int32> BranchPoints = BranchFacade.GetPoints(BranchIndex);
+	const int32 PointBranchIndex = BranchPoints.Find(PointIndex);
+	if (PointBranchIndex == INDEX_NONE)
+	{
+		return TArray<int32>();
+	}
+
+	const int32 PreviousPointBranchIndex = PointBranchIndex > 0 ? PointBranchIndex - 1 : INDEX_NONE;
+	if (PreviousPointBranchIndex == INDEX_NONE)
+	{
+		return TArray<int32>();
+	}
+
+	const int32 PreviousPointIndex = BranchPoints[PreviousPointBranchIndex];
+
+	const TArray<int32>& BranchChildren = BranchFacade.GetChildren(BranchIndex);
+	const TManagedArray<int32>& BranchNumbers = BranchFacade.GetBranchNumbers();
+
+	const float BranchPointLengthFromRoot = PointFacade.GetLengthFromRoot(PointIndex);
+	const float PreviousBranchPointLengthFromRoot = PointFacade.GetLengthFromRoot(PreviousPointIndex);
+
+	TArray<int32> ChildrenIndices;
+	ChildrenIndices.Reserve(BranchChildren.Num());
+
+	for (const int32 BranchChild : BranchChildren)
+	{
+		const int32 ChildIndex = BranchNumbers.Find(BranchChild);
+		if (ChildIndex == INDEX_NONE)
+		{
+			continue;
+		}
+
+		const int32 ParentIndex = BranchFacade.GetParentIndex(ChildIndex);
+		if (ParentIndex != BranchIndex)
+		{
+			continue;
+		}
+
+		const TArray<int32>& ChildBranchPoints = BranchFacade.GetPoints(ChildIndex);
+		if (ChildBranchPoints.Num() == 0)
+		{
+			continue;
+		}
+
+		const int32 FirstChildPointIndex = ChildBranchPoints[0];
+		const float ChildLengthFromRoot = PointFacade.GetLengthFromRoot(FirstChildPointIndex);
+		if (ChildLengthFromRoot > BranchPointLengthFromRoot || ChildLengthFromRoot <= PreviousBranchPointLengthFromRoot)
+		{
+			continue;
+		}
+
+		ChildrenIndices.Add(ChildIndex);
+	}
+
+	return ChildrenIndices;
+};
diff --git a/Engine/Plugins/Experimental/ProceduralVegetationEditor/Source/ProceduralVegetation/Public/Helpers/PVFacadeHelper.h b/Engine/Plugins/Experimental/ProceduralVegetationEditor/Source/ProceduralVegetation/Public/Helpers/PVFacadeHelper.h
--- a/Engine/Plugins/Experimental/ProceduralVegetationEditor/Source/ProceduralVegetation/Public/Helpers/PVFacadeHelper.h
+++ b/Engine/Plugins/Experimental/ProceduralVegetationEditor/Source/ProceduralVegetation/Public/Helpers/PVFacadeHelper.h
@@ -2,10 +2,14 @@
 #pragma once
 
 #include "Facades/PVPlantFacade.h"
-#include "Facades/PVFoliageFacade.h"
 #include "Facades/PVBranchFacade.h"
 #include "Facades/PVPointFacade.h"
 
+namespace PV::Facades
+{
+	class FFoliageFacade;
+}
+
 namespace PV::FacadeHelper
 {
 	template <typename StateType>
@@ -48,14 +52,14 @@ namespace PV::FacadeHelper
 		}
 	}
 
-	void WalkBranches(
+	inline void WalkBranches(
 		const PV::Facades::FPlantFacade& PlantFacade, 
 		const PV::Facades::FBranchFacade& BranchFacade, 
 		const TFunction<bool(int32 BranchIndex, int32 Depth)>& Body
 	)
 	{
 		int32 Depth = 0;
-		WalkBranches<int32>(PlantFacade, BranchFacade, Depth, [&](int32 BranchIndex, int32& Depth) 
+		WalkBranches<int32>(PlantFacade, BranchFacade, Depth, [&Body](int32 BranchIndex, int32& Depth) 
 		{
 			if (!Body(BranchIndex, Depth))
 			{
@@ -66,7 +70,7 @@ namespace PV::FacadeHelper
 		});
 	}
 
-	void WalkBranches(
+	inline void WalkBranches(
 		const PV::Facades::FPlantFacade& PlantFacade, 
 		const PV::Facades::FBranchFacade& BranchFacade, 
 		const TFunction<bool(int32 BranchIndex)>& Body
@@ -82,49 +86,13 @@ namespace PV::FacadeHelper
 		});
 	}
 
-	TArray<int32> GetBranchPointFoliageEntryIds(
+	PROCEDURALVEGETATION_API extern TArray<int32> GetBranchPointFoliageEntryIds(
 		const PV::Facades::FFoliageFacade& FoliageFacade,
 		const PV::Facades::FBranchFacade& BranchFacade,
 		const PV::Facades::FPointFacade& PointFacade,
 		int32 BranchIndex,
-		int32 PointIndex)
-	{
-		const TArray<int32> BranchPoints = BranchFacade.GetPoints(BranchIndex);
-		const int32 PointBranchIndex = BranchPoints.Find(PointIndex);
-		if (PointBranchIndex == INDEX_NONE)
-		{
-			return TArray<int32>();
-		}
-
-		const int32 PreviousPointBranchIndex = PointBranchIndex > 0 ? PointBranchIndex - 1 : INDEX_NONE;
-		if (PreviousPointBranchIndex == INDEX_NONE)
-		{
-			return TArray<int32>();
-		}
-
-		const int32 PreviousPointIndex = BranchPoints[PreviousPointBranchIndex];
-
-		const float BranchPointLengthFromRoot = PointFacade.GetLengthFromRoot(PointIndex);
-		const float PreviousBranchPointLengthFromRoot = PointFacade.GetLengthFromRoot(PreviousPointIndex);
-
-		const TArray<int32> FoliageEntryIds = FoliageFacade.GetFoliageEntryIdsForBranch(BranchIndex);
-
-		TArray<int32> OutFoliageEntryIds;
-		OutFoliageEntryIds.Reserve(FoliageEntryIds.Num());
-
-		for (const int32 FoliageEntryId : FoliageEntryIds)
-		{
-			const float FoliageEntryLengthFromRoot = FoliageFacade.GetLengthFromRoot(FoliageEntryId);
-
-			if (FoliageEntryLengthFromRoot <= BranchPointLengthFromRoot && FoliageEntryLengthFromRoot > PreviousBranchPointLengthFromRoot)
-			{
-				OutFoliageEntryIds.Add(FoliageEntryId);
-			}
-		}
-
-		return OutFoliageEntryIds;
-	}
-
+		int32 PointIndex);
+	
 	enum EWalkBranchMode
 	{
 		FusedPoints,
@@ -134,112 +102,21 @@ namespace PV::FacadeHelper
 	template<EWalkBranchMode WalkBranchMode = EWalkBranchMode::SplitPoints>
 	TArray<int32> GetBranchPointChildPoints(const PV::Facades::FBranchFacade& BranchFacade, const PV::Facades::FPointFacade& PointFacade, int32 BranchIndex, int32 PointIndex);
 
-	template<> TArray<int32> GetBranchPointChildPoints<EWalkBranchMode::FusedPoints>(
+	template<>
+	PROCEDURALVEGETATION_API TArray<int32> GetBranchPointChildPoints<EWalkBranchMode::FusedPoints>(
 		const PV::Facades::FBranchFacade& BranchFacade, 
 		const PV::Facades::FPointFacade& PointFacade, 
 		int32 BranchIndex,
 		int32 PointIndex
-	)
-	{
-		const TArray<int32>& BranchChildren = BranchFacade.GetChildren(BranchIndex);
-		const TManagedArray<int32>& BranchNumbers = BranchFacade.GetBranchNumbers();
-
-		TArray<int32> ChildrenIndices;
-		ChildrenIndices.Reserve(BranchChildren.Num());
+	);
 
-		for (const int32 BranchChild : BranchChildren)
-		{
-			const int32 ChildIndex = BranchNumbers.Find(BranchChild);
-			if (ChildIndex == INDEX_NONE)
-			{
-				continue;
-			}
-
-			const int32 ParentIndex = BranchFacade.GetParentIndex(ChildIndex);
-			if (ParentIndex != BranchIndex)
-			{
-				continue;
-			}
-
-			const TArray<int32>& ChildBranchPoints = BranchFacade.GetPoints(ChildIndex);
-			if (ChildBranchPoints.Num() == 0)
-			{
-				continue;
-			}
-
-			const int32 FirstChildPointIndex = ChildBranchPoints[0];
-			if (FirstChildPointIndex == PointIndex)
-			{
-				ChildrenIndices.Add(ChildIndex);
-			}
-		}
-
-		return ChildrenIndices;
-	};
-
-	template<> TArray<int32> GetBranchPointChildPoints<EWalkBranchMode::SplitPoints>(
+	template<>
+	PROCEDURALVEGETATION_API TArray<int32> GetBranchPointChildPoints<EWalkBranchMode::SplitPoints>(
 		const PV::Facades::FBranchFacade& BranchFacade, 
 		const PV::Facades::FPointFacade& PointFacade, 
 		int32 BranchIndex, 
 		int32 PointIndex
-	)
-	{
-		const TArray<int32> BranchPoints = BranchFacade.GetPoints(BranchIndex);
-		const int32 PointBranchIndex = BranchPoints.Find(PointIndex);
-		if (PointBranchIndex == INDEX_NONE)
-		{
-			return TArray<int32>();
-		}
-
-		const int32 PreviousPointBranchIndex = PointBranchIndex > 0 ? PointBranchIndex - 1 : INDEX_NONE;
-		if (PreviousPointBranchIndex == INDEX_NONE)
-		{
-			return TArray<int32>();
-		}
-
-		const int32 PreviousPointIndex = BranchPoints[PreviousPointBranchIndex];
-
-		const TArray<int32>& BranchChildren = BranchFacade.GetChildren(BranchIndex);
-		const TManagedArray<int32>& BranchNumbers = BranchFacade.GetBranchNumbers();
-
-		const float BranchPointLengthFromRoot = PointFacade.GetLengthFromRoot(PointIndex);
-		const float PreviousBranchPointLengthFromRoot = PointFacade.GetLengthFromRoot(PreviousPointIndex);
-
-		TArray<int32> ChildrenIndices;
-		ChildrenIndices.Reserve(BranchChildren.Num());
-
-		for (const int32 BranchChild : BranchChildren)
-		{
-			const int32 ChildIndex = BranchNumbers.Find(BranchChild);
-			if (ChildIndex == INDEX_NONE)
-			{
-				continue;
-			}
-
-			const int32 ParentIndex = BranchFacade.GetParentIndex(ChildIndex);
-			if (ParentIndex != BranchIndex)
-			{
-				continue;
-			}
-
-			const TArray<int32>& ChildBranchPoints = BranchFacade.GetPoints(ChildIndex);
-			if (ChildBranchPoints.Num() == 0)
-			{
-				continue;
-			}
-
-			const int32 FirstChildPointIndex = ChildBranchPoints[0];
-			const float ChildLengthFromRoot = PointFacade.GetLengthFromRoot(FirstChildPointIndex);
-			if (ChildLengthFromRoot > BranchPointLengthFromRoot || ChildLengthFromRoot <= PreviousBranchPointLengthFromRoot)
-			{
-				continue;
-			}
-
-			ChildrenIndices.Add(ChildIndex);
-		}
-
-		return ChildrenIndices;
-	};
+	);
 
 	template <typename StateType, EWalkBranchMode WalkBranchMode = EWalkBranchMode::SplitPoints>
 	bool WalkBranchPoints(
@@ -294,7 +171,7 @@ namespace PV::FacadeHelper
 			PointFacade,
 			BranchIndex,
 			0,
-			[&](int32 BranchIndex, int32 BranchPointIndex, int32 PointIndex, const TArray<int32>& ChildPointIndices, int32& PointState)
+			[&](int32 BranchIndex, int32 BranchPointIndex, int32 PointIndex, const TArray<int32>& ChildPointIndices, int32&)
 			{
 				return Body(BranchIndex, BranchPointIndex, PointIndex, ChildPointIndices);
 			}
@@ -329,7 +206,7 @@ namespace PV::FacadeHelper
 	{
 		WalkPlantPoints<int32, WalkBranchMode>(PlantFacade, BranchFacade, PointFacade, INDEX_NONE, [&](int32 BranchIndex, int32 BranchPointIndex, int32 PointIndex, const TArray<int32>& ChildPointIndices, int32&)
 		{
-			Body(BranchIndex, BranchPointIndex, PointIndex, ChildPointIndices);
+			return Body(BranchIndex, BranchPointIndex, PointIndex, ChildPointIndices);
 		});
 	}
 }
