[PVE] Fixed the error log for data asset
The error log for data json was getting overriden from the meta file json.

#Synced-CL 48842348

diff --git a/Engine/Plugins/Experimental/ProceduralVegetationEditor/Source/ProceduralVegetation/Private/ProceduralVegetationPreset.cpp b/Engine/Plugins/Experimental/ProceduralVegetationEditor/Source/ProceduralVegetation/Private/ProceduralVegetationPreset.cpp
--- a/Engine/Plugins/Experimental/ProceduralVegetationEditor/Source/ProceduralVegetation/Private/ProceduralVegetationPreset.cpp
+++ b/Engine/Plugins/Experimental/ProceduralVegetationEditor/Source/ProceduralVegetation/Private/ProceduralVegetationPreset.cpp
@@ -85,18 +85,25 @@ void UProceduralVegetationPreset::UpdateDataAsset()
 		FString FullPath = JsonDirectoryPath.Path / FileName;
 		FManagedArrayCollection Collection;
 		FString OutError;
+		FString MetaJsonError;
 		if ( PV::LoadMegaPlantsJsonToCollection(Collection, FullPath, OutError))
 		{
 			UE_LOG(LogProceduralVegetation, Log, TEXT("Variant %s loaded from file"), *FPaths::GetBaseFilename(FullPath));
 			Variants.Add(FPaths::GetBaseFilename(FullPath), Collection);
 		}
-		else if (TSharedPtr<FJsonObject> MetaJson = PV::LoadMetaFileIntoJsonObject(FullPath, OutError))
+		else if (TSharedPtr<FJsonObject> MetaJson = PV::LoadMetaFileIntoJsonObject(FullPath, MetaJsonError))
 		{
+			OutError.Empty();
 			MetaFiles.Add(MetaJson);
 		}
+		
+		if (!MetaJsonError.IsEmpty())
+		{
+			UE_LOG(LogProceduralVegetation, Warning, TEXT("Invalid meta json file %s %s"),*FullPath, *MetaJsonError);
+		}
 		if (!OutError.IsEmpty())
 		{
-			UE_LOG(LogProceduralVegetation, Warning, TEXT("%s"), *OutError);
+			UE_LOG(LogProceduralVegetation, Warning, TEXT("Invalid growth data json File %s %s"),*FullPath, *OutError);
 		}
 	}
 	
