[PVE] Added Offset and show option in debug visualization

#Synced-CL 48842348

diff --git a/Engine/Plugins/Experimental/ProceduralVegetationEditor/Source/ProceduralVegetation/Public/DataTypes/PVData.h b/Engine/Plugins/Experimental/ProceduralVegetationEditor/Source/ProceduralVegetation/Public/DataTypes/PVData.h
--- a/Engine/Plugins/Experimental/ProceduralVegetationEditor/Source/ProceduralVegetation/Public/DataTypes/PVData.h
+++ b/Engine/Plugins/Experimental/ProceduralVegetationEditor/Source/ProceduralVegetation/Public/DataTypes/PVData.h
@@ -57,6 +57,9 @@ struct FPVVisualizationSettings
 	UPROPERTY(EditAnywhere, Category = DebugSettings)
 	EPVDebugType DebugType = EPVDebugType::Point;
 
+	UPROPERTY(EditAnywhere, Category = DebugSettings)
+	bool bShow = true;
+
 	UPROPERTY(EditAnywhere, Category = DebugSettings)
 	bool bShowAnchorPoints = true;
 
@@ -90,6 +93,9 @@ struct FPVVisualizationSettings
 	UPROPERTY(EditAnywhere, Category = DebugSettings, meta = (EditCondition = "VisualizationMode==EPVDebugValueVisualizationMode::Text"))
 	float TextScale = 2.5f;
 
+	UPROPERTY(EditAnywhere, Category = DebugSettings, meta = (EditCondition = "VisualizationMode==EPVDebugValueVisualizationMode::Text"))
+	FVector3f TextOffset = FVector3f::Zero();
+
 	UPROPERTY(EditAnywhere, Category = DebugSettings)
 	float GizmoScale = 1.f;
 	
diff --git a/Engine/Plugins/Experimental/ProceduralVegetationEditor/Source/ProceduralVegetationEditor/Private/Visualizations/DebugVisualization/PVDebugVisualizationBase.cpp b/Engine/Plugins/Experimental/ProceduralVegetationEditor/Source/ProceduralVegetationEditor/Private/Visualizations/DebugVisualization/PVDebugVisualizationBase.cpp
--- a/Engine/Plugins/Experimental/ProceduralVegetationEditor/Source/ProceduralVegetationEditor/Private/Visualizations/DebugVisualization/PVDebugVisualizationBase.cpp
+++ b/Engine/Plugins/Experimental/ProceduralVegetationEditor/Source/ProceduralVegetationEditor/Private/Visualizations/DebugVisualization/PVDebugVisualizationBase.cpp
@@ -455,6 +455,11 @@ void FPVDebugVisualizationBase::DrawAttributes(FVisualizerDrawContext& InContext
 	const FName AttributeName = VisualizationSettings.AttributeToFilter;
 	const FName GroupName = VisualizationSettings.GetGroupToFilter();
 
+	if (!VisualizationSettings.bShow)
+	{
+		return;
+	}
+	
 	if (GroupName.IsNone() || AttributeName.IsNone() || !Collection.HasAttribute(AttributeName, GroupName) || !Collection.HasGroup(GroupName))
 	{
 		return;
@@ -481,7 +486,7 @@ void FPVDebugVisualizationBase::DrawAttributes(FVisualizerDrawContext& InContext
 				GetPivot(Collection, i, OutPos, OutScale);
 
 				const FColor DrawColor = VisualizationSettings.bRandomizeColors ? FColor::MakeRandomColor() : VisualizationSettings.Color;
-				DrawAsText(InContext.SceneSetupParams, FText::FromString(AttributeAsString[i]), OutPos, OutScale, VisualizationSettings.TextScale, DrawColor);
+				DrawAsText(InContext.SceneSetupParams, FText::FromString(AttributeAsString[i]), OutPos, OutScale, VisualizationSettings.TextScale, DrawColor, VisualizationSettings.TextOffset);
 			}
 			break;
 		}
@@ -518,7 +523,7 @@ void FPVDebugVisualizationBase::DrawAttributes(FVisualizerDrawContext& InContext
 
 					if (AttributeAsVector[i].Num() > 1)
 					{
-						DrawAsText(InContext.SceneSetupParams, FText::FromString(FString::Printf(TEXT("%d"), j)), FVector3f(EndPos), 0.f, VisualizationSettings.TextScale, DrawColor);
+						DrawAsText(InContext.SceneSetupParams, FText::FromString(FString::Printf(TEXT("%d"), j)), FVector3f(EndPos), 0.f, VisualizationSettings.TextScale, DrawColor, VisualizationSettings.TextOffset);
 					}
 				}
 			}
@@ -617,7 +622,7 @@ void FPVDebugVisualizationBase::AddDrawAsPoint(const FVector& InPos, float Point
 	PointsToDraw.Add({ InPos, PointScale });
 }
 
-void FPVDebugVisualizationBase::DrawAsText(FPCGSceneSetupParams& InOutParams, const FText& InTextToDraw, const FVector3f& InPos, float InScale, float TextSize, const FLinearColor& InColor)
+void FPVDebugVisualizationBase::DrawAsText(FPCGSceneSetupParams& InOutParams, const FText& InTextToDraw, const FVector3f& InPos, float InScale, float TextSize, const FLinearColor& InColor, const FVector3f TextOffset)
 {
 	const FVector3f& Pos = InPos + (FVector3f::RightVector * -1 * InScale);
 
@@ -629,7 +634,7 @@ void FPVDebugVisualizationBase::DrawAsText(FPCGSceneSetupParams& InOutParams, co
 	Component->SetCollisionEnabled(ECollisionEnabled::NoCollision);
 	
 	InOutParams.ManagedResources.Add(Component);
-	InOutParams.Scene->AddComponent(Component, FTransform(FVector(Pos.X, Pos.Y, Pos.Z)));
+	InOutParams.Scene->AddComponent(Component, FTransform(FVector(Pos.X + TextOffset.X, Pos.Y + TextOffset.Y, Pos.Z + TextOffset.Z)));
 }
 
 void FPVDebugVisualizationBase::DrawAsLine(UPVLineBatchComponent* LineBatchComponent, const FVector& InStartPos, const FVector& InEndPos,
diff --git a/Engine/Plugins/Experimental/ProceduralVegetationEditor/Source/ProceduralVegetationEditor/Private/Visualizations/DebugVisualization/PVDebugVisualizationBase.h b/Engine/Plugins/Experimental/ProceduralVegetationEditor/Source/ProceduralVegetationEditor/Private/Visualizations/DebugVisualization/PVDebugVisualizationBase.h
--- a/Engine/Plugins/Experimental/ProceduralVegetationEditor/Source/ProceduralVegetationEditor/Private/Visualizations/DebugVisualization/PVDebugVisualizationBase.h
+++ b/Engine/Plugins/Experimental/ProceduralVegetationEditor/Source/ProceduralVegetationEditor/Private/Visualizations/DebugVisualization/PVDebugVisualizationBase.h
@@ -28,7 +28,7 @@ protected:
 	virtual TArray<FVector3f> GetPivotPositions(const FManagedArrayCollection& InCollection) = 0;
 	virtual void GetPivot(const FManagedArrayCollection& InCollection, const int InIndex, FVector3f& OutPos, float& OutScale) = 0;
 	
-	static void DrawAsText(FPCGSceneSetupParams& InOutParams, const FText& InTextToDraw, const FVector3f& InPos, float InScale, float TextSize, const FLinearColor& InColor = FLinearColor::White);
+	static void DrawAsText(FPCGSceneSetupParams& InOutParams, const FText& InTextToDraw, const FVector3f& InPos, float InScale, float TextSize, const FLinearColor& InColor = FLinearColor::White, const FVector3f TextOffset = FVector3f::Zero());
 	static void DrawAsLine(UPVLineBatchComponent* LineBatchComponent, const FVector& InStartPos, const FVector& InEndPos, const FLinearColor& InColor = FLinearColor::White);
 	static void DrawAsSphere(UPVLineBatchComponent* LineBatchComponent, const FVector& InPos, float Radius, const FLinearColor& InColor = FLinearColor::White);
 	static void DrawAsPoints(UPVLineBatchComponent* LineBatchComponent, const TArray<TPair<FVector, float>>& InPoints, const FLinearColor& InColor = FLinearColor::White);
