[PVE] Object Avoidance Bug Fixing
- Fixed bug where some points slip through the object in cases where skeleton point lies directly on the mesh surface

#Synced-CL 48856039

diff --git a/Engine/Plugins/Experimental/ProceduralVegetationEditor/Source/ProceduralVegetation/Private/Implementations/PVGravity.cpp b/Engine/Plugins/Experimental/ProceduralVegetationEditor/Source/ProceduralVegetation/Private/Implementations/PVGravity.cpp
--- a/Engine/Plugins/Experimental/ProceduralVegetationEditor/Source/ProceduralVegetation/Private/Implementations/PVGravity.cpp
+++ b/Engine/Plugins/Experimental/ProceduralVegetationEditor/Source/ProceduralVegetation/Private/Implementations/PVGravity.cpp
@@ -87,7 +87,7 @@ void FPVGravity::ApplyGravity(
 	FVector3f PivotPosition
 )
 {
-	PVE_LOOP_DEBUG_CHECK(return);
+	PVE_OUTER_LOOP_DEBUG_CHECK(return);
 	
 	const auto CalculateDownForce = [](const FVector3f InRelativeDirection, const int InPointIndex, const float InPScaleGradient,
 	                                          const FPVGravityParams& InGravityParams,
diff --git a/Engine/Plugins/Experimental/ProceduralVegetationEditor/Source/ProceduralVegetation/Private/Nodes/PVObjectAvoidance.cpp b/Engine/Plugins/Experimental/ProceduralVegetationEditor/Source/ProceduralVegetation/Private/Nodes/PVObjectAvoidance.cpp
--- a/Engine/Plugins/Experimental/ProceduralVegetationEditor/Source/ProceduralVegetation/Private/Nodes/PVObjectAvoidance.cpp
+++ b/Engine/Plugins/Experimental/ProceduralVegetationEditor/Source/ProceduralVegetation/Private/Nodes/PVObjectAvoidance.cpp
@@ -125,15 +125,14 @@ bool FPVObjectAvoidanceElement::ExecuteInternal(FPCGContext* InContext) const
 						}
 					};
 
-
 				TArray<int32> BranchIndices;
 				BranchIndices.Reserve(BranchFacade.GetElementCount());
 				BranchFacade.GetSortedBranchIndicesByHierarchy(BranchIndices);
 
 				for (int32 BranchIndex : BranchIndices)
 				{
-					PVE_LOOP_DEBUG_CHECK(break);
-					
+					PVE_OUTER_LOOP_DEBUG_CHECK(break);
+
 					const TArray<int32> BranchPoints = BranchFacade.GetPoints(BranchIndex);
 					for (int32 BranchPointIndex = 1; BranchPointIndex < BranchPoints.Num(); ++BranchPointIndex)
 					{
@@ -144,38 +143,51 @@ bool FPVObjectAvoidanceElement::ExecuteInternal(FPCGContext* InContext) const
 						const FVector NextPointPosition = static_cast<FVector>(PointPositions[NextPointIndex]);
 						const FVector SegmentVector = NextPointPosition - PointPosition;
 
+						// Get the previous direction
+						FVector PreviousDirection;
+						if (BranchPoints.IsValidIndex(BranchPointIndex - 2))
+						{
+							PreviousDirection = (PointPosition - BranchPoints[BranchPointIndex - 2]).GetSafeNormal(UE_SMALL_NUMBER, FVector::UpVector);
+						}
+						else
+						{
+							PreviousDirection = (NextPointPosition - PointPosition).GetSafeNormal(UE_SMALL_NUMBER, FVector::UpVector);
+						}
+
+						const float PointScale = PointFacade.GetPointScale(PointIndex);
+
 						const FTransform InverseTransform = Transform.Inverse();
 
 						FRay3d Ray(PointPosition, SegmentVector);
-						FRay3d InverseRay(InverseTransform.TransformPosition(Ray.Origin), InverseTransform.TransformVectorNoScale(Ray.Direction));
+						FRay3d InverseRay(InverseTransform.TransformPosition(Ray.Origin), InverseTransform.TransformVector(Ray.Direction));
 
 						const float PointLFR = PointFacade.GetLengthFromRoot(PointIndex);
 						const float NextPointLFR = PointFacade.GetLengthFromRoot(NextPointIndex);
 
-						const float PointScale = PointFacade.GetPointScale(PointIndex);
-						const FVector MaxLengthVector = Ray.Direction * (SegmentVector.Length() + PointScale);
-						const double MaxLength = InverseTransform.TransformVector(MaxLengthVector).Length();
+						const double MaxLength = SegmentVector.Length() + PointScale;
+						const FVector MaxLengthVector = Ray.Direction * MaxLength;
+						const double InverseMaxLength = InverseTransform.TransformVector(MaxLengthVector).Length();
 
 						int32 HitTriangle = INDEX_NONE;
 						FVector BaryCoords;
-						if (double RayT; Octree.FindNearestHitTriangle(InverseRay, RayT, HitTriangle, BaryCoords, {MaxLength}))
+						if (double RayT; Octree.FindNearestHitTriangle(InverseRay, RayT, HitTriangle, BaryCoords, {InverseMaxLength}))
 						{
-							const FVector HitNormal = Transform.TransformVectorNoScale(GetTriangleNormal(Mesh, HitTriangle, BaryCoords));
-
-							if (FVector::DotProduct(Ray.Direction, HitNormal) > 0)
-							{
-								continue;
-							}
-
+							const FVector HitNormal = InverseTransform.TransformVector(GetTriangleNormal(Mesh, HitTriangle, BaryCoords))
+							                                          .GetSafeNormal(UE_SMALL_NUMBER, FVector::UpVector);
 							const FVector HitPoint = Transform.TransformPosition(InverseRay.PointAt(RayT));
 
 							const FVector InvertedNormal = -HitNormal;
 							const double DistanceToHit = FVector::Distance(PointPosition, HitPoint);
+							const double DistanceAfterHit = FVector::Distance(HitPoint, Ray.PointAt(MaxLength));
 							const FVector BiNormal = FVector::CrossProduct(Ray.Direction, InvertedNormal);
-							const FVector Tangent = FVector::CrossProduct(InvertedNormal, BiNormal);
+							FVector Tangent = FVector::CrossProduct(InvertedNormal, BiNormal);
+							const FVector PreviousBiNormal = FVector::CrossProduct(PreviousDirection, InvertedNormal);
+							const FVector PreviousTangent = FVector::CrossProduct(InvertedNormal, PreviousBiNormal);
+							// Get the mean tangent of this and previous tangent
+							Tangent = (PreviousTangent + Tangent).GetSafeNormal(UE_SMALL_NUMBER, FVector::UpVector);
 							const FVector AvoidanceOffset = Tangent * DistanceToHit * (FVector::DotProduct(Ray.Direction, InvertedNormal) + 1);
-							const FVector AvoidancePosition = HitPoint + AvoidanceOffset + HitNormal * PointScale;
-							const FVector AvoidanceDirection = (AvoidancePosition - PointPosition).GetUnsafeNormal();
+							const FVector AvoidancePosition = HitPoint + HitNormal * (PointScale + DistanceAfterHit) + AvoidanceOffset;
+							const FVector AvoidanceDirection = (AvoidancePosition - PointPosition).GetSafeNormal(UE_SMALL_NUMBER, FVector::DownVector);
 
 							const FVector OffsetPosition = PointPosition + AvoidanceDirection * FVector::Distance(PointPosition, NextPointPosition);
 							const FVector3f Offset = FVector3f(OffsetPosition - NextPointPosition);
@@ -183,10 +195,10 @@ bool FPVObjectAvoidanceElement::ExecuteInternal(FPCGContext* InContext) const
 							PVE_LOOP_DEBUG_PARAMS_START()
 								PVE_LOOP_DEBUG_DIRECTION_PARAM("Ray", Ray.Origin, MaxLengthVector);
 								PVE_LOOP_DEBUG_DIRECTION_PARAM("Hit Normal", HitPoint, HitNormal);
-								PVE_LOOP_DEBUG_POINT_PARAM("Hit Point", HitPoint);
-								PVE_LOOP_DEBUG_DIRECTION_PARAM("Avoidance Direction", HitPoint, AvoidanceDirection);
-								PVE_LOOP_DEBUG_VECTOR_PARAM("Offset Vector", NextPointPosition, OffsetPosition);
-								PVE_LOOP_DEBUG_TEXT_PARAM("Hit Distance", (PointPosition + HitPoint) / 2.0f, FText::FromString(FString::Printf(TEXT("%.2f units"), DistanceToHit)));
+								PVE_LOOP_DEBUG_DIRECTION_PARAM("Avoidance Offset", HitPoint, AvoidanceOffset);
+								PVE_LOOP_DEBUG_DIRECTION_PARAM("Avoidance Direction", NextPointPosition, AvoidanceDirection);
+								PVE_LOOP_DEBUG_VECTOR_PARAM("Avoidance Vector", PointPosition, AvoidancePosition);
+								PVE_LOOP_DEBUG_DIRECTION_PARAM("Offset", NextPointPosition, Offset);
 							PVE_LOOP_DEBUG_PARAMS_END()
 
 							PVE_LOOP_DEBUG_STEP(break);
@@ -201,7 +213,7 @@ bool FPVObjectAvoidanceElement::ExecuteInternal(FPCGContext* InContext) const
 								const float SiblingPointLFR = PointFacade.GetLengthFromRoot(SiblingPoint);
 								const float NextSiblingPointLFR = PointFacade.GetLengthFromRoot(NextSiblingPoint);
 
-								const FVector3f PointOffset = Offset / FMath::Pow(1.1f, BranchPointIndex - SiblingPointIndex);
+								const FVector3f PointOffset = Offset / FMath::Pow(1.5f, BranchPointIndex - SiblingPointIndex);
 
 								if (SiblingPointIndex != StartIndex)
 								{
diff --git a/Engine/Plugins/Experimental/ProceduralVegetationEditor/Source/ProceduralVegetation/Public/Helpers/PVUtilities.h b/Engine/Plugins/Experimental/ProceduralVegetationEditor/Source/ProceduralVegetation/Public/Helpers/PVUtilities.h
--- a/Engine/Plugins/Experimental/ProceduralVegetationEditor/Source/ProceduralVegetation/Public/Helpers/PVUtilities.h
+++ b/Engine/Plugins/Experimental/ProceduralVegetationEditor/Source/ProceduralVegetation/Public/Helpers/PVUtilities.h
@@ -48,7 +48,7 @@ struct FLoopDebugStepper
 	UPROPERTY(EditAnywhere, Category="Loop Debug")
 	mutable bool bDebugFocus = false;
 
-	UPROPERTY(EditAnywhere, Category="Loop Debug", meta=(ClampMin=1, UIMin=1, EditCondition="bDebug"))
+	UPROPERTY(EditAnywhere, Category="Loop Debug", meta=(EditCondition="bDebug"))
 	uint32 MaxSteps = 1;
 
 	void IncrementStep()
@@ -68,10 +68,12 @@ struct FLoopDebuggerState
 	const FLoopDebugStepper* LoopDebugStepper = nullptr;
 	UPVData* Data = nullptr;
 
+	bool bComplete = false;
+
 	bool IsLastLoopIteration() const
 	{
 		check(LoopDebugStepper);
-		return NumSteps == LoopDebugStepper->MaxSteps;
+		return NumSteps >= LoopDebugStepper->MaxSteps;
 	}
 
 	bool IterationsCompleted() const
@@ -110,7 +112,8 @@ inline TUniquePtr<FLoopDebuggerState> GlobalLoopDebuggerState;
 #define PVE_LOOP_DEBUG_PARAMS_START()																						\
 	if (																													\
 		GlobalLoopDebuggerState &&																							\
-		GlobalLoopDebuggerState->IsLastLoopIteration()																		\
+		GlobalLoopDebuggerState->IsLastLoopIteration() &&																	\
+		!GlobalLoopDebuggerState->bComplete																					\
 	)																														\
 	{																														\
 		FPVDebugSettings __DebugSettings = GlobalLoopDebuggerState->GetDebugSettings();										\
@@ -148,6 +151,14 @@ inline TUniquePtr<FLoopDebuggerState> GlobalLoopDebuggerState;
 		if (GlobalLoopDebuggerState)																						\
 			GlobalLoopDebuggerState->SetDebugSettings(MoveTemp(__DebugSettings));											\
 	}
+#define PVE_OUTER_LOOP_DEBUG_CHECK(ExitMethod)																				\
+	if (																													\
+		GlobalLoopDebuggerState &&																							\
+		GlobalLoopDebuggerState->bComplete																					\
+	)																														\
+	{																														\
+		ExitMethod;																											\
+	}
 
 #define PVE_LOOP_DEBUG_CHECK(ExitMethod)																					\
 	if (																													\
@@ -156,6 +167,7 @@ inline TUniquePtr<FLoopDebuggerState> GlobalLoopDebuggerState;
 	)																														\
 	{																														\
 		GlobalLoopDebuggerState->NumSteps++;																				\
+		GlobalLoopDebuggerState->bComplete = true;																			\
 		ExitMethod;																											\
 	}
 
@@ -180,6 +192,7 @@ inline TUniquePtr<FLoopDebuggerState> GlobalLoopDebuggerState;
 #define PVE_LOOP_DEBUG_TEXT_PARAM(InParamName, InPivot, InText)
 #define PVE_LOOP_DEBUG_PARAMS_END()
 
+#define PVE_OUTER_LOOP_DEBUG_CHECK(ExitMethod)
 #define PVE_LOOP_DEBUG_CHECK(ExitMethod)
 #define PVE_LOOP_DEBUG_STEP(ExitMethod)
 
diff --git a/Engine/Plugins/Experimental/ProceduralVegetationEditor/Source/ProceduralVegetationEditor/Private/Visualizations/PVDataVisualization.cpp b/Engine/Plugins/Experimental/ProceduralVegetationEditor/Source/ProceduralVegetationEditor/Private/Visualizations/PVDataVisualization.cpp
--- a/Engine/Plugins/Experimental/ProceduralVegetationEditor/Source/ProceduralVegetationEditor/Private/Visualizations/PVDataVisualization.cpp
+++ b/Engine/Plugins/Experimental/ProceduralVegetationEditor/Source/ProceduralVegetationEditor/Private/Visualizations/PVDataVisualization.cpp
@@ -120,17 +120,6 @@ FPCGSetupSceneFunc FPVDataVisualization::GetViewportSetupFunc(const UPCGSettings
 				TObjectPtr<UStaticMeshComponent> StaticMeshComponent = NewObject<UStaticMeshComponent>(GetTransientPackage(), NAME_None, RF_Transient);
 				StaticMeshComponent->SetStaticMesh(Mesh.LoadSynchronous());
 				StaticMeshComponent->UpdateBounds();
-				if (WeakData->GetDebugSettings().LoopDebugVisualizationSettings.Num())
-				{
-					const int32 NumMaterials = StaticMeshComponent->GetNumMaterials();
-					for (int32 Index = 0; Index < NumMaterials; Index++)
-					{
-						StaticMeshComponent->SetMaterial(
-							Index,
-							LoadObject<UMaterialInterface>(nullptr, TEXT("/ProceduralVegetationEditor/Materials/WireframeMaterial.WireframeMaterial"))
-						);
-					}
-				}
 
 				InOutParams.ManagedResources.Add(StaticMeshComponent);
 				InOutParams.Scene->AddComponent(StaticMeshComponent, Transform);
