[PVE] Fix cis issue
Removed unused header files

#Synced-CL 48842348

diff --git a/Engine/Plugins/Experimental/ProceduralVegetationEditor/Source/ProceduralVegetation/Private/Implementations/PVGrowerData.cpp b/Engine/Plugins/Experimental/ProceduralVegetationEditor/Source/ProceduralVegetation/Private/Implementations/PVGrowerData.cpp
--- a/Engine/Plugins/Experimental/ProceduralVegetationEditor/Source/ProceduralVegetation/Private/Implementations/PVGrowerData.cpp
+++ b/Engine/Plugins/Experimental/ProceduralVegetationEditor/Source/ProceduralVegetation/Private/Implementations/PVGrowerData.cpp
@@ -2,9 +2,6 @@
 
 #include "PVGrowerData.h"
 
-#include "Engine/SkeletalMesh.h"
-#include "GeometryCollection/GeometryCollectionComponent.h"
-
 void UPVGrowerData::AddPoint(TObjectPtr<UPVGrowerPoint> Point, TObjectPtr<UPVGrowerPrimitive> Primitive, bool bAddToBranch /*= true*/)
 {
 	MaxBudNumber += 1;
