[PVE] Minor bug fix in Object Avoidance PVE

#Synced-CL 48856039

diff --git a/Engine/Plugins/Experimental/ProceduralVegetationEditor/Source/ProceduralVegetation/Private/Nodes/PVObjectAvoidance.cpp b/Engine/Plugins/Experimental/ProceduralVegetationEditor/Source/ProceduralVegetation/Private/Nodes/PVObjectAvoidance.cpp
--- a/Engine/Plugins/Experimental/ProceduralVegetationEditor/Source/ProceduralVegetation/Private/Nodes/PVObjectAvoidance.cpp
+++ b/Engine/Plugins/Experimental/ProceduralVegetationEditor/Source/ProceduralVegetation/Private/Nodes/PVObjectAvoidance.cpp
@@ -156,24 +156,22 @@ bool FPVObjectAvoidanceElement::ExecuteInternal(FPCGContext* InContext) const
 
 						const float PointScale = PointFacade.GetPointScale(PointIndex);
 
-						const FTransform InverseTransform = Transform.Inverse();
-
 						FRay3d Ray(PointPosition, SegmentVector);
-						FRay3d InverseRay(InverseTransform.TransformPosition(Ray.Origin), InverseTransform.TransformVector(Ray.Direction));
+						FRay3d InverseRay(Transform.InverseTransformPosition(Ray.Origin), Transform.InverseTransformVector(Ray.Direction));
 
 						const float PointLFR = PointFacade.GetLengthFromRoot(PointIndex);
 						const float NextPointLFR = PointFacade.GetLengthFromRoot(NextPointIndex);
 
 						const double MaxLength = SegmentVector.Length() + PointScale;
 						const FVector MaxLengthVector = Ray.Direction * MaxLength;
-						const double InverseMaxLength = InverseTransform.TransformVector(MaxLengthVector).Length();
+						const double InverseMaxLength = Transform.InverseTransformVector(MaxLengthVector).Length();
 
 						int32 HitTriangle = INDEX_NONE;
 						FVector BaryCoords;
 						if (double RayT; Octree.FindNearestHitTriangle(InverseRay, RayT, HitTriangle, BaryCoords, {InverseMaxLength}))
 						{
-							const FVector HitNormal = InverseTransform.TransformVector(GetTriangleNormal(Mesh, HitTriangle, BaryCoords))
-							                                          .GetSafeNormal(UE_SMALL_NUMBER, FVector::UpVector);
+							const FVector HitNormal = Transform.TransformVectorNoScale(GetTriangleNormal(Mesh, HitTriangle, BaryCoords)).
+							                                    GetSafeNormal(UE_SMALL_NUMBER, FVector::UpVector);
 							const FVector HitPoint = Transform.TransformPosition(InverseRay.PointAt(RayT));
 
 							const FVector InvertedNormal = -HitNormal;
