Refine first implementation of StaticMesh to PVE growth data

- Re-name functions for improved clarity
- Allow filtering the mesh with multiple materials
- Add new implementation for finding branch radius and center along vertices (allow switching between new and old implementation using the CVar "PVE.Import.ComputeVerticesCenter")
- Restructure algorithm to allow more parts to run per branch instead of for the whole tree to improve results
- Optimize ComputeDist (now named ComputeLengthFromRoot)
- Remove usage of nested namespaces, and "using namespace" to improve compile times
- Make functions static, improving linking times during iterative compilation

#Synced-CL 48983295

diff --git a/Engine/Plugins/Experimental/ProceduralVegetationEditor/Source/ProceduralVegetation/Private/Nodes/PVStaticMeshImporterSettings.cpp b/Engine/Plugins/Experimental/ProceduralVegetationEditor/Source/ProceduralVegetation/Private/Nodes/PVStaticMeshImporterSettings.cpp
--- a/Engine/Plugins/Experimental/ProceduralVegetationEditor/Source/ProceduralVegetation/Private/Nodes/PVStaticMeshImporterSettings.cpp
+++ b/Engine/Plugins/Experimental/ProceduralVegetationEditor/Source/ProceduralVegetation/Private/Nodes/PVStaticMeshImporterSettings.cpp
@@ -6,6 +6,7 @@
 #include "ProceduralVegetationModule.h"
 #include "DataTypes/PVGrowthData.h"
 #include "ProceduralVegetationPreset.h"
+#include "Engine/StaticMesh.h"
 #include "Facades/PVBranchFacade.h"
 #include "Facades/PVPointFacade.h"
 
@@ -53,9 +54,26 @@ bool FPVStaticMeshImporterElement::ExecuteInternal(FPCGContext* InContext) const
 		return true;
 	}
 
+	if (Settings->Params.StaticMeshAsset->IsCompiling())
+	{
+		PCGLog::LogWarningOnGraph(LOCTEXT("StaticMeshAssetNotReady", "Static Mesh asset not ready"), InContext);
+		return false;
+	}
+
+	if (Settings->Params.MaterialsToKeep.Num() == 0)
+	{
+		PCGLog::LogWarningOnGraph(LOCTEXT("NoMaterialsSelected", "No materials to keep set. Output will be empty"), InContext);
+		return true;
+	}
+	
 #if WITH_EDITOR
+	if (!PV::GPVEditorInterface)
+	{
+		return true;
+	}
+	
 	FManagedArrayCollection Collection;
-	if (!PV::GPVEditorInterface || !PV::GPVEditorInterface->ImportGrowthDataFromStaticMesh(Settings->Params, Collection))
+	if (!PV::GPVEditorInterface->ImportGrowthDataFromStaticMesh(Settings->Params, Collection))
 	{
 		return false;
 	}
diff --git a/Engine/Plugins/Experimental/ProceduralVegetationEditor/Source/ProceduralVegetation/Public/PVImportStaticMeshParams.h b/Engine/Plugins/Experimental/ProceduralVegetationEditor/Source/ProceduralVegetation/Public/PVImportStaticMeshParams.h
--- a/Engine/Plugins/Experimental/ProceduralVegetationEditor/Source/ProceduralVegetation/Public/PVImportStaticMeshParams.h
+++ b/Engine/Plugins/Experimental/ProceduralVegetationEditor/Source/ProceduralVegetation/Public/PVImportStaticMeshParams.h
@@ -14,16 +14,14 @@ enum class EPVImportStaticMeshDebugState : uint8
 	InitialMesh,
 	KeepMaterial,
 	ComputeNormals,
-	AverageNormals,
-	SmoothNormalsWithLengthPreservation,
-	GetSpherical,
-	BlurRadius,
-	ExtractReductionCloud,
-	RadialFuse,
-	GroupPathPoints,
-	ComputeDist,
-	DrawBranches,
-	RemoveOverlaps,
+	ComputeVerticesCenterAndRadius,
+	CollapseBranchesToLine,
+	FindEndPoints,
+	ConnectBranchesToParent,
+	FindRootPoints,
+	ComputeLengthFromRoot,
+	DrawBranchCurves,
+	RemoveBranchCurveOverlaps,
 	ExtractBranches,
 	MAX UMETA(Hidden)
 };
@@ -37,7 +35,7 @@ struct FPVImportStaticMeshParams
 	TObjectPtr<UStaticMesh> StaticMeshAsset = nullptr;
 
 	UPROPERTY(EditAnywhere, Category = "Asset")
-	FName MaterialFilter = NAME_None;
+	TArray<FName> MaterialsToKeep;
 
 	UPROPERTY(EditAnywhere, Category = "Asset")
 	EPVImportStaticMeshDebugState DebugState = EPVImportStaticMeshDebugState::None;
diff --git a/Engine/Plugins/Experimental/ProceduralVegetationEditor/Source/ProceduralVegetationEditor/Private/Importer/PVImporter_StaticMesh.cpp b/Engine/Plugins/Experimental/ProceduralVegetationEditor/Source/ProceduralVegetationEditor/Private/Importer/PVImporter_StaticMesh.cpp
--- a/Engine/Plugins/Experimental/ProceduralVegetationEditor/Source/ProceduralVegetationEditor/Private/Importer/PVImporter_StaticMesh.cpp
+++ b/Engine/Plugins/Experimental/ProceduralVegetationEditor/Source/ProceduralVegetationEditor/Private/Importer/PVImporter_StaticMesh.cpp
@@ -1,69 +1,123 @@
 // Copyright Epic Games, Inc. All Rights Reserved.
 
 #include "PVImporter_StaticMesh.h"
-#include "ProceduralVegetationEditorModule.h"
 
-#include "ProceduralVegetationPreset.h"
 #include "PVPlantDescription.h"
-#include "Helpers/PVUtilities.h"
 #include "PVImportHelpers.h"
 
-#include "AssetRegistry/IAssetRegistry.h"
+#include "Engine/StaticMesh.h"
 #include "DynamicMesh/DynamicMesh3.h"
+#include "DynamicMesh/DynamicVertexAttribute.h"
+#include "DynamicMesh/DynamicMeshAttributeSet.h"
 #include "DynamicMesh/DynamicMeshAABBTree3.h"
 #include "DynamicMesh/MeshNormals.h"
-#include "DynamicMesh/DynamicMeshAttributeSet.h"
-#include "Engine/StaticMesh.h"
-#include "Parameterization/MeshDijkstra.h"
-#include "SphereTypes.h"
 #include "ConversionUtils/SceneComponentToDynamicMesh.h"
 #include "GeometryCollectionToDynamicMesh.h"
 #include "GeometryCollection/GeometryCollection.h"
 #include "Selections/MeshConnectedComponents.h"
 
-namespace PV::Import::StaticMesh::Internal
+template<typename AttribValueType, int AttribDimension>
+class TDynamicMeshVertexAttributeExt 
+	: public UE::Geometry::TDynamicMeshVertexAttribute<AttribValueType, AttribDimension>
+{
+public:
+	using Super = UE::Geometry::TDynamicMeshVertexAttribute<AttribValueType, AttribDimension>;
+
+	TDynamicMeshVertexAttributeExt() : Super() {}
+	TDynamicMeshVertexAttributeExt(UE::Geometry::FDynamicMesh3* ParentIn, bool bAutoInit = true) : Super(ParentIn, bAutoInit) {}
+
+	inline AttribValueType GetValue(int32 VertexID) const
+	{
+		AttribValueType Value = (AttribValueType)0;
+		Super::GetValue(VertexID, &Value);
+		return Value;
+	}
+
+	inline void SetValue(int32 VertexID, const AttribValueType& NewValue)
+	{
+		Super::SetValue(VertexID, &NewValue);
+	}
+
+	template <typename Predicate>
+	TArray<int32> FindAllByPredicate(Predicate Pred) const
+	{
+		TArray<int32> Result;
+
+		if (Super::Parent)
+		{
+			Result.Reserve(Super::Parent->VertexCount());
+
+			for (int32 VertexID : Super::Parent->VertexIndicesItr())
+			{
+				const AttribValueType Value = GetValue(VertexID);
+				if (Pred(Value))
+				{
+					Result.Add(VertexID);
+				}
+			}
+		}
+
+		return Result;
+	}
+
+	TArray<int32> FindAllNonZero() const
+	{
+		return FindAllByPredicate([](const auto& Value) { return Value != (AttribValueType)0; });
+	}
+
+	TArray<int32> FindAllZero() const
+	{
+		return FindAllByPredicate([](const auto& Value) { return Value == (AttribValueType)0; });
+	}
+};
+
+template<typename T>
+class TDynamicMeshVertexAttributeDefinition
 {
-	const static FName PScaleAttributeName = TEXT("pscale");
-	const static FName FuseToAttributeName = TEXT("fuseTo");
-	const static FName StartGroupAttributeName = TEXT("End_Grp");
-	const static FName EndGroupAttributeName = TEXT("Start_Grp");
-	const static FName SeedAttributeName = TEXT("seed");
-	const static FName SeedZPosAttributeName = TEXT("seed_Z_pos");
-	const static FName DistAttributeName = TEXT("dist");
-
-	typedef UE::Geometry::TDynamicMeshVertexAttribute<float, 1> FPScaleAttribute;
-	typedef UE::Geometry::TDynamicMeshVertexAttribute<int32, 1> FFuseToAttribute;
-	typedef UE::Geometry::TDynamicMeshVertexAttribute<bool, 1> FStartGroupAttribute;
-	typedef UE::Geometry::TDynamicMeshVertexAttribute<bool, 1> FEndGroupAttribute;
-	typedef UE::Geometry::TDynamicMeshVertexAttribute<bool, 1> FSeedAttribute;
-	typedef UE::Geometry::TDynamicMeshVertexAttribute<float, 1> FSeedZPosAttribute;
-	typedef UE::Geometry::TDynamicMeshVertexAttribute<float, 1> FDistAttribute;
-
-	template<typename T> 
-	T* GetAttachedAttribute(UE::Geometry::FDynamicMesh3& DynamicMesh, FName AttributeName)
+public:
+	using FAttributeType = TDynamicMeshVertexAttributeExt<T, 1>;
+
+	const FName AttributeName;
+
+	TDynamicMeshVertexAttributeDefinition(FName InAttributeName)
+		: AttributeName(InAttributeName)
+	{
+	}
+
+	const FAttributeType* GetAttribute(const UE::Geometry::FDynamicMesh3& DynamicMesh) const
 	{
 		if (DynamicMesh.HasAttributes() && DynamicMesh.Attributes()->HasAttachedAttribute(AttributeName))
 		{
-			return static_cast<T*>(DynamicMesh.Attributes()->GetAttachedAttribute(AttributeName));
+			return static_cast<const FAttributeType*>(DynamicMesh.Attributes()->GetAttachedAttribute(AttributeName));
 		}
 		return nullptr;
 	}
 
-	template<typename T>
-	const T* GetAttachedAttribute(const UE::Geometry::FDynamicMesh3& DynamicMesh, FName AttributeName)
+	FAttributeType* GetAttribute(UE::Geometry::FDynamicMesh3& DynamicMesh) const
 	{
 		if (DynamicMesh.HasAttributes() && DynamicMesh.Attributes()->HasAttachedAttribute(AttributeName))
 		{
-			return static_cast<const T*>(DynamicMesh.Attributes()->GetAttachedAttribute(AttributeName));
+			return static_cast<FAttributeType*>(DynamicMesh.Attributes()->GetAttachedAttribute(AttributeName));
 		}
 		return nullptr;
 	}
 
-	template<typename AttribType>
-	void AttachVertexAttribute(UE::Geometry::FDynamicMesh3& DynamicMesh, FName AttributeName)
+	const FAttributeType& GetAttributeChecked(const UE::Geometry::FDynamicMesh3& DynamicMesh) const
 	{
-		using namespace UE::Geometry;
+		const FAttributeType* Attribute = GetAttribute(DynamicMesh);
+		check(Attribute != nullptr);
+		return *Attribute;
+	}
+
+	FAttributeType& GetAttributeChecked(UE::Geometry::FDynamicMesh3& DynamicMesh) const
+	{
+		FAttributeType* Attribute = GetAttribute(DynamicMesh);
+		check(Attribute != nullptr);
+		return *Attribute;
+	}
 
+	void AttachAttribute(UE::Geometry::FDynamicMesh3& DynamicMesh) const
+	{
 		if (!DynamicMesh.HasAttributes())
 		{
 			DynamicMesh.EnableAttributes();
@@ -71,15 +125,15 @@ namespace PV::Import::StaticMesh::Internal
 
 		check(!DynamicMesh.Attributes()->HasAttachedAttribute(AttributeName));
 
-		auto* Attribute = new AttribType(&DynamicMesh);
+		auto* Attribute = new FAttributeType(&DynamicMesh);
 		Attribute->Initialize();
 		DynamicMesh.Attributes()->AttachAttribute(AttributeName, Attribute);
 	}
 
-	template<typename T>
-	void CopyAttachedVertexAttributeToCollection(const UE::Geometry::FDynamicMesh3& DynamicMesh, FManagedArrayCollection& OutCollection, FName AttributeName)
+	void CopyAttributeToCollection(const UE::Geometry::FDynamicMesh3& DynamicMesh, FManagedArrayCollection& OutCollection) const
 	{
-		if (!DynamicMesh.HasAttributes() || !DynamicMesh.Attributes()->HasAttachedAttribute(AttributeName))
+		const FAttributeType* AttachedAttribute = GetAttribute(DynamicMesh);
+		if (!AttachedAttribute)
 		{
 			return;
 		}
@@ -89,7 +143,7 @@ namespace PV::Import::StaticMesh::Internal
 		if (!OutCollection.HasGroup(DebugVertexGroup))
 		{
 			OutCollection.AddGroup(DebugVertexGroup);
-			OutCollection.AddElements(DynamicMesh.VertexCount(), DebugVertexGroup);
+			OutCollection.AddElements(DynamicMesh.MaxVertexID(), DebugVertexGroup);
 
 			auto& DebugVertexPositionAttribute = OutCollection.AddAttribute<FVector3f>(FGeometryCollection::VertexPositionAttribute, DebugVertexGroup);
 			for (int32 VertexIndex : DynamicMesh.VertexIndicesItr())
@@ -98,832 +152,1164 @@ namespace PV::Import::StaticMesh::Internal
 			}
 		}
 
-		auto* AttachedAttribute = Internal::GetAttachedAttribute<UE::Geometry::TDynamicMeshVertexAttribute<T, 1>>(DynamicMesh, AttributeName);
-
 		TManagedArray<T>& CollectionAttribute = OutCollection.AddAttribute<T>(AttributeName, DebugVertexGroup);
 		for (int32 VertexIndex : DynamicMesh.VertexIndicesItr())
 		{
-			T Value;
-			AttachedAttribute->GetValue(VertexIndex, &Value);
-			CollectionAttribute[VertexIndex] = Value;
+			CollectionAttribute[VertexIndex] = AttachedAttribute->GetValue(VertexIndex);
 		}
 	}
+};
 
-	FPScaleAttribute* GetPScaleAttribute(UE::Geometry::FDynamicMesh3& DynamicMesh) { return GetAttachedAttribute<FPScaleAttribute>(DynamicMesh, PScaleAttributeName); }
-	const FPScaleAttribute* GetPScaleAttribute(const UE::Geometry::FDynamicMesh3& DynamicMesh) { return GetAttachedAttribute<const FPScaleAttribute>(DynamicMesh, PScaleAttributeName); }
-	
-	FFuseToAttribute* GetFuseToAttribute(UE::Geometry::FDynamicMesh3& DynamicMesh) { return GetAttachedAttribute<FFuseToAttribute>(DynamicMesh, FuseToAttributeName); }
-	const FFuseToAttribute* GetFuseToAttribute(const UE::Geometry::FDynamicMesh3& DynamicMesh) { return GetAttachedAttribute<const FFuseToAttribute>(DynamicMesh, FuseToAttributeName); }
-	
-	FStartGroupAttribute* GetStartGroupAttribute(UE::Geometry::FDynamicMesh3& DynamicMesh) { return GetAttachedAttribute<FStartGroupAttribute>(DynamicMesh, StartGroupAttributeName); }
-	const FStartGroupAttribute* GetStartGroupAttribute(const UE::Geometry::FDynamicMesh3& DynamicMesh) { return GetAttachedAttribute<const FStartGroupAttribute>(DynamicMesh, StartGroupAttributeName); }
-	
-	FEndGroupAttribute* GetEndGroupAttribute(UE::Geometry::FDynamicMesh3& DynamicMesh) { return GetAttachedAttribute<FEndGroupAttribute>(DynamicMesh, EndGroupAttributeName); }
-	const FEndGroupAttribute* GetEndGroupAttribute(const UE::Geometry::FDynamicMesh3& DynamicMesh) { return GetAttachedAttribute<const FEndGroupAttribute>(DynamicMesh, EndGroupAttributeName); }
-	
-	FSeedAttribute* GetSeedAttribute(UE::Geometry::FDynamicMesh3& DynamicMesh) { return GetAttachedAttribute<FSeedAttribute>(DynamicMesh, SeedAttributeName); }
-	const FSeedAttribute* GetSeedAttribute(const UE::Geometry::FDynamicMesh3& DynamicMesh) { return GetAttachedAttribute<const FSeedAttribute>(DynamicMesh, SeedAttributeName); }
-	
-	FSeedZPosAttribute* GetSeedZPosAttribute(UE::Geometry::FDynamicMesh3& DynamicMesh) { return GetAttachedAttribute<FSeedZPosAttribute>(DynamicMesh, SeedZPosAttributeName); }
-	const FSeedZPosAttribute* GetSeedZPosAttribute(const UE::Geometry::FDynamicMesh3& DynamicMesh) { return GetAttachedAttribute<const FSeedZPosAttribute>(DynamicMesh, SeedZPosAttributeName); }
-	
-	FDistAttribute* GetDistAttribute(UE::Geometry::FDynamicMesh3& DynamicMesh) { return GetAttachedAttribute<FDistAttribute>(DynamicMesh, DistAttributeName); }
-	const FDistAttribute* GetDistAttribute(const UE::Geometry::FDynamicMesh3& DynamicMesh) { return GetAttachedAttribute<const FDistAttribute>(DynamicMesh, DistAttributeName); }
+const static TDynamicMeshVertexAttributeDefinition<float> VertexScaleAttributeDefinition(TEXT("VertexScale"));
+const static TDynamicMeshVertexAttributeDefinition<int32> FuseToAttributeDefinition(TEXT("FuseTo"));
+const static TDynamicMeshVertexAttributeDefinition<bool> FillerVertexAttributeDefinition(TEXT("FillerVertex"));
+const static TDynamicMeshVertexAttributeDefinition<bool> EndPointAttributeDefinition(TEXT("EndPoint"));
+const static TDynamicMeshVertexAttributeDefinition<bool> RootPointAttributeDefinition(TEXT("RootPoint"));
+const static TDynamicMeshVertexAttributeDefinition<float> LengthFromRootAttributeDefinition(TEXT("LengthFromRoot"));
 
-	void AttachPScaleAttribute(UE::Geometry::FDynamicMesh3& DynamicMesh) { AttachVertexAttribute<FPScaleAttribute>(DynamicMesh, PScaleAttributeName); }
-	void AttachFuseToAttribute(UE::Geometry::FDynamicMesh3& DynamicMesh) { AttachVertexAttribute<FFuseToAttribute>(DynamicMesh, FuseToAttributeName); }
-	void AttachStartGroupAttribute(UE::Geometry::FDynamicMesh3& DynamicMesh) { AttachVertexAttribute<FStartGroupAttribute>(DynamicMesh, StartGroupAttributeName); }
-	void AttachEndGroupAttribute(UE::Geometry::FDynamicMesh3& DynamicMesh) { AttachVertexAttribute<FEndGroupAttribute>(DynamicMesh, EndGroupAttributeName); }
-	void AttachSeedAttribute(UE::Geometry::FDynamicMesh3& DynamicMesh) { AttachVertexAttribute<FSeedAttribute>(DynamicMesh, SeedAttributeName); }
-	void AttachSeedZPosAttribute(UE::Geometry::FDynamicMesh3& DynamicMesh) { AttachVertexAttribute<FSeedZPosAttribute>(DynamicMesh, SeedZPosAttributeName); }
-	void AttachDistAttribute(UE::Geometry::FDynamicMesh3& DynamicMesh) { AttachVertexAttribute<FDistAttribute>(DynamicMesh, DistAttributeName); }
+static FVector GenerateFillerVertex(const FVector& Pos1, const FVector& Pos2)
+{
+	// Since we don't have edges for the DynamicMesh we add a new vertex at the center (offset by 0.1mm) whenever we need to represent a line
+	return ((Pos1 + Pos2) * 0.5)  + (Pos1 - Pos2).GetSafeNormal().Cross(FVector::ForwardVector).Cross(FVector::UpVector)/* * 0.01*/;
+}
 
-	void RemoveUnusedVertices(UE::Geometry::FDynamicMesh3& DynamicMesh)
+static TArray<int32> GetUniqueVertexIndices(const UE::Geometry::FDynamicMesh3& DynamicMesh, const TArray<int32>& TriangleIds)
+{
+	TSet<int32> UniqueVertexIndices;
+	UniqueVertexIndices.Reserve(TriangleIds.Num() * 3);
+	for (int32 tid : TriangleIds)
 	{
-		TRACE_CPUPROFILER_EVENT_SCOPE(PV::Import::StaticMesh::OrphanRemoval);
+		const UE::Geometry::FIndex3i VertexIndices = DynamicMesh.GetTriangle(tid);
+		UniqueVertexIndices.Add(VertexIndices[0]);
+		UniqueVertexIndices.Add(VertexIndices[1]);
+		UniqueVertexIndices.Add(VertexIndices[2]);
+	}
 
-		DynamicMesh.RemoveUnusedVertices();
-		DynamicMesh.CompactInPlace();
+	return UniqueVertexIndices.Array();
+}
+
+static TArray<int32> GetVertexNeighbours(const UE::Geometry::FDynamicMesh3& DynamicMesh, int32 VertexID)
+{
+	TArray<int32> OutNeighbours;
+	OutNeighbours.Reserve(DynamicMesh.GetVtxEdgeCount(VertexID));
+	for (int32 EdgeID : DynamicMesh.VtxEdgesItr(VertexID))
+	{
+		const UE::Geometry::FIndex2i VertexIDs = DynamicMesh.GetEdgeV(EdgeID);
+		const int32 OtherVertexID = VertexIDs.A == VertexID ? VertexIDs.B : VertexIDs.A;
+		OutNeighbours.Add(OtherVertexID);
 	}
+	return OutNeighbours;
+}
+
+static TArray<int32> FindVerticesInRadius(const UE::Geometry::FDynamicMeshAABBTree3& DynamicMeshAABBTree, const FVector& Pos, float Radius)
+{
+	TArray<int32> OutVertexIndices;
 
-	void KeepMaterial(UE::Geometry::FDynamicMesh3& DynamicMesh, int32 MaterialIndexToKeep)
+	const UE::Geometry::FDynamicMesh3* DynamicMesh = DynamicMeshAABBTree.GetMesh();
+	if (DynamicMesh == nullptr)
 	{
-		TRACE_CPUPROFILER_EVENT_SCOPE(PV::Import::StaticMesh::KeepMaterial);
+		return OutVertexIndices;
+	}
 
-		if (MaterialIndexToKeep == INDEX_NONE)
-		{
-			return;
-		}
+	OutVertexIndices.Reserve(DynamicMesh->VertexCount());
 
-		UE::Geometry::FDynamicMeshAttributeSet* MeshAttributeSet = DynamicMesh.Attributes();
-		UE::Geometry::FDynamicMeshMaterialAttribute* MaterialIDAttribute = MeshAttributeSet ? MeshAttributeSet->GetMaterialID() : nullptr;
-		
-		if (!MaterialIDAttribute)
-		{
-			return;
-		}
+	const double RadiusSqr = Radius * Radius;
+
+	UE::Geometry::FDynamicMeshAABBTree3::FTreeTraversal TreeTraversal;
+
+	TreeTraversal.NextBoxF = [&](const UE::Geometry::FAxisAlignedBox3d& Box, int32 Depth)->bool
+	{
+		return Box.DistanceSquared(Pos) <= RadiusSqr;
+	};
 
-		for (int32 TriangleID : DynamicMesh.TriangleIndicesItr())
+	TreeTraversal.NextTriangleF = [&](int32 TID)
+	{
+		const UE::Geometry::FIndex3i Triangle = DynamicMesh->GetTriangle(TID);
+		for (int32 j = 0; j < 3; ++j)
 		{
-			const int32 MaterialID = MaterialIDAttribute->GetValue(TriangleID);
-			if (MaterialID != MaterialIndexToKeep)
+			const double VertexDistSqr = (DynamicMesh->GetVertex(Triangle[j]) - Pos).SizeSquared();
+			if (VertexDistSqr <= RadiusSqr)
 			{
-				DynamicMesh.RemoveTriangle(TriangleID);
+				OutVertexIndices.AddUnique(Triangle[j]);
 			}
 		}
+	};
 
-		RemoveUnusedVertices(DynamicMesh);
-	}
+	DynamicMeshAABBTree.DoTraversal(TreeTraversal);
 
-	void ComputeNormals(UE::Geometry::FDynamicMesh3& DynamicMesh)
-	{
-		TRACE_CPUPROFILER_EVENT_SCOPE(PV::Import::StaticMesh::ComputeNormals);
+	return OutVertexIndices;
+}
 
-		if (!DynamicMesh.HasVertexNormals())
-		{
-			DynamicMesh.EnableVertexNormals(FVector3f::UnitX());
-		}
+static void RemoveUnusedVertices(UE::Geometry::FDynamicMesh3& DynamicMesh)
+{
+	TRACE_CPUPROFILER_EVENT_SCOPE(OrphanRemoval);
 
-		if (DynamicMesh.Attributes() && DynamicMesh.Attributes()->NumNormalLayers() > 0)
-		{
-			UE::Geometry::FDynamicMeshNormalOverlay* Overlay = DynamicMesh.Attributes()->GetNormalLayer(0);
-			for (int32 EID : Overlay->ElementIndicesItr())
-			{
-				int32 ParentVID = Overlay->GetParentVertex(EID);
-				if (ParentVID != INDEX_NONE)
-				{
-					DynamicMesh.SetVertexNormal(ParentVID, Overlay->GetElement(EID));
-				}
-			}
-		}
-		else
+	DynamicMesh.RemoveUnusedVertices();
+	DynamicMesh.CompactInPlace();
+}
+
+static bool FilterTrianglesUsingMaterials(UE::Geometry::FDynamicMesh3& DynamicMesh, const UStaticMesh* StaticMesh, const TArray<FName>& MaterialsToKeep)
+{
+	TRACE_CPUPROFILER_EVENT_SCOPE(KeepMaterial);
+
+	TSet<int32> MaterialIDs;
+	MaterialIDs.Reserve(MaterialsToKeep.Num());
+
+	for (FName MaterialName : MaterialsToKeep)
+	{
+		const int32 MaterialIndexToKeep = StaticMesh->GetMaterialIndex(MaterialName);
+		if (MaterialIndexToKeep != INDEX_NONE)
 		{
-			UE::Geometry::FMeshNormals::QuickComputeVertexNormals(DynamicMesh);
+			MaterialIDs.Add(MaterialIndexToKeep);
 		}
 	}
 
-	void AverageNormals(UE::Geometry::FDynamicMesh3& DynamicMesh)
+	if (MaterialIDs.Num() == 0)
 	{
-		// This takes the normals of each vertex instance and averages them onto the vertex (or point as it's called in Houdini)
-		// This one seems more important than the rest for getting any good results
-
-		// TODO: Not super important for now, do later
+		return false;
 	}
 
-	void SmoothNormalsWithLengthPreservation(UE::Geometry::FDynamicMesh3& DynamicMesh)
+	UE::Geometry::FDynamicMeshAttributeSet* MeshAttributeSet = DynamicMesh.Attributes();
+	UE::Geometry::FDynamicMeshMaterialAttribute* MaterialIDAttribute = MeshAttributeSet ? MeshAttributeSet->GetMaterialID() : nullptr;
+		
+	if (!MaterialIDAttribute)
 	{
-		// TODO: Has minimal impact on the final result, can wait with this
+		return false;
 	}
 
-	void GetSpherical_Internal(UE::Geometry::FDynamicMesh3& DynamicMesh, const TArray<int32>& TriangleIds)
+	for (int32 TriangleID : DynamicMesh.TriangleIndicesItr())
 	{
-		TRACE_CPUPROFILER_EVENT_SCOPE(PV::Import::StaticMesh::GetSpherical_Internal);
-
-		using namespace UE::Geometry;
-
-		TSet<int32> UniqueVertexIndices;
-		UniqueVertexIndices.Reserve(TriangleIds.Num() * 3);
-		for (int32 tid : TriangleIds)
+		const int32 MaterialID = MaterialIDAttribute->GetValue(TriangleID);
+		if (!MaterialIDs.Contains(MaterialID))
 		{
-			const FIndex3i VertexIndices = DynamicMesh.GetTriangle(tid);
-			UniqueVertexIndices.Add(VertexIndices[0]);
-			UniqueVertexIndices.Add(VertexIndices[1]);
-			UniqueVertexIndices.Add(VertexIndices[2]);
+			DynamicMesh.RemoveTriangle(TriangleID);
 		}
+	}
 
-		const TArray<int32> ComponentVertexIndices = UniqueVertexIndices.Array();
+	RemoveUnusedVertices(DynamicMesh);
 
-		FDynamicMeshAABBTree3 DynamicMeshAABBTree(&DynamicMesh, false);
-		DynamicMeshAABBTree.Build(TriangleIds);
+	return true;
+}
 
-		const static auto ComputeRadius = [](const FVector& Point, const FVector& Normal, const FVector& QueryPosition)->double
-		{
-			// Compute radius of the ball that touches points p and q and whose center falls on the normal n from p
-			const FVector Delta = Point - QueryPosition;
-			const double Dist = Delta.Size();
-			const double CosTheta = Normal.Dot(Delta) / Dist;
-			return Dist / (2 * CosTheta);
-		};
+static void ComputeNormals(UE::Geometry::FDynamicMesh3& DynamicMesh)
+{
+	TRACE_CPUPROFILER_EVENT_SCOPE(ComputeNormals);
 
-		double NearestDistSqr = 0.f;
+	if (!DynamicMesh.HasVertexNormals())
+	{
+		DynamicMesh.EnableVertexNormals(FVector3f::UnitX());
+	}
 
-		TArray<FVector> NewVertexPositions;
-		NewVertexPositions.SetNum(ComponentVertexIndices.Num());
+	if (DynamicMesh.Attributes() && DynamicMesh.Attributes()->NumNormalLayers() > 0)
+	{
+		UE::Geometry::FDynamicMeshNormalOverlay* Overlay = DynamicMesh.Attributes()->GetNormalLayer(0);
+		for (int32 EID : Overlay->ElementIndicesItr())
+		{
+			int32 ParentVID = Overlay->GetParentVertex(EID);
+			if (ParentVID != INDEX_NONE)
+			{
+				DynamicMesh.SetVertexNormal(ParentVID, Overlay->GetElement(EID));
+			}
+		}
+	}
+	else
+	{
+		UE::Geometry::FMeshNormals::QuickComputeVertexNormals(DynamicMesh);
+	}
+}
 
-		TArray<float> PScales;
-		PScales.SetNum(ComponentVertexIndices.Num());
+static void ComputeVerticesCenterAndRadius_Internal_V1(UE::Geometry::FDynamicMesh3& DynamicMesh, const TArray<int32>& TriangleIds)
+{
+	// This function tries to find the center point of the branch/trunk for each vertex, it then moves the vertex to said center point and stores the
+	// branch radius in the VertexScale attribute
+	//
+	// The algorithm will perform a sphere overlap with a radius of 50cm, at 50cm behind the vertex normal. If it overlaps a vertex it will try to move the
+	// center point to lay between the current vertex, and the found vertex. It then performs another overlap, and repeats the process until it no longer overlaps any vertices.
+	// Once an opposing vertex has been found, an extra check is performed to make sure the normal of the closest triangle is facing away from the center point.
+	
+	const TArray<int32> ComponentVertexIndices = GetUniqueVertexIndices(DynamicMesh, TriangleIds);
 
-		for (int32 i = 0; i < ComponentVertexIndices.Num(); ++i)
-		{
-			const int32 VertexIndex = ComponentVertexIndices[i];
-			const FVector& VertexPosition = DynamicMesh.GetVertex(VertexIndex);
-			const FVector VertexNormal = FVector(DynamicMesh.GetVertexNormal(VertexIndex));
+	UE::Geometry::FDynamicMeshAABBTree3 DynamicMeshAABBTree(&DynamicMesh, false);
+	DynamicMeshAABBTree.Build(TriangleIds);
 
-			NewVertexPositions[i] = VertexPosition;
-			PScales[i] = 0.5;
+	const static auto ComputeRadius = [](const FVector& Point, const FVector& Normal, const FVector& QueryPosition)->double
+	{
+		// Compute radius of the ball that touches points p and q and whose center falls on the normal n from p
+		const FVector Delta = Point - QueryPosition;
+		const double Dist = Delta.Size();
+		const double CosTheta = Normal.Dot(Delta) / Dist;
+		return Dist / (2 * CosTheta);
+	};
 
-			/*if (VertexNormal.Size() < 0.5) // TODO: Re-visit once we avg normals
-			{
-				continue;
-			}
-			
-			VertexNormal.Normalize();
-			*/
+	double NearestDistSqr = 0.f;
 
-			float Radius = 50; // 50cm
-			int32 OldPoint = INDEX_NONE;
-			FVector SphereCenter = VertexPosition - (FVector(VertexNormal) * Radius);
-			
-			int32 NearPoint = DynamicMeshAABBTree.FindNearestVertex(SphereCenter, NearestDistSqr, Radius);
-			while (NearPoint != INDEX_NONE)
-			{
-				// Get the position of the nearest point of the previous near point.
-				const FVector point_position = DynamicMesh.GetVertex(NearPoint);
+	TArray<FVector> NewVertexPositions;
+	NewVertexPositions.SetNum(ComponentVertexIndices.Num());
 
-				// Computes the radius of a sphere that would fit spherically between these two positions
-				const float NewRadius = ComputeRadius(VertexPosition, VertexNormal, point_position);
+	TArray<float> Scales;
+	Scales.SetNum(ComponentVertexIndices.Num());
 
-				Radius = FMath::Min(Radius, FMath::Max(NewRadius, 0));
-				// Sets the center of the lookup shere to be at the positions of the orignal position, minus the offset of the new found radius
-				SphereCenter = VertexPosition - VertexNormal * Radius;
+	for (int32 i = 0; i < ComponentVertexIndices.Num(); ++i)
+	{
+		const int32 VertexIndex = ComponentVertexIndices[i];
+		const FVector& VertexPosition = DynamicMesh.GetVertex(VertexIndex);
+		const FVector VertexNormal = FVector(DynamicMesh.GetVertexNormal(VertexIndex));
 
-				// Stores the current near point as the old point for comparison
-				OldPoint = NearPoint;
-				NearPoint = DynamicMeshAABBTree.FindNearestVertex(SphereCenter, NearestDistSqr, Radius);
+		NewVertexPositions[i] = VertexPosition;
+		Scales[i] = 0.5;
 
-				if (NearPoint == VertexIndex || OldPoint == NearPoint)
-				{ 
-					break; 
-				}
-			}
+		float Radius = 50; // 50cm
+		int32 OldPoint = INDEX_NONE;
+		FVector SphereCenter = VertexPosition - (FVector(VertexNormal) * Radius);
+			
+		int32 NearPoint = DynamicMeshAABBTree.FindNearestVertex(SphereCenter, NearestDistSqr, Radius);
+		while (NearPoint != INDEX_NONE)
+		{
+			// Get the position of the nearest point of the previous near point.
+			const FVector point_position = DynamicMesh.GetVertex(NearPoint);
 
-			const int32 NearTriID = DynamicMeshAABBTree.FindNearestTriangle(SphereCenter, NearestDistSqr);
-			if (NearTriID == INDEX_NONE)
-			{
-				continue;
-			}
+			// Computes the radius of a sphere that would fit spherically between these two positions
+			const float NewRadius = ComputeRadius(VertexPosition, VertexNormal, point_position);
+
+			Radius = FMath::Min(Radius, FMath::Max(NewRadius, 0));
+			// Sets the center of the lookup shere to be at the positions of the orignal position, minus the offset of the new found radius
+			SphereCenter = VertexPosition - VertexNormal * Radius;
 
-			const FDistPoint3Triangle3d Query = TMeshQueries<FDynamicMesh3>::TriangleDistance(DynamicMesh, NearTriID, SphereCenter);
-			const FVector& ClosestTrianglePoint = Query.ClosestTrianglePoint;
-			const FVector ClosestTrianglePointNormal = DynamicMesh.GetTriBaryNormal(NearTriID, Query.TriangleBaryCoords[0], Query.TriangleBaryCoords[1], Query.TriangleBaryCoords[2]);
+			// Stores the current near point as the old point for comparison
+			OldPoint = NearPoint;
+			NearPoint = DynamicMeshAABBTree.FindNearestVertex(SphereCenter, NearestDistSqr, Radius);
 
-			const FVector CenterDir = (SphereCenter - ClosestTrianglePoint).GetSafeNormal();
-			const double CenterDot = ClosestTrianglePointNormal.Dot(CenterDir);
-			if (CenterDot > 0) 
+			if (NearPoint == VertexIndex || OldPoint == NearPoint)
 			{ 
-				continue; 
+				break; 
 			}
-
-			PScales[i] = FMath::Max(Radius, 0.01);
-			NewVertexPositions[i] = SphereCenter;
 		}
 
-		FPScaleAttribute* PScaleAttribute = GetPScaleAttribute(DynamicMesh);
-		check(PScaleAttribute != nullptr);
-
-		for (int32 i = 0; i < ComponentVertexIndices.Num(); ++i)
+		const int32 NearTriID = DynamicMeshAABBTree.FindNearestTriangle(SphereCenter, NearestDistSqr);
+		if (NearTriID == INDEX_NONE)
 		{
-			const int32 VertexIndex = ComponentVertexIndices[i];
-			DynamicMesh.SetVertex(VertexIndex, NewVertexPositions[i]);
-			PScaleAttribute->SetValue(VertexIndex, &PScales[i]);
+			continue;
 		}
-	}
-
-	void GetSpherical(UE::Geometry::FDynamicMesh3& DynamicMesh)
-	{
-		TRACE_CPUPROFILER_EVENT_SCOPE(PV::Import::StaticMesh::GetSpherical);
 
-		UE::Geometry::FMeshConnectedComponents ConnectedComponents(&DynamicMesh);
-		ConnectedComponents.FindConnectedTriangles();
+		const UE::Geometry::FDistPoint3Triangle3d Query = UE::Geometry::TMeshQueries<UE::Geometry::FDynamicMesh3>::TriangleDistance(DynamicMesh, NearTriID, SphereCenter);
+		const FVector& ClosestTrianglePoint = Query.ClosestTrianglePoint;
+		const FVector ClosestTrianglePointNormal = DynamicMesh.GetTriBaryNormal(NearTriID, Query.TriangleBaryCoords[0], Query.TriangleBaryCoords[1], Query.TriangleBaryCoords[2]);
 
-		ParallelFor(ConnectedComponents.Num(), [&](int32 Index)
+		const FVector CenterDir = (SphereCenter - ClosestTrianglePoint).GetSafeNormal();
+		const double CenterDot = ClosestTrianglePointNormal.Dot(CenterDir);
+		if (CenterDot > 0) 
 		{
-			Internal::GetSpherical_Internal(DynamicMesh, ConnectedComponents.GetComponent(Index).Indices);
-		});
+			continue; 
+		}
+
+		Scales[i] = FMath::Max(Radius, 0.01);
+		NewVertexPositions[i] = SphereCenter;
 	}
 
-	void BlurRadius(UE::Geometry::FDynamicMesh3& DynamicMesh)
+	auto& VertexScaleAttribute = VertexScaleAttributeDefinition.GetAttributeChecked(DynamicMesh);
+
+	for (int32 i = 0; i < ComponentVertexIndices.Num(); ++i)
 	{
-		// TODO: Has smaller impact on the final result, can wait with this
+		const int32 VertexIndex = ComponentVertexIndices[i];
+		DynamicMesh.SetVertex(VertexIndex, NewVertexPositions[i]);
+		VertexScaleAttribute.SetValue(VertexIndex, Scales[i]);
 	}
-	
-	void ExtractReductionCloud(UE::Geometry::FDynamicMesh3& DynamicMesh)
+}
+
+static void ComputeVerticesCenterAndRadius_Internal_V2(UE::Geometry::FDynamicMesh3& DynamicMesh, const TArray<int32>& TriangleIds)
+{
+	// This is a second attempt at creating an algorithm that can walk along each vertex on a tree mesh and find the corresponding center point and
+	// radius of the branch at the point of the vertex.
+	//
+	// This algorithm will iterate the vertices and perform a ray-cast in the opposite direction of the vertex normal.
+	// If it hits another triangle then it will check whether the normal of said triangle is opposing to the current vertex (e.g. the dot of the
+	// vertex normal and the triangle normal is < -0.5). If the ray-cast fails to find a triangle with an opposing normal, then a second algorithm
+	// is employed to try and find an opposing vertex. This second algorithm finds the closest vertex to a position 100cm from the vertex in the
+	// opposite direction of the vertex normal, it then checks if the normal of this vertex is opposing. If this check fails, then it moves the
+	// query point closer and repeats the check. This is repeated 3 times before the algorithm gives up and the vertex is left with a scale if 0 which
+	// indicates that at radius/center point was not found.
+
+	// Issues with this algorithm:
+	// This algorithm only works if the branch mesh is fully enclosed. If there are holes then the ray may exist through those holes and hit parts
+	// of the branch we do not desire for it to hit. It also does not work well if the mesh is built up using only faces (e.g flat geometry instead
+	// of a tube for the branches)
+
+	const TArray<int32> ComponentVertexIndices = GetUniqueVertexIndices(DynamicMesh, TriangleIds);
+
+	UE::Geometry::FDynamicMeshAABBTree3 DynamicMeshAABBTree(&DynamicMesh, false);
+	DynamicMeshAABBTree.Build(TriangleIds);
+
+	TArray<FVector> NewVertexPositions;
+	NewVertexPositions.SetNum(ComponentVertexIndices.Num());
+
+	TArray<float> Scales;
+	Scales.SetNumZeroed(ComponentVertexIndices.Num());
+
+	TArray<int32> ThisVertexTriangleIDs;
+	for (int32 i = 0; i < ComponentVertexIndices.Num(); ++i)
 	{
-		TRACE_CPUPROFILER_EVENT_SCOPE(PV::Import::StaticMesh::ExtractReductionCloud);
+		const int32 VertexIndex = ComponentVertexIndices[i];
+
+		const FVector VertexNormal = FVector(DynamicMesh.GetVertexNormal(VertexIndex));
+		const FVector& VertexPosition = DynamicMesh.GetVertex(VertexIndex);
 
-		using namespace UE::Geometry;
+		const FRay3d Ray(VertexPosition, -VertexNormal);
 
-		FPScaleAttribute* PScaleAttribute = GetPScaleAttribute(DynamicMesh);
-		FFuseToAttribute* FuseToAttribute = GetFuseToAttribute(DynamicMesh);
-		check(PScaleAttribute != nullptr);
-		check(FuseToAttribute != nullptr);
+		ThisVertexTriangleIDs.Reset();
+		DynamicMesh.GetVtxTriangles(VertexIndex, ThisVertexTriangleIDs);
 		
-		FDynamicMeshAABBTree3 DynamicMeshAABBTree(&DynamicMesh, false);
-		DynamicMeshAABBTree.Build();
+		double HitTriangleDistance = 0;
+		int HitTriangleID = INDEX_NONE;
+		FVector3d HitBaryCoords = FVector3d::Zero();
+		UE::Geometry::FDynamicMeshAABBTree3::FQueryOptions QueryOptions;
+		QueryOptions.TriangleFilterF = [&](int32 TriangleID) { return !ThisVertexTriangleIDs.Contains(TriangleID); }; // Make sure to ignore triangles which includes this vertex, otherise we will instantly get a hit on ourselves
+		DynamicMeshAABBTree.FindNearestHitTriangle(Ray, HitTriangleDistance, HitTriangleID, HitBaryCoords, QueryOptions);
 
-		const auto FindOverlappingVertices = [&](const FVector& Pos, float Radius)->TArray<int32>
+		const float CollinearLimit = -0.5;
+
+		if (HitTriangleID != INDEX_NONE)
 		{
-			TArray<int32> OutVertexIndices;
-			OutVertexIndices.Reserve(DynamicMesh.VertexCount());
+			const FVector HitNormal = DynamicMesh.GetTriBaryNormal(HitTriangleID, HitBaryCoords[0], HitBaryCoords[1], HitBaryCoords[2]);
+			const FVector HitLocation = DynamicMesh.GetTriBaryPoint(HitTriangleID, HitBaryCoords[0], HitBaryCoords[1], HitBaryCoords[2]);
 
-			const double RadiusSqr = Radius * Radius;
+			const float Dot = HitNormal.Dot(VertexNormal);
+			const bool bIsOpposingAndCollinear = Dot < CollinearLimit;
+			if (bIsOpposingAndCollinear)
+			{
+				const FVector Center = (HitLocation + VertexPosition) * 0.5;
+				Scales[i] = (VertexPosition - Center).Size();
+				NewVertexPositions[i] = Center;
+			}
+		}
 
-			FDynamicMeshAABBTree3::FTreeTraversal TreeTraversal;
+		if (Scales[i] == 0)
+		{
+			// Could not find opposing point using ray intersection, try using FindNearestVertex instead.
+			const int32 MaxNumSamples = 3;
+			for (int32 j = 1; j <= MaxNumSamples; ++j)
+			{
+				double NearestDistSqr = 0.f;
+				const float MaxSampleDist = 100; // Value is arbitrary, should probably be based on the bounds of the mesh
+				const FVector SamplePosition = VertexPosition - VertexNormal * MaxSampleDist / j;
+				const int32 NearestVertexID = DynamicMeshAABBTree.FindNearestVertex(SamplePosition, NearestDistSqr, MaxSampleDist);
+				if (NearestVertexID == INDEX_NONE)
+				{
+					continue;
+				}
+
+				const FVector OtherVertexPosition = DynamicMesh.GetVertex(NearestVertexID);
+				const FVector OtherVertexNormal = FVector(DynamicMesh.GetVertexNormal(NearestVertexID));
 
-			TreeTraversal.NextBoxF = [&](const FAxisAlignedBox3d& Box, int Depth)->bool
+				const float Dot = OtherVertexNormal.Dot(VertexNormal);
+				const bool bIsOpposingAndCollinear = Dot < CollinearLimit;
+				if (!bIsOpposingAndCollinear)
+				{
+					continue;
+				}
+
+				const FVector OpposingPosition = FVector::PointPlaneProject(VertexPosition, OtherVertexPosition, -VertexNormal);
+				const FVector Center = (VertexPosition + OpposingPosition) * 0.5;
+				Scales[i] = (VertexPosition - Center).Size();
+				NewVertexPositions[i] = Center;
+				break;
+			}
+		}
+	}
+
+	// Remove extreme radius values
+	{
+		// In the below example, point "a" may have a normal which causes the ray to hit point "b" which is not the outcome we would expect.
+		// We solve this by taking the average radius of all vertices and removing any radius which lies way outside the average.
+		// This works well as long as the branches we're operating on have a fairly uniform size.
+		//
+		//       /    /
+		//      /  b*/
+		//     /    /
+		//    /    /
+		//    \*a /
+		//   /   /
+		//  /    \
+		// /    /
+		 
+		int32 RadiiNum = 0;
+		float AvgRadius = 0;
+		for (float Scale : Scales)
+		{
+			if (Scale > 0)
 			{
-				return Box.DistanceSquared(Pos) <= RadiusSqr;
-			};
+				RadiiNum++;
+				AvgRadius += Scale;
+			}
+		}
 
-			TreeTraversal.NextTriangleF = [&](int32 TID)
+		if (RadiiNum > 0)
+		{
+			AvgRadius /= RadiiNum;
+				
+			for (int32 i = 0; i < ComponentVertexIndices.Num(); ++i)
 			{
-				const FIndex3i Triangle = DynamicMesh.GetTriangle(TID);
-				for (int j = 0; j < 3; ++j)
+				const float Scale = Scales[i];
+				if (Scale > 0)
 				{
-					double VertexDistSqr = DistanceSquared(DynamicMesh.GetVertex(Triangle[j]), Pos);
-					if (VertexDistSqr <= RadiusSqr)
+					const float Ratio = Scale / AvgRadius;
+					if (Ratio > 3)
 					{
-						OutVertexIndices.AddUnique(Triangle[j]);
+						Scales[i] = 0;
 					}
 				}
-			};
-
-			DynamicMeshAABBTree.DoTraversal(TreeTraversal);
+			}
+		}
+	}
 
-			return OutVertexIndices;
-		};
+	auto& VertexScaleAttribute = VertexScaleAttributeDefinition.GetAttributeChecked(DynamicMesh);
 
-		TArray<int32> VertexIndices;
-		VertexIndices.Reserve(DynamicMesh.VertexCount());
-		for (int32 VertexIndex : DynamicMesh.VertexIndicesItr())
+	for (int32 i = 0; i < ComponentVertexIndices.Num(); ++i)
+	{
+		if (Scales[i] > 0)
 		{
-			VertexIndices.Add(VertexIndex);
+			const int32 VertexIndex = ComponentVertexIndices[i];
+			DynamicMesh.SetVertex(VertexIndex, NewVertexPositions[i]);
+			VertexScaleAttribute.SetValue(VertexIndex, Scales[i]);
 		}
+	}
+}
+
+static TAutoConsoleVariable<int32> ComputeVerticesCenterVersion(
+	TEXT("PVE.Import.ComputeVerticesCenter"),
+	2,
+	TEXT("PVE.Import.ComputeVerticesCenter {VERSION}"),
+	ECVF_Default
+);
+
+static void ComputeVerticesCenterAndRadius_Internal(UE::Geometry::FDynamicMesh3& DynamicMesh, const TArray<int32>& TriangleIds)
+{
+	TRACE_CPUPROFILER_EVENT_SCOPE(ComputeVerticesCenterAndRadius_Internal);
+
+	const int32 UseVersion = ComputeVerticesCenterVersion.GetValueOnAnyThread();
+	if (UseVersion == 1)
+	{
+		ComputeVerticesCenterAndRadius_Internal_V1(DynamicMesh, TriangleIds);
+	}
+	else
+	{
+		ComputeVerticesCenterAndRadius_Internal_V2(DynamicMesh, TriangleIds);
+	}
+}
+
+static void ComputeVerticesCenterAndRadius(UE::Geometry::FDynamicMesh3& DynamicMesh)
+{
+	TRACE_CPUPROFILER_EVENT_SCOPE(ComputeVerticesCenterAndRadius);
+
+	// We execute this per component (detached mesh) as that often correlates to individual branches which significantly improves
+	// the results as we won't accidentally find an opposing vertex in a different branch from the one we're currently operating on. 
+	UE::Geometry::FMeshConnectedComponents ConnectedComponents(&DynamicMesh);
+	ConnectedComponents.FindConnectedTriangles(); 
+
+	ParallelFor(ConnectedComponents.Num(), [&](int32 Index)
+	{
+		ComputeVerticesCenterAndRadius_Internal(DynamicMesh, ConnectedComponents.GetComponent(Index).Indices);
+	});
+}
+
+static void ExtractReductionCloud_Internal(UE::Geometry::FDynamicMesh3& DynamicMesh, const TArray<int32>& TriangleIds, float SearchRadiusInflate)
+{
+	TRACE_CPUPROFILER_EVENT_SCOPE(ExtractReductionCloud_Internal);
+
+	const auto& VertexScaleAttribute = VertexScaleAttributeDefinition.GetAttributeChecked(DynamicMesh);
+	auto& FuseToAttribute = FuseToAttributeDefinition.GetAttributeChecked(DynamicMesh);
 		
-		TArray<bool> Removed; 
-		Removed.SetNumZeroed(DynamicMesh.MaxVertexID());
+	UE::Geometry::FDynamicMeshAABBTree3 DynamicMeshAABBTree(&DynamicMesh, false);
+	DynamicMeshAABBTree.Build(TriangleIds);
+
+	TArray<int32> VertexIndices = GetUniqueVertexIndices(DynamicMesh, TriangleIds);
+
+	const int32 MaxVertexID = DynamicMesh.MaxVertexID();
 		
-		TArray<bool> Checked;
-		Checked.SetNumZeroed(DynamicMesh.MaxVertexID());
+	TArray<bool> Removed; 
+	Removed.SetNumZeroed(MaxVertexID);
 		
-		TArray<float> Radii;
-		Radii.SetNumUninitialized(DynamicMesh.MaxVertexID());
+	TArray<bool> Checked;
+	Checked.SetNumZeroed(MaxVertexID);
 		
-		for (int32 i = 0; i < VertexIndices.Num(); i++)
-		{
-			PScaleAttribute->GetValue(VertexIndices[i], &Radii[VertexIndices[i]]);
-			Radii[VertexIndices[i]] *= 2.0;
+	TArray<float> Radii;
+	Radii.SetNumUninitialized(MaxVertexID);
+	
+	for (int32 i = 0; i < VertexIndices.Num(); i++)
+	{
+		const int32 VertexIndex = VertexIndices[i];
+		Radii[VertexIndex] = VertexScaleAttribute.GetValue(VertexIndex);
+	}
+
+	{
+		TRACE_CPUPROFILER_EVENT_SCOPE(SortByRadius);
+		VertexIndices.Sort([&](int32 A, int32 B) { return Radii[A] > Radii[B]; });
+	}
+
+	for (int32 i = 0; i < VertexIndices.Num(); i++)
+	{
+		const int32 VertexIndex = VertexIndices[i];
+
+		if (Removed[VertexIndex])
+		{ 
+			continue; 
 		}
 
+		Checked[VertexIndex] = true;
+
+		const FVector VertexPosition = DynamicMesh.GetVertex(VertexIndex);
+		const float VertexRadius = Radii[VertexIndex];
+		if (VertexRadius <= 0)
 		{
-			TRACE_CPUPROFILER_EVENT_SCOPE(SortByRadius);
-			VertexIndices.Sort([&](int32 A, int32 B) { return Radii[A] > Radii[B]; });
+			FuseToAttribute.SetValue(VertexIndex, VertexIndex);
+			continue;
 		}
 
-		for (int i = 0; i < VertexIndices.Num(); i++)
+		const float SearchRadius = VertexRadius * SearchRadiusInflate;
+		const TArray<int32> OverlappingVertices = FindVerticesInRadius(DynamicMeshAABBTree, VertexPosition, SearchRadius);
+		for (int32 j = 0; j < OverlappingVertices.Num(); j++)
 		{
-			const int32 VertexIndex = VertexIndices[i];
-
-			if (Removed[VertexIndex])
-			{ 
-				continue; 
+			const int32 OverlappingVertexIndex = OverlappingVertices[j];
+				
+			if (OverlappingVertexIndex == VertexIndex 
+				|| Removed[OverlappingVertexIndex] 
+				|| Checked[OverlappingVertexIndex])
+			{
+				continue;
 			}
 
-			Checked[VertexIndex] = true;
+			const float OverlappingVertexRadius = Radii[OverlappingVertexIndex];
 
-			const FVector VertexPosition = DynamicMesh.GetVertex(VertexIndex);
-			const float VertexRadius = Radii[VertexIndex];
-
-			const TArray<int32> OverlappingVertices = FindOverlappingVertices(VertexPosition, VertexRadius);
-			for (int j = 0; j < OverlappingVertices.Num(); j++)
+			if (OverlappingVertexRadius > 0)
 			{
-				const int32 OverlappingVertexIndex = OverlappingVertices[j];
-				if (OverlappingVertexIndex == VertexIndex || Removed[OverlappingVertexIndex] || Checked[OverlappingVertexIndex])
-				{
-					continue;
-				}
-
 				const FVector OverlappingVertexPosition = DynamicMesh.GetVertex(OverlappingVertexIndex);
-				const float OverlappingVertexRadius = Radii[OverlappingVertexIndex];
-				const float OverBounds = OverlappingVertexRadius + VertexRadius;
-
-				const FVector Dist = (OverlappingVertexPosition - VertexPosition).GetAbs() - OverBounds;
-				if (Dist.X > 0 || Dist.Y > 0 || Dist.Z > 0)
+				const float CombinedRadius = OverlappingVertexRadius + VertexRadius;
+				const float CombinedRadiusSqr = CombinedRadius * CombinedRadius;
+				const float DistSquared = (OverlappingVertexPosition - VertexPosition).SizeSquared();
+				if (DistSquared > CombinedRadiusSqr)
 				{
 					continue;
 				}
-
-				Removed[OverlappingVertexIndex] = true;
-				FuseToAttribute->SetValue(OverlappingVertexIndex, &VertexIndex);
 			}
 
-			PScaleAttribute->SetValue(VertexIndex, &VertexRadius);
-			FuseToAttribute->SetValue(VertexIndex, &VertexIndex);
+			Removed[OverlappingVertexIndex] = true;
+			FuseToAttribute.SetValue(OverlappingVertexIndex, VertexIndex);
 		}
+
+		FuseToAttribute.SetValue(VertexIndex, VertexIndex);
 	}
+}
 	
-	void RadialFuse(UE::Geometry::FDynamicMesh3& DynamicMesh)
+static void FuseVerticesUsingFuseAttribute(UE::Geometry::FDynamicMesh3& DynamicMesh)
+{
+	TRACE_CPUPROFILER_EVENT_SCOPE(FuseVerticesUsingFuseAttribute);
+
+	auto& FuseToAttribute = FuseToAttributeDefinition.GetAttributeChecked(DynamicMesh);
+	auto& FillerVertexAttribute = FillerVertexAttributeDefinition.GetAttributeChecked(DynamicMesh);
+
+	TArray<int32> TriangleIndices;
+	TriangleIndices.Reserve(DynamicMesh.TriangleCount());
+	for (int32 TID : DynamicMesh.TriangleIndicesItr())
 	{
-		TRACE_CPUPROFILER_EVENT_SCOPE(PV::Import::StaticMesh::RadialFuse);
+		TriangleIndices.Add(TID);
+	}
 
-		FFuseToAttribute* FuseToAttribute = GetFuseToAttribute(DynamicMesh);
-		check(FuseToAttribute != nullptr);
+	for (int32 i = 0; i < TriangleIndices.Num(); i++)
+	{
+		const int32 TriangleIndex = TriangleIndices[i];
+		const UE::Geometry::FIndex3i VertexIndices = DynamicMesh.GetTriangle(TriangleIndex);
 
-		TArray<int32> TriangleIndices;
-		TriangleIndices.Reserve(DynamicMesh.TriangleCount());
-		for (int32 TID : DynamicMesh.TriangleIndicesItr())
+		TArray<int32, TInlineAllocator<3>> NewTriangleIndices;
+		for (int32 j = 0; j < 3; j++) 
 		{
-			TriangleIndices.Add(TID);
+			const int32 FuseTo = FuseToAttribute.GetValue(VertexIndices[j]);
+			NewTriangleIndices.AddUnique(FuseTo);
+
+			FillerVertexAttribute.SetValue(VertexIndices[j], false);
 		}
 
-		for (int i = 0; i < TriangleIndices.Num(); i++)
-		{
-			const int32 TriangleIndex = TriangleIndices[i];
+		DynamicMesh.RemoveTriangle(TriangleIndex, false);
 
-			const UE::Geometry::FIndex3i VertexIndices = DynamicMesh.GetTriangle(TriangleIndex);
-			
-			TArray<int, TInlineAllocator<3>> NewTriangleIndices;
-			for (int j = 0; j < 3; j++) 
-			{
-				int fuse_to;
-				FuseToAttribute->GetValue(VertexIndices[j], &fuse_to);
-				NewTriangleIndices.AddUnique(fuse_to);
-			}
+		if (NewTriangleIndices.Num() == 2)
+		{ 
+			const FVector P1 = DynamicMesh.GetVertex(NewTriangleIndices[0]);
+			const FVector P2 = DynamicMesh.GetVertex(NewTriangleIndices[1]);
 
-			DynamicMesh.RemoveTriangle(TriangleIndex, false);
-			
-			if (NewTriangleIndices.Num() == 2)
-			{ 
-				const FVector P1 = DynamicMesh.GetVertex(NewTriangleIndices[0]);
-				const FVector P2 = DynamicMesh.GetVertex(NewTriangleIndices[1]);
+			const FVector P3 = GenerateFillerVertex(P1, P2);
+			const int32 P3Index = DynamicMesh.AppendVertex(P3);
+			DynamicMesh.AppendTriangle(NewTriangleIndices[0], NewTriangleIndices[1], P3Index);
 
-				// Since we don't have edges for the Dynamic mesh we add a new vertex at the center of the edge, offset by a random vector
-				const FVector P3 = ((P1 + P2) / 2) + (P1 - P2).GetSafeNormal().Cross(FVector::ForwardVector).Cross(FVector::UpVector);
-				
-				const int32 P3Index = DynamicMesh.AppendVertex(P3);
-				DynamicMesh.AppendTriangle(NewTriangleIndices[0], NewTriangleIndices[1], P3Index);
-			}
-			else if (NewTriangleIndices.Num() == 3)
-			{
-				DynamicMesh.AppendTriangle(NewTriangleIndices[0], NewTriangleIndices[1], NewTriangleIndices[2]);
-			}
-			else if (NewTriangleIndices.Num() > 3)
-			{
-				check(false); // This should not be able to happen
-			}
+			FillerVertexAttribute.SetValue(P3Index, true);
+		}
+		else if (NewTriangleIndices.Num() == 3)
+		{
+			DynamicMesh.AppendTriangle(NewTriangleIndices[0], NewTriangleIndices[1], NewTriangleIndices[2]);
+		}
+		else if (NewTriangleIndices.Num() > 3)
+		{
+			check(false); // This should not be able to happen
 		}
-
-		Internal::RemoveUnusedVertices(DynamicMesh);
 	}
-	
-	void GroupPathPoints(UE::Geometry::FDynamicMesh3& DynamicMesh)
+
+	RemoveUnusedVertices(DynamicMesh);
+}
+
+static void CollapseBranchesToLine(UE::Geometry::FDynamicMesh3& DynamicMesh)
+{
+	TRACE_CPUPROFILER_EVENT_SCOPE(CollapseBranchesToLine);
+
+	UE::Geometry::FMeshConnectedComponents ConnectedComponents(&DynamicMesh);
+	ConnectedComponents.FindConnectedTriangles();
+
+	ParallelFor(ConnectedComponents.Num(), [&](int32 Index)
 	{
-		TRACE_CPUPROFILER_EVENT_SCOPE(PV::Import::StaticMesh::GroupPathPoints);
+		ExtractReductionCloud_Internal(DynamicMesh, ConnectedComponents.GetComponent(Index).Indices, 2.f);
+	});
+
+	FuseVerticesUsingFuseAttribute(DynamicMesh);
+}
 
-		FStartGroupAttribute* StartGroupAttribute = GetStartGroupAttribute(DynamicMesh);
-		FEndGroupAttribute* EndGroupAttribute = GetEndGroupAttribute(DynamicMesh);
-		FSeedAttribute* SeedAttribute = GetSeedAttribute(DynamicMesh);
-		FSeedZPosAttribute* SeedZPosAttribute = GetSeedZPosAttribute(DynamicMesh);
-		check(StartGroupAttribute != nullptr);
-		check(EndGroupAttribute != nullptr);
-		check(SeedAttribute != nullptr);
-		check(SeedZPosAttribute != nullptr);
+static void FindEndPoints(UE::Geometry::FDynamicMesh3& DynamicMesh)
+{
+	TRACE_CPUPROFILER_EVENT_SCOPE(FindEndPoints);
+	
+	auto& EndPointAttribute = EndPointAttributeDefinition.GetAttributeChecked(DynamicMesh);
+	auto& FillerVertexAttribute = FillerVertexAttributeDefinition.GetAttributeChecked(DynamicMesh);
 
-		const UE::Geometry::FAxisAlignedBox3d BoxBounds = DynamicMesh.GetBounds();
+	TArray<double> xComps;
+	TArray<double> yComps;
+	TArray<double> zComps;
+	for (int32 VertexIndex : DynamicMesh.VertexIndicesItr())
+	{
+		const FVector VertexPosition = DynamicMesh.GetVertex(VertexIndex);
 
-		TArray<double> xComps;
-		TArray<double> yComps;
-		TArray<double> zComps;
-		for (int32 VertexIndex : DynamicMesh.VertexIndicesItr())
+		if (FillerVertexAttribute.GetValue(VertexIndex))
 		{
-			const FVector VertexPosition = DynamicMesh.GetVertex(VertexIndex);
-			
-			const int32 EdgeCount = DynamicMesh.GetVtxEdgeCount(VertexIndex);
-			check(EdgeCount > 0);
+			continue;
+		}
 
-			xComps.Reset(EdgeCount);
-			yComps.Reset(EdgeCount);
-			zComps.Reset(EdgeCount);
+		const int32 EdgeCount = DynamicMesh.GetVtxEdgeCount(VertexIndex);
+		check(EdgeCount > 0);
 
-			for (int32 EdgeID : DynamicMesh.VtxEdgesItr(VertexIndex))
+		xComps.Reset(EdgeCount);
+		yComps.Reset(EdgeCount);
+		zComps.Reset(EdgeCount);
+
+		for (int32 EdgeID : DynamicMesh.VtxEdgesItr(VertexIndex))
+		{
+			const UE::Geometry::FIndex2i VertexIDs = DynamicMesh.GetEdgeV(EdgeID);
+			const int32 OtherVertexID = VertexIDs.A == VertexIndex ? VertexIDs.B : VertexIDs.A;
+
+			if (FillerVertexAttribute.GetValue(OtherVertexID))
 			{
-				const UE::Geometry::FIndex2i VertexIDs = DynamicMesh.GetEdgeV(EdgeID);
-				const int32 OtherVertexID = VertexIDs.A == VertexIndex ? VertexIDs.B : VertexIDs.A;
-
-				const FVector NeighbourVertexPosition = DynamicMesh.GetVertex(OtherVertexID);
-				const FVector Dir = (NeighbourVertexPosition - VertexPosition).GetSafeNormal();
-				xComps.Add(Dir.X);
-				yComps.Add(Dir.Y);
-				zComps.Add(Dir.Z);
+				continue;
 			}
 
-			const double* MaxXComp = Algo::MaxElement(xComps);
-			const double* MinXComp = Algo::MinElement(xComps);
-			const double* MaxYComp = Algo::MaxElement(yComps);
-			const double* MinYComp = Algo::MinElement(yComps);
-			const double* MaxZComp = Algo::MaxElement(zComps);
-			const double* MinZComp = Algo::MinElement(zComps);
+			const FVector NeighbourVertexPosition = DynamicMesh.GetVertex(OtherVertexID);
+			const FVector Dir = (NeighbourVertexPosition - VertexPosition).GetSafeNormal();
+			xComps.Add(Dir.X);
+			yComps.Add(Dir.Y);
+			zComps.Add(Dir.Z);
+		}
+		
+		check(xComps.Num() > 0);
+
+		const double* MaxXComp = Algo::MaxElement(xComps);
+		const double* MinXComp = Algo::MinElement(xComps);
+		const double* MaxYComp = Algo::MaxElement(yComps);
+		const double* MinYComp = Algo::MinElement(yComps);
+		const double* MaxZComp = Algo::MaxElement(zComps);
+		const double* MinZComp = Algo::MinElement(zComps);
+
+		const double xDiff = *MaxXComp - *MinXComp;
+		const double yDiff = *MaxYComp - *MinYComp;
+		const double zDiff = *MaxZComp - *MinZComp;
+		const double MaxDiff = FMath::Max3(xDiff, yDiff, zDiff);
+
+		const bool bIsEndPoint = MaxDiff <= 1.0;
+		EndPointAttribute.SetValue(VertexIndex, bIsEndPoint);
+	}
+}
 
-			const double xDiff = *MaxXComp - *MinXComp;
-			const double yDiff = *MaxYComp - *MinYComp;
-			const double zDiff = *MaxZComp - *MinZComp;
-			const double MaxDiff = FMath::Max3(xDiff, yDiff, zDiff);
+static void ConnectBranchesToParent(UE::Geometry::FDynamicMesh3& DynamicMesh)
+{
+	TRACE_CPUPROFILER_EVENT_SCOPE(ConnectBranchesToParent);
+
+	// This algorithm takes the tips of all the branches which have been individually collapsed into a line, and tries to connect them to their parent.
+	// This is done by performing an overlap query at the tip of each end point, and connecting that tip to the closest overlapping vertex which
+	// is not part of the same mesh component.
+	//
+	// Issues with this algorithm:
+	// 1. Since we do not know what is the "tip" and the "base" of the branch at this point, we perform the query for each
+	// end point and just hope that the tips do not overlap any other vertex. They could however be overlapping, which would lead to the tips becoming
+	// fused and for loops to be created in the tree.
+	// 2. We only fuse the end point to the parent if the radii of the end point and the nearby vertex overlap. This may lead to endpoints not being
+	// connected if the closest vertex is too far away (even if the end-point is very close to the parent mesh itself). This would lead to the entire
+	// branch being excluded from the final tree.
+
+	const auto& VertexScaleAttribute = VertexScaleAttributeDefinition.GetAttributeChecked(DynamicMesh);
+	auto& EndPointAttribute = EndPointAttributeDefinition.GetAttributeChecked(DynamicMesh);
+	auto& FillerVertexAttribute = FillerVertexAttributeDefinition.GetAttributeChecked(DynamicMesh);
+
+	const UE::Geometry::FDynamicMeshAABBTree3 DynamicMeshAABBTree(&DynamicMesh);
+	UE::Geometry::FMeshConnectedComponents ConnectedComponents(&DynamicMesh);
+	ConnectedComponents.FindConnectedVertices();
+
+	TMap<int32, int32> VertexComponentMap;
+	VertexComponentMap.Reserve(DynamicMesh.VertexCount());
+
+	for (int32 i = 0; i < ConnectedComponents.Components.Num(); ++i)
+	{
+		for (int32 VertexID : ConnectedComponents.Components[i].Indices)
+		{
+			VertexComponentMap.Add(VertexID, i);
+		}
+	}
+
+	const TArray<int32> EndPoints = EndPointAttribute.FindAllNonZero();
+	TArray<int32> MergeMap;
+	MergeMap.SetNum(EndPoints.Num());
+	for (int32 i = 0; i < EndPoints.Num(); ++i)
+	{
+		const int32 EndPointVertexIndex = EndPoints[i];
+		
+		const FVector VertexPosition = DynamicMesh.GetVertex(EndPointVertexIndex);
+		const float VertexRadius = VertexScaleAttribute.GetValue(EndPointVertexIndex);
+		const float SearchRadius = VertexRadius * 4.f;
+		const TArray<int32> OverlappingVertices = FindVerticesInRadius(DynamicMeshAABBTree, VertexPosition, SearchRadius);
+
+		float ClosestVertexDistSqr = MAX_flt;
+		int32 VertexToMergeWith = INDEX_NONE;
+		for (int32 OverlappingVertexIndex : OverlappingVertices)
+		{
+			if (EndPointAttribute.GetValue(OverlappingVertexIndex)
+				|| FillerVertexAttribute.GetValue(OverlappingVertexIndex))
+			{
+				continue;
+			}
 
-			if (MaxDiff > 1.0)
+			int32* ComponentIndex1 = VertexComponentMap.Find(EndPointVertexIndex);
+			int32* ComponentIndex2 = VertexComponentMap.Find(OverlappingVertexIndex);
+			check(ComponentIndex1 && ComponentIndex2);
+			if (*ComponentIndex1 == *ComponentIndex2)
 			{
 				continue;
 			}
 
-			const float Z = FMath::GetMappedRangeValueClamped(FVector2D(BoxBounds.Min.Z, BoxBounds.Max.Z), FVector2D(0, 1), VertexPosition.Z);
-			const bool bIsEndGroup = Z > 0.1;
-			const bool bIsStartGroup = !bIsEndGroup;
+			const float OverlappingVertexRadius = VertexScaleAttribute.GetValue(OverlappingVertexIndex);
+			const FVector OverlappingVertexPosition = DynamicMesh.GetVertex(OverlappingVertexIndex);
+			const float CombinedRadius = OverlappingVertexRadius + VertexRadius;
+			const float CombinedRadiusSqr = CombinedRadius * CombinedRadius;
+			const float DistSquared = (OverlappingVertexPosition - VertexPosition).SizeSquared();
+			if (DistSquared > CombinedRadiusSqr)
+			{
+				continue;
+			}
 
-			StartGroupAttribute->SetValue(VertexIndex, &bIsStartGroup);
-			EndGroupAttribute->SetValue(VertexIndex, &bIsEndGroup);
-			SeedAttribute->SetValue(VertexIndex, &bIsStartGroup);
-			SeedZPosAttribute->SetValue(VertexIndex, &Z);
+			const FVector OtherVertexPosition = DynamicMesh.GetVertex(OverlappingVertexIndex);
+			const float DistSqr = (VertexPosition - OtherVertexPosition).SizeSquared();
+			if (DistSqr < ClosestVertexDistSqr)
+			{
+				VertexToMergeWith = OverlappingVertexIndex;
+				ClosestVertexDistSqr = DistSqr;
+			}
 		}
+
+		MergeMap[i] = VertexToMergeWith;
 	}
-	
-	void ComputeDist(UE::Geometry::FDynamicMesh3& DynamicMesh)
-	{
-		TRACE_CPUPROFILER_EVENT_SCOPE(PV::Import::StaticMesh::ComputeDist);
 
-		FStartGroupAttribute* StartGroupAttribute = GetStartGroupAttribute(DynamicMesh);
-		FEndGroupAttribute* EndGroupAttribute = GetEndGroupAttribute(DynamicMesh);
-		FDistAttribute* DistAttribute = GetDistAttribute(DynamicMesh);
-		check(StartGroupAttribute != nullptr);
-		check(EndGroupAttribute != nullptr);
-		check(DistAttribute != nullptr);
+	for (int32 i = 0; i < EndPoints.Num(); ++i)
+	{
+		const int32 EndPointVertexIndex = EndPoints[i];
+		const int32 VertexToMergeWith = MergeMap[i];
 
-		typedef UE::Geometry::TMeshDijkstra<UE::Geometry::FDynamicMesh3> FMeshDijkstra;
-		FMeshDijkstra MeshDijkstra(&DynamicMesh);
+		if (VertexToMergeWith == INDEX_NONE)
+		{
+			continue;
+		}
+		const FVector VertexPosition = DynamicMesh.GetVertex(EndPointVertexIndex);
+		const FVector OtherVertexPosition = DynamicMesh.GetVertex(VertexToMergeWith);
 
-		TArray<FMeshDijkstra::FSeedPoint> SeedPoints;
+		const FVector Normal = (VertexPosition - OtherVertexPosition).GetSafeNormal();
+		const FVector NewVertexPosition = OtherVertexPosition + (Normal * 0.1); // We want to keep this vertex to not lose the radius information, so just move it really close and connect it to the other vertex
 
-		for (int32 VertexIndex : DynamicMesh.VertexIndicesItr())
+		DynamicMesh.SetVertex(EndPointVertexIndex, NewVertexPosition);
+		
+		const TArray<int32> Neighbours = GetVertexNeighbours(DynamicMesh, EndPointVertexIndex);
+		for (int32 Neighbour : Neighbours)
 		{
-			bool bIsStart = false;
-			StartGroupAttribute->GetValue(VertexIndex, &bIsStart);
-
-			if (bIsStart)
+			if (FillerVertexAttribute.GetValue(Neighbour))
 			{
-				FMeshDijkstra::FSeedPoint& Seed = SeedPoints.AddDefaulted_GetRef();
-				Seed.PointID = VertexIndex;
+				const TArray<int32> FillerNeighbours = GetVertexNeighbours(DynamicMesh, Neighbour);
+				const int32* AdjacentVertex = FillerNeighbours.FindByPredicate([&](const auto& Value) { return Value != EndPointVertexIndex && Value != Neighbour; });
+				check(AdjacentVertex != nullptr);
+				const FVector AdjacentVertexPosition = DynamicMesh.GetVertex(*AdjacentVertex);
+				const FVector NewFillerVertex = GenerateFillerVertex(AdjacentVertexPosition, NewVertexPosition);
+				
+				DynamicMesh.SetVertex(Neighbour, NewFillerVertex);
 			}
 		}
+		
+		const FVector NewFillerVertex = GenerateFillerVertex(NewVertexPosition, OtherVertexPosition);
+		const int32 NewFillerVertexIndex = DynamicMesh.AppendVertex(NewFillerVertex);
+		DynamicMesh.AppendTriangle(EndPointVertexIndex, VertexToMergeWith, NewFillerVertexIndex);
+		EndPointAttribute.SetValue(EndPointVertexIndex, false);
+	}
+}
 
-		for (int32 VertexIndex : DynamicMesh.VertexIndicesItr())
-		{
-			float Dist = 0;
+static void FindRootPoints(UE::Geometry::FDynamicMesh3& DynamicMesh)
+{
+	TRACE_CPUPROFILER_EVENT_SCOPE(FindRootPoints);
 
-			MeshDijkstra.ComputeToTargetPoint(SeedPoints, VertexIndex);
+	// This function currently just finds all the end points that are close to the bottom of the plant bounding box and converts them to root points.
+	// This could result in issues if a plant has multiple points touching the floor. 
+	
+	auto& RootPointAttribute = RootPointAttributeDefinition.GetAttributeChecked(DynamicMesh);
+	auto& EndPointAttribute = EndPointAttributeDefinition.GetAttributeChecked(DynamicMesh);
 
-			TArray<int32> VertexPath;
-			if (!MeshDijkstra.FindPathToNearestSeed(VertexIndex, VertexPath))
-			{
-				const bool FalseValue = false;
-				EndGroupAttribute->SetValue(VertexIndex, &FalseValue);
-			}
-			else
-			{
-				// TODO: Is there a built-in function in TMeshDijkstra for this? Feels expensive to have to compute this for each point...
-				for (int32 i = 0; i < VertexPath.Num() - 1; ++i)
-				{
-					int32 Vertex0ID = VertexPath[i];
-					int32 Vertex1ID = VertexPath[i + 1];
-					Dist += (DynamicMesh.GetVertex(Vertex0ID) - DynamicMesh.GetVertex(Vertex1ID)).Size();
-				}
-			}
+	const UE::Geometry::FAxisAlignedBox3d BoxBounds = DynamicMesh.GetBounds();
 
-			DistAttribute->SetValue(VertexIndex, &Dist);
+	for (int32 VertexIndex : EndPointAttribute.FindAllNonZero())
+	{
+		const FVector VertexPosition = DynamicMesh.GetVertex(VertexIndex);
+
+		const float Z = FMath::GetMappedRangeValueClamped(FVector2D(BoxBounds.Min.Z, BoxBounds.Max.Z), FVector2D(0, 1), VertexPosition.Z);
+		const bool bIsRootPoint = Z <= 0.1;
+		if (bIsRootPoint)
+		{
+			RootPointAttribute.SetValue(VertexIndex, true);
+			EndPointAttribute.SetValue(VertexIndex, false);			
 		}
 	}
+}
 
-	void DrawBranches(const UE::Geometry::FDynamicMesh3& DynamicMesh, TArray<TArray<int32>>& OutBranchCurves)
-	{
-		TRACE_CPUPROFILER_EVENT_SCOPE(PV::Import::StaticMesh::DrawBranches);
+static void ComputeLengthFromRoot(UE::Geometry::FDynamicMesh3& DynamicMesh)
+{
+	TRACE_CPUPROFILER_EVENT_SCOPE(ComputeLengthFromRoot);
+
+	// This algorithm starts at the root point and walks along every vertex, keeping track of its distance to the root point and stores that distance
+	// in the LengthFromRootAttribute. This distance is used to find the closest path between the root point and the tips in DrawBranchCurves.
 
-		const FStartGroupAttribute* StartGroupAttribute = GetStartGroupAttribute(DynamicMesh);
-		const FEndGroupAttribute* EndGroupAttribute = GetEndGroupAttribute(DynamicMesh);
-		const FDistAttribute* DistAttribute = GetDistAttribute(DynamicMesh);
-		check(StartGroupAttribute != nullptr);
-		check(EndGroupAttribute != nullptr);
-		check(DistAttribute != nullptr);
+	auto& RootPointsAttribute = RootPointAttributeDefinition.GetAttributeChecked(DynamicMesh);
+	auto& LengthFromRootAttribute = LengthFromRootAttributeDefinition.GetAttributeChecked(DynamicMesh);
+	auto& FillerVertexAttribute = FillerVertexAttributeDefinition.GetAttributeChecked(DynamicMesh);
 
-		const auto GetVertexNeighbours = [&](int32 VertexID)->TArray<int32>
+	for (int32 VertexID : DynamicMesh.VertexIndicesItr())
+	{
+		LengthFromRootAttribute.SetValue(VertexID, -1);
+	}
+	
+	TArray<bool> bHasComputedVertex;
+	bHasComputedVertex.SetNumZeroed(DynamicMesh.MaxVertexID());
+	const TFunction<void(int32, const FVector&, float)> ComputeLengthFromRootRecursive = [&](int32 VertexIndex, const FVector& PrevVertexLocation, float PrevVertexLengthFromRoot)
+	{
+		if (FillerVertexAttribute.GetValue(VertexIndex))
 		{
-			TArray<int32> OutNeighbours;
-			OutNeighbours.Reserve(DynamicMesh.GetVtxEdgeCount(VertexID));
-			for (int32 EdgeID : DynamicMesh.VtxEdgesItr(VertexID))
-			{
-				const UE::Geometry::FIndex2i VertexIDs = DynamicMesh.GetEdgeV(EdgeID);
-				const int32 OtherVertexID = VertexIDs.A == VertexID ? VertexIDs.B : VertexIDs.A;
-				OutNeighbours.Add(OtherVertexID);
-			}
-			return OutNeighbours;
-		};
+			return;
+		}
+			
+		if (bHasComputedVertex[VertexIndex])
+		{
+			return;
+		}
 
-		const auto FindNextPoint = [&](int32 PointIndex, float MinDist)->int32
+		bHasComputedVertex[VertexIndex] = true;
+			
+		const FVector VertexPosition = DynamicMesh.GetVertex(VertexIndex);
+		const float VertexLengthFromRoot = PrevVertexLengthFromRoot + (VertexPosition - PrevVertexLocation).Size();
+		LengthFromRootAttribute.SetValue(VertexIndex, VertexLengthFromRoot);
+		
+		const TArray<int32> Neighbours = GetVertexNeighbours(DynamicMesh, VertexIndex);
+		for (int32 Neighbour : Neighbours)
 		{
-			int32 OutNextPointIndex = INDEX_NONE;
-			const TArray<int32> Neighbours = GetVertexNeighbours(PointIndex);
-			for (int32 Neighbour : Neighbours)
-			{
-				float Dist = 0;
-				DistAttribute->GetValue(Neighbour, &Dist);
-				if (Dist > MinDist)
-				{ 
-					continue; 
-				}
-				MinDist = Dist;
-				OutNextPointIndex = Neighbour;
-			}
-			return OutNextPointIndex;
-		};
+			ComputeLengthFromRootRecursive(Neighbour, VertexPosition, VertexLengthFromRoot);
+		}
+	};
+
+	for (int32 VertexIndex : RootPointsAttribute.FindAllNonZero())
+	{
+		const FVector VertexPosition = DynamicMesh.GetVertex(VertexIndex);
+		ComputeLengthFromRootRecursive(VertexIndex, VertexPosition, 0.0);
+	}
+}
+
+static void DrawBranchCurves(const UE::Geometry::FDynamicMesh3& DynamicMesh, TArray<TArray<int32>>& OutBranchCurves)
+{
+	TRACE_CPUPROFILER_EVENT_SCOPE(DrawBranchCurves);
 
-		const auto ExpandPointGroup = [&](const UE::Geometry::TDynamicMeshVertexAttribute<bool, 1>* GroupAttribute)->TArray<int32>
+	// This algorithm finds the shortest path from all the end points (tips) to the root point and stores that in OutBranchCurves
+
+	const auto& RootPointAttribute = RootPointAttributeDefinition.GetAttributeChecked(DynamicMesh);
+	const auto& EndPointAttribute = EndPointAttributeDefinition.GetAttributeChecked(DynamicMesh);
+	const auto& LengthFromRootAttribute = LengthFromRootAttributeDefinition.GetAttributeChecked(DynamicMesh);
+	const auto& FillerVertexAttribute = FillerVertexAttributeDefinition.GetAttributeChecked(DynamicMesh);
+
+	const auto FindNextPoint = [&](int32 PointIndex, float MinLengthFromRoot)->int32
+	{
+		int32 OutNextPointIndex = INDEX_NONE;
+		const TArray<int32> Neighbours = GetVertexNeighbours(DynamicMesh, PointIndex);
+		for (int32 Neighbour : Neighbours)
 		{
-			TArray<int32> OutVertexIndices;
-			OutVertexIndices.Reserve(DynamicMesh.VertexCount());
-			if (GroupAttribute == nullptr)
+			if (FillerVertexAttribute.GetValue(Neighbour))
 			{
-				for (int32 VertexIndex : DynamicMesh.VertexIndicesItr())
-				{
-					OutVertexIndices.Add(VertexIndex);
-				}
+				continue;
 			}
-			else
-			{
-				for (int32 VertexIndex : DynamicMesh.VertexIndicesItr())
-				{
-					bool bIsInGroup = false;
-					GroupAttribute->GetValue(VertexIndex, &bIsInGroup);
 
-					if (bIsInGroup)
-					{
-						OutVertexIndices.Add(VertexIndex);
-					}
-				}
+			const float LengthFromRoot = LengthFromRootAttribute.GetValue(Neighbour);
+			check(LengthFromRoot >= 0);
+			
+			if (LengthFromRoot >= MinLengthFromRoot)
+			{ 
+				continue; 
 			}
+			MinLengthFromRoot = LengthFromRoot;
+			OutNextPointIndex = Neighbour;
+		}
+		return OutNextPointIndex;
+	};
 
-			return OutVertexIndices;
-		};
+	const TArray<int32> EndPoints = EndPointAttribute.FindAllNonZero();
 
-		const TArray<int32> EndPoints = ExpandPointGroup(EndGroupAttribute);
-		const TArray<int32> EtartPoints = ExpandPointGroup(StartGroupAttribute);
-		
-		for (int32 i = 0; i < EndPoints.Num(); i++)
-		{
-			int32 CurrentPointIndex = EndPoints[i];
+	TArray<bool> bHasProcessedVertex;
+	for (int32 i = 0; i < EndPoints.Num(); i++)
+	{
+		int32 CurrentPointIndex = EndPoints[i];
 			
-			float MinDist = 0;
-			DistAttribute->GetValue(CurrentPointIndex, &MinDist);
+		float MinLengthFromRoot = LengthFromRootAttribute.GetValue(CurrentPointIndex);
+		if (MinLengthFromRoot < 0)
+		{
+			continue; // This end point was unable to trace down to a root point, perhaps log a warning?
+		}
 			
-			TArray<int32> CurvePoints = { CurrentPointIndex };
+		TArray<int32> CurvePoints = { CurrentPointIndex };
 
-			for (int j = 0; j < DynamicMesh.VertexCount(); j++)
-			{
-				const int32 NextPointIndex = FindNextPoint(CurrentPointIndex, MinDist);
-				if (NextPointIndex == INDEX_NONE)
-				{ 
-					break; 
-				}
-
-				CurvePoints.Add(NextPointIndex);
-
-				bool bIsStart = false;
-				StartGroupAttribute->GetValue(NextPointIndex, &bIsStart);
-				if (bIsStart)
-				{ 
-					break; 
-				}
+		bHasProcessedVertex.Reset();
+		bHasProcessedVertex.SetNumZeroed(DynamicMesh.MaxVertexID());
+		
+		while (true)
+		{
+			const int32 NextPointIndex = FindNextPoint(CurrentPointIndex, MinLengthFromRoot);
+			if (NextPointIndex == INDEX_NONE)
+			{ 
+				break; 
+			}
 
-				CurrentPointIndex = NextPointIndex;
+			if (!ensureMsgf(!bHasProcessedVertex[NextPointIndex], TEXT("Infinite loop detected")))
+			{
+				CurvePoints.Empty();
+				break;
 			}
+			
+			bHasProcessedVertex[NextPointIndex] = true;
 
-			if (CurvePoints.Num() > 1)
+			CurvePoints.Add(NextPointIndex);
+
+			if (RootPointAttribute.GetValue(NextPointIndex))
 			{
-				OutBranchCurves.Add(MoveTemp(CurvePoints));
+				break; 
 			}
+
+			CurrentPointIndex = NextPointIndex;
+			MinLengthFromRoot = LengthFromRootAttribute.GetValue(CurrentPointIndex);
+		}
+
+		if (CurvePoints.Num() > 1)
+		{
+			OutBranchCurves.Add(MoveTemp(CurvePoints));
 		}
 	}
+}
 
-	void RemoveOverlaps(const UE::Geometry::FDynamicMesh3& DynamicMesh, TArray<TArray<int32>>& InOutBranchCurves)
+static void RemoveBranchCurveOverlaps(const UE::Geometry::FDynamicMesh3& DynamicMesh, TArray<TArray<int32>>& InOutBranchCurves)
+{
+	TRACE_CPUPROFILER_EVENT_SCOPE(RemoveBranchCurveOverlaps);
+
+	// Since DrawBranchCurves finds the full path from all end points to the root point, there will be overlaps in the branch curves
+	//
+	// In the below example, the curves between a->root and b->root would overlap from c->root. This algorithm finds those overlapping sections
+	// and only keeps that section on the longest curve (a->root in this case).
+	//
+	// a     
+	//  \   b
+	//   \ /
+	//    |c
+	//    |
+	//    |
+	//   root
+	//
+	// Issues with this algorithm: Using the length of the branch to determine the branch parent only works well on small trees. As the tree grows larger
+	// the likelihood of a smaller branches have a longer distance to root gets bigger. A better value for selecting the parent would probably be the
+	// branch radius.
+
+	TArray<float> BranchLengths;
+	BranchLengths.SetNum(InOutBranchCurves.Num());
+
+	int32 MaxNumVertices = 0;
+	for (int32 i = 0; i < InOutBranchCurves.Num(); ++i)
 	{
-		TRACE_CPUPROFILER_EVENT_SCOPE(PV::Import::StaticMesh::RemoveOverlaps);
-
-		TArray<float> BranchLengths;
-		BranchLengths.SetNum(InOutBranchCurves.Num());
-
-		int32 MaxNumVertices = 0;
-		for (int32 i = 0; i < InOutBranchCurves.Num(); ++i)
+		const TArray<int32>& BranchCurve = InOutBranchCurves[i];
+		MaxNumVertices += BranchCurve.Num();
+		for (int32 j = 0; j < BranchCurve.Num() - 1; ++j)
 		{
-			const TArray<int32>& BranchCurve = InOutBranchCurves[i];
-			MaxNumVertices += BranchCurve.Num();
-			for (int32 j = 0; j < BranchCurve.Num() - 1; ++j)
-			{
-				const FVector& Vertex0 = DynamicMesh.GetVertex(BranchCurve[j]);
-				const FVector& Vertex1 = DynamicMesh.GetVertex(BranchCurve[j + 1]);
-				BranchLengths[i] += (Vertex0 - Vertex1).Size();
-			}
+			const FVector& Vertex0 = DynamicMesh.GetVertex(BranchCurve[j]);
+			const FVector& Vertex1 = DynamicMesh.GetVertex(BranchCurve[j + 1]);
+			BranchLengths[i] += (Vertex0 - Vertex1).Size();
 		}
+	}
 
-		struct FVertexLongestBranch
-		{
-			int32 BranchIndex = INDEX_NONE;
-			float BranchLength = 0;
-		};
+	struct FVertexLongestBranch
+	{
+		int32 BranchIndex = INDEX_NONE;
+		float BranchLength = 0;
+	};
 
-		TMap<int32, FVertexLongestBranch> VerticesLongestBranch;
-		VerticesLongestBranch.Reserve(MaxNumVertices);
+	TMap<int32, FVertexLongestBranch> VerticesLongestBranch;
+	VerticesLongestBranch.Reserve(MaxNumVertices);
 
-		for (int32 i = 0; i < InOutBranchCurves.Num(); ++i)
-		{
-			const TArray<int32>& BranchCurve = InOutBranchCurves[i];
-			const float BranchLength = BranchLengths[i];
+	for (int32 i = 0; i < InOutBranchCurves.Num(); ++i)
+	{
+		const TArray<int32>& BranchCurve = InOutBranchCurves[i];
+		const float BranchLength = BranchLengths[i];
 			
-			for (int32 j = 0; j < BranchCurve.Num(); ++j)
-			{
-				const int32 VertexIndex = BranchCurve[j];
-				FVertexLongestBranch& VertexLongestBranch = VerticesLongestBranch.FindOrAdd(VertexIndex);
+		for (int32 j = 0; j < BranchCurve.Num(); ++j)
+		{
+			const int32 VertexIndex = BranchCurve[j];
+			FVertexLongestBranch& VertexLongestBranch = VerticesLongestBranch.FindOrAdd(VertexIndex);
 
-				if (BranchLength > VertexLongestBranch.BranchLength)
-				{
-					VertexLongestBranch.BranchLength = BranchLength;
-					VertexLongestBranch.BranchIndex = i;
-				}
+			if (BranchLength > VertexLongestBranch.BranchLength)
+			{
+				VertexLongestBranch.BranchLength = BranchLength;
+				VertexLongestBranch.BranchIndex = i;
 			}
 		}
+	}
 
-		for (int32 i = InOutBranchCurves.Num() - 1; i >= 0; --i)
-		{
-			TArray<int32>& BranchCurve = InOutBranchCurves[i];
+	for (int32 i = InOutBranchCurves.Num() - 1; i >= 0; --i)
+	{
+		TArray<int32>& BranchCurve = InOutBranchCurves[i];
 			
-			for (int32 j = BranchCurve.Num() - 1; j >= 0; --j)
-			{
-				const int32 ThisVertexIndex = BranchCurve[j];
-				const int32 ThisLongestBranchIndex = VerticesLongestBranch[ThisVertexIndex].BranchIndex;
+		for (int32 j = BranchCurve.Num() - 1; j >= 0; --j)
+		{
+			const int32 ThisVertexIndex = BranchCurve[j];
+			const int32 ThisLongestBranchIndex = VerticesLongestBranch[ThisVertexIndex].BranchIndex;
 				
-				if (ThisLongestBranchIndex != i)
+			if (ThisLongestBranchIndex != i)
+			{
+				if (j == 0)
+				{
+					BranchCurve.RemoveAt(j);
+				}
+				else
 				{
-					if (j == 0)
+					const int32 NextVertexIndex = BranchCurve[j - 1];
+					const int32 NextLongestBranchIndex = VerticesLongestBranch[NextVertexIndex].BranchIndex;
+					if (NextLongestBranchIndex != i)
 					{
 						BranchCurve.RemoveAt(j);
 					}
-					else
-					{
-						const int32 NextVertexIndex = BranchCurve[j - 1];
-						const int32 NextLongestBranchIndex = VerticesLongestBranch[NextVertexIndex].BranchIndex;
-						if (NextLongestBranchIndex != i)
-						{
-							BranchCurve.RemoveAt(j);
-						}
-					}
 				}
 			}
+		}
 
-			if (BranchCurve.Num() < 2)
-			{
-				InOutBranchCurves.RemoveAtSwap(i);
-			}
+		if (BranchCurve.Num() < 2)
+		{
+			InOutBranchCurves.RemoveAtSwap(i);
 		}
 	}
-};
+}
 
-namespace PV::Import::StaticMesh::Internal2
+static void ExtractBranchesFromBranchCurves(const UE::Geometry::FDynamicMesh3& DynamicMesh,
+	const TArray<TArray<int32>>& InBranchCurves, 
+	TArray<FVector3f>& OutPoints, 
+	TArray<float>& OutPointsRadii, 
+	TArray<TArray<int32>>& OutBranchPointIndices
+)
 {
-	void ExtractBranchesFromBranchCurves(const UE::Geometry::FDynamicMesh3& DynamicMesh, 
-		const TArray<TArray<int32>>& InBranchCurves, 
-		TArray<FVector3f>& OutPoints, 
-		TArray<float>& OutPointsRadii, 
-		TArray<TArray<int32>>& OutBranchPointIndices
-	)
-	{
-		TRACE_CPUPROFILER_EVENT_SCOPE(PV::Import::StaticMesh::ExtractBranchesFromBranchCurves);
-
-		const PV::Import::StaticMesh::Internal::FPScaleAttribute* PScaleAttribute = PV::Import::StaticMesh::Internal::GetPScaleAttribute(DynamicMesh);
-		const PV::Import::StaticMesh::Internal::FStartGroupAttribute* EndGroupAttribute = PV::Import::StaticMesh::Internal::GetEndGroupAttribute(DynamicMesh);
-		check(PScaleAttribute != nullptr);
-		check(EndGroupAttribute != nullptr);
-
-		int32 MaxNrOfPoints = 0;
-		for (const TArray<int32>& BranchCurve : InBranchCurves)
-		{
-			MaxNrOfPoints += BranchCurve.Num();
-		}
+	TRACE_CPUPROFILER_EVENT_SCOPE(ExtractBranchesFromBranchCurves);
 
-		OutPoints.Reserve(MaxNrOfPoints);
-		OutPointsRadii.Reserve(MaxNrOfPoints);
-		OutBranchPointIndices.Reserve(InBranchCurves.Num());
+	const auto& VertexScaleAttribute = VertexScaleAttributeDefinition.GetAttributeChecked(DynamicMesh);
+	const auto& EndPointAttribute = EndPointAttributeDefinition.GetAttributeChecked(DynamicMesh);
 
-		TMap<int32, int32> VertexIndexToPointIndex;
-		VertexIndexToPointIndex.Reserve(MaxNrOfPoints);
+	int32 MaxNrOfPoints = 0;
+	for (const TArray<int32>& BranchCurve : InBranchCurves)
+	{
+		MaxNrOfPoints += BranchCurve.Num();
+	}
 
-		for (const TArray<int32>& BranchCurve : InBranchCurves)
-		{
-			TArray<int32>& Branch = OutBranchPointIndices.AddDefaulted_GetRef();
-			Branch.Reserve(BranchCurve.Num());
+	OutPoints.Reserve(MaxNrOfPoints);
+	OutPointsRadii.Reserve(MaxNrOfPoints);
+	OutBranchPointIndices.Reserve(InBranchCurves.Num());
 
-			for (int32 i = BranchCurve.Num() - 1; i >= 0; --i) // Reverse since curves go from tip to base while branches to from base to tip
-			{
-				const int32 VertexIndex = BranchCurve[i];
-				int32 PointIndex = INDEX_NONE;
+	TMap<int32, int32> VertexIndexToPointIndex;
+	VertexIndexToPointIndex.Reserve(MaxNrOfPoints);
 
-				int32* ExistingPointIndex = VertexIndexToPointIndex.Find(VertexIndex);
-				if (!ExistingPointIndex)
-				{
-					PointIndex = OutPoints.Add(FVector3f(DynamicMesh.GetVertex(VertexIndex)));
-					
-					bool bIsInEndGroup = false;
-					EndGroupAttribute->GetValue(VertexIndex, &bIsInEndGroup);
+	TMap<int32, TArray<int32, TInlineAllocator<2>>> PointIndexToBranchIndices;
+	PointIndexToBranchIndices.Reserve(MaxNrOfPoints);
 
-					if (bIsInEndGroup)
-					{
-						OutPointsRadii.Add(0.01f);
-					}
-					else
-					{
-						float VertexRadius = 0;
-						PScaleAttribute->GetValue(VertexIndex, &VertexRadius);
-						OutPointsRadii.Add(VertexRadius);
-					}
+	for (const TArray<int32>& BranchCurve : InBranchCurves)
+	{
+		TArray<int32>& Branch = OutBranchPointIndices.AddDefaulted_GetRef();
+		Branch.Reserve(BranchCurve.Num());
+
+		for (int32 i = BranchCurve.Num() - 1; i >= 0; --i) // Reverse since curves go from tip to base while branches to from base to tip
+		{
+			const int32 VertexIndex = BranchCurve[i];
+			int32 PointIndex = INDEX_NONE;
 
-					VertexIndexToPointIndex.Add(VertexIndex, PointIndex);
+			int32* ExistingPointIndex = VertexIndexToPointIndex.Find(VertexIndex);
+			if (!ExistingPointIndex)
+			{
+				PointIndex = OutPoints.Add(FVector3f(DynamicMesh.GetVertex(VertexIndex)));
+					
+				const bool bIsEndPoint = EndPointAttribute.GetValue(VertexIndex);
+				if (bIsEndPoint)
+				{
+					OutPointsRadii.Add(0.01f);
 				}
 				else
 				{
-					PointIndex = *ExistingPointIndex;
+					const float VertexRadius = VertexScaleAttribute.GetValue(VertexIndex);
+					OutPointsRadii.Add(VertexRadius);
 				}
 
-				Branch.Add(PointIndex);
+				VertexIndexToPointIndex.Add(VertexIndex, PointIndex);
+			}
+			else
+			{
+				PointIndex = *ExistingPointIndex;
 			}
+
+			Branch.Add(PointIndex);
+			PointIndexToBranchIndices.FindOrAdd(VertexIndex).Add(OutBranchPointIndices.Num() - 1);
 		}
 	}
-};
+}
 
-PV::Import::StaticMesh::EImportResult PV::Import::StaticMesh::ImportGrowthDataFromStaticMesh(const TSoftObjectPtr<UStaticMesh>& InSourceStaticMesh, FManagedArrayCollection& OutCollection, const FImportMeshOptions& InOptions)
+PV::StaticMeshImport::EImportResult PV::StaticMeshImport::ImportGrowthDataFromStaticMesh(const TObjectPtr<UStaticMesh>& InSourceStaticMesh, FManagedArrayCollection& OutCollection, const FImportOptions& InOptions)
 {
-	TRACE_CPUPROFILER_EVENT_SCOPE(PV::Import::StaticMesh::ImportGrowthDataFromStaticMesh);
-
-	UStaticMesh* StaticMesh = InSourceStaticMesh.LoadSynchronous();
-	if (!StaticMesh)
+	TRACE_CPUPROFILER_EVENT_SCOPE(PV::StaticMeshImport::ImportGrowthDataFromStaticMesh);
+	
+	if (!InSourceStaticMesh)
 	{
 		return EImportResult::InvalidSourceMesh;
 	}
 
-	if (StaticMesh->IsCompiling())
+	if (InSourceStaticMesh->IsCompiling())
 	{
 		return EImportResult::MeshNotReady;
 	}
 
+	if (InOptions.MaterialsToKeep.Num() == 0)
+	{
+		return EImportResult::InvalidMaterialFilter;
+	}
+
 	UE::Geometry::FDynamicMesh3 DynamicMesh;
 	{
 		TRACE_CPUPROFILER_EVENT_SCOPE(StaticMeshToDynamicMesh);
 
 		FText OutErrorMessage;
 		UE::Conversion::FStaticMeshConversionOptions Options;
-		if (!UE::Conversion::StaticMeshToDynamicMesh(StaticMesh, DynamicMesh, OutErrorMessage, Options))
+		if (!UE::Conversion::StaticMeshToDynamicMesh(InSourceStaticMesh, DynamicMesh, OutErrorMessage, Options))
 		{
 			return EImportResult::InvalidMeshData;
 		}
@@ -941,93 +1327,83 @@ PV::Import::StaticMesh::EImportResult PV::Import::StaticMesh::ImportGrowthDataFr
 		return EImportResult::Success;
 	};
 
-	if (InOptions.DebugState == FImportMeshOptions::EDebugState::InitialMesh)
+	if (InOptions.DebugState == FImportOptions::EDebugState::InitialMesh)
 	{
 		return ReturnWithDebugMesh(DynamicMesh);
 	}
-
-	const int32 MaterialIndexToKeep = StaticMesh->GetMaterialIndex(InOptions.MaterialFilter);
-	Internal::KeepMaterial(DynamicMesh, MaterialIndexToKeep); // For now we only allow filtering of one material. Could easily be extended to allow for multiple materials
-	if (InOptions.DebugState == FImportMeshOptions::EDebugState::KeepMaterial)
+	
+	if (!FilterTrianglesUsingMaterials(DynamicMesh, InSourceStaticMesh, InOptions.MaterialsToKeep))
+	{
+		return EImportResult::InvalidMaterialFilter;
+	}
+	
+	if (InOptions.DebugState == FImportOptions::EDebugState::KeepMaterial)
 	{
 		return ReturnWithDebugMesh(DynamicMesh);
 	}
 
-	Internal::ComputeNormals(DynamicMesh);
-	if (InOptions.DebugState == FImportMeshOptions::EDebugState::ComputeNormals)
+	ComputeNormals(DynamicMesh);
+	if (InOptions.DebugState == FImportOptions::EDebugState::ComputeNormals)
 	{
 		return ReturnWithDebugMesh(DynamicMesh);
 	}
 
 	const UE::Geometry::FDynamicMesh3 ReferenceMesh = DynamicMesh;
-	Internal::AttachPScaleAttribute(DynamicMesh);
-	Internal::AttachFuseToAttribute(DynamicMesh);
-	Internal::AttachStartGroupAttribute(DynamicMesh);
-	Internal::AttachEndGroupAttribute(DynamicMesh);
-	Internal::AttachSeedAttribute(DynamicMesh);
-	Internal::AttachSeedZPosAttribute(DynamicMesh);
-	Internal::AttachDistAttribute(DynamicMesh);
+	VertexScaleAttributeDefinition.AttachAttribute(DynamicMesh);
 
-	/*
-	Internal::AverageNormals(DynamicMesh);
-	if (InOptions.DebugState == FImportMeshOptions::EDebugState::AverageNormals)
+	ComputeVerticesCenterAndRadius(DynamicMesh);
+	if (InOptions.DebugState == FImportOptions::EDebugState::ComputeVerticesCenterAndRadius)
 	{
-		return ReturnWithDebugMesh();
+		VertexScaleAttributeDefinition.CopyAttributeToCollection(DynamicMesh, OutCollection);
+		return ReturnWithDebugMesh(DynamicMesh);
 	}
 
-	Internal::SmoothNormalsWithLengthPreservation(DynamicMesh);
-	if (InOptions.DebugState == FImportMeshOptions::EDebugState::SmoothNormalsWithLengthPreservation)
-	{
-		return ReturnWithDebugMesh();
-	}*/
+	FuseToAttributeDefinition.AttachAttribute(DynamicMesh);
+	FillerVertexAttributeDefinition.AttachAttribute(DynamicMesh);
 
-	Internal::GetSpherical(DynamicMesh);
-	if (InOptions.DebugState == FImportMeshOptions::EDebugState::GetSpherical)
+	CollapseBranchesToLine(DynamicMesh);
+	if (InOptions.DebugState == FImportOptions::EDebugState::CollapseBranchesToLine)
 	{
-		Internal::CopyAttachedVertexAttributeToCollection<float>(DynamicMesh, OutCollection, Internal::PScaleAttributeName);
+		FuseToAttributeDefinition.CopyAttributeToCollection(DynamicMesh, OutCollection);
+		VertexScaleAttributeDefinition.CopyAttributeToCollection(DynamicMesh, OutCollection);
 		return ReturnWithDebugMesh(DynamicMesh);
 	}
 
-	/*Internal::BlurRadius(DynamicMesh);
-	if (InOptions.DebugState == FImportMeshOptions::EDebugState::BlurRadius)
-	{
-		return ReturnWithDebugMesh();
-	}*/
+	EndPointAttributeDefinition.AttachAttribute(DynamicMesh);
 
-	Internal::ExtractReductionCloud(DynamicMesh);
-	if (InOptions.DebugState == FImportMeshOptions::EDebugState::ExtractReductionCloud)
+	FindEndPoints(DynamicMesh);
+	if (InOptions.DebugState == FImportOptions::EDebugState::FindEndPoints)
 	{
-		Internal::CopyAttachedVertexAttributeToCollection<float>(DynamicMesh, OutCollection, Internal::PScaleAttributeName);
-		Internal::CopyAttachedVertexAttributeToCollection<int32>(DynamicMesh, OutCollection, Internal::FuseToAttributeName);
+		EndPointAttributeDefinition.CopyAttributeToCollection(DynamicMesh, OutCollection);
 		return ReturnWithDebugMesh(DynamicMesh);
 	}
 
-	Internal::RadialFuse(DynamicMesh);
-	if (InOptions.DebugState == FImportMeshOptions::EDebugState::RadialFuse)
+	ConnectBranchesToParent(DynamicMesh);
+	if (InOptions.DebugState == FImportOptions::EDebugState::ConnectBranchesToParent)
 	{
-		return ReturnWithDebugMesh(DynamicMesh);
+		EndPointAttributeDefinition.CopyAttributeToCollection(DynamicMesh, OutCollection);
+		return ReturnWithDebugMesh(DynamicMesh);	
 	}
 
-	Internal::GroupPathPoints(DynamicMesh);
-	if (InOptions.DebugState == FImportMeshOptions::EDebugState::GroupPathPoints)
+	RootPointAttributeDefinition.AttachAttribute(DynamicMesh);
+	FindRootPoints(DynamicMesh);
+	if (InOptions.DebugState == FImportOptions::EDebugState::FindRootPoints)
 	{
-		Internal::CopyAttachedVertexAttributeToCollection<bool>(DynamicMesh, OutCollection, Internal::StartGroupAttributeName);
-		Internal::CopyAttachedVertexAttributeToCollection<bool>(DynamicMesh, OutCollection, Internal::EndGroupAttributeName);
-		Internal::CopyAttachedVertexAttributeToCollection<bool>(DynamicMesh, OutCollection, Internal::SeedAttributeName);
-		Internal::CopyAttachedVertexAttributeToCollection<float>(DynamicMesh, OutCollection, Internal::SeedZPosAttributeName);
-
+		RootPointAttributeDefinition.CopyAttributeToCollection(DynamicMesh, OutCollection);
+		EndPointAttributeDefinition.CopyAttributeToCollection(DynamicMesh, OutCollection);
 		return ReturnWithDebugMesh(DynamicMesh);
 	}
 	
-	Internal::ComputeDist(DynamicMesh);
-	if (InOptions.DebugState == FImportMeshOptions::EDebugState::ComputeDist)
-	{
-		Internal::CopyAttachedVertexAttributeToCollection<float>(DynamicMesh, OutCollection, Internal::DistAttributeName);
+	LengthFromRootAttributeDefinition.AttachAttribute(DynamicMesh);
 
+	ComputeLengthFromRoot(DynamicMesh);
+	if (InOptions.DebugState == FImportOptions::EDebugState::ComputeLengthFromRoot)
+	{
+		LengthFromRootAttributeDefinition.CopyAttributeToCollection(DynamicMesh, OutCollection);
 		return ReturnWithDebugMesh(DynamicMesh);
 	}
 
-	const static auto AddBranchCurvesToCollecton = [](
+	const static auto Debug_AddBranchCurvesToCollection = [](
 		FManagedArrayCollection& InCollection, 
 		const TArray<TArray<int32>>& BranchCurves, 
 		const TFunction<FVector3f(int32)>& GetVertexPosFunc)
@@ -1051,25 +1427,25 @@ PV::Import::StaticMesh::EImportResult PV::Import::StaticMesh::ImportGrowthDataFr
 	};
 
 	TArray<TArray<int32>> BranchCurves;
-	Internal::DrawBranches(DynamicMesh, BranchCurves);
-	if (InOptions.DebugState == FImportMeshOptions::EDebugState::DrawBranches)
+	DrawBranchCurves(DynamicMesh, BranchCurves);
+	if (InOptions.DebugState == FImportOptions::EDebugState::DrawBranchCurves)
 	{
-		AddBranchCurvesToCollecton(OutCollection, BranchCurves, [&](int32 VertexIndex) { return FVector3f(DynamicMesh.GetVertex(VertexIndex)); });
+		Debug_AddBranchCurvesToCollection(OutCollection, BranchCurves, [&](int32 VertexIndex) { return FVector3f(DynamicMesh.GetVertex(VertexIndex)); });
 		return ReturnWithDebugMesh(ReferenceMesh);
 	}
 
-	Internal::RemoveOverlaps(DynamicMesh, BranchCurves);
-	if (InOptions.DebugState == FImportMeshOptions::EDebugState::RemoveOverlaps)
+	RemoveBranchCurveOverlaps(DynamicMesh, BranchCurves);
+	if (InOptions.DebugState == FImportOptions::EDebugState::RemoveBranchCurveOverlaps)
 	{
-		AddBranchCurvesToCollecton(OutCollection, BranchCurves, [&](int32 VertexIndex) { return FVector3f(DynamicMesh.GetVertex(VertexIndex)); });
+		Debug_AddBranchCurvesToCollection(OutCollection, BranchCurves, [&](int32 VertexIndex) { return FVector3f(DynamicMesh.GetVertex(VertexIndex)); });
 		return ReturnWithDebugMesh(ReferenceMesh);
 	}
 
 	TArray<FVector3f> Points;
 	TArray<float> PointsRadii;
 	TArray<TArray<int32>> BranchPointIndices;
-	Internal2::ExtractBranchesFromBranchCurves(DynamicMesh, BranchCurves, Points, PointsRadii, BranchPointIndices);
-	if (InOptions.DebugState == FImportMeshOptions::EDebugState::ExtractBranches)
+	ExtractBranchesFromBranchCurves(DynamicMesh, BranchCurves, Points, PointsRadii, BranchPointIndices);
+	if (InOptions.DebugState == FImportOptions::EDebugState::ExtractBranches)
 	{
 		const static FName PointsGroupName = TEXT("Points");
 		const static FName PositionAttributeName = TEXT("Position");
@@ -1086,18 +1462,18 @@ PV::Import::StaticMesh::EImportResult PV::Import::StaticMesh::ImportGrowthDataFr
 			ScaleAttribute[i] = PointsRadii[i];
 		}
 
-		AddBranchCurvesToCollecton(OutCollection, BranchPointIndices, [&](int32 VertexIndex) { return FVector3f(Points[VertexIndex]); });
+		Debug_AddBranchCurvesToCollection(OutCollection, BranchPointIndices, [&](int32 VertexIndex) { return FVector3f(Points[VertexIndex]); });
 
 		return ReturnWithDebugMesh(ReferenceMesh);
 	}
 
 	FPVPlantDescription PlantDescription;
 	PlantDescription.BranchHierarchy = PV::Import::Helper::CreateBranchHierarchyFromPoints(Points, PointsRadii, BranchPointIndices);
-
+	
 	if (!PV::Import::Helper::GenerateGrowthDataFromPlantDescription(OutCollection, PlantDescription))
 	{
-		return PV::Import::StaticMesh::EImportResult::FailedToGenerateGrowthData;
+		return EImportResult::FailedToGenerateGrowthData;
 	}
 
-	return PV::Import::StaticMesh::EImportResult::Success;
-}
+	return EImportResult::Success;
+}
\ No newline at end of file
diff --git a/Engine/Plugins/Experimental/ProceduralVegetationEditor/Source/ProceduralVegetationEditor/Private/Importer/PVImporter_StaticMesh.h b/Engine/Plugins/Experimental/ProceduralVegetationEditor/Source/ProceduralVegetationEditor/Private/Importer/PVImporter_StaticMesh.h
--- a/Engine/Plugins/Experimental/ProceduralVegetationEditor/Source/ProceduralVegetationEditor/Private/Importer/PVImporter_StaticMesh.h
+++ b/Engine/Plugins/Experimental/ProceduralVegetationEditor/Source/ProceduralVegetationEditor/Private/Importer/PVImporter_StaticMesh.h
@@ -7,7 +7,7 @@
 
 class UStaticMesh;
 
-namespace PV::Import::StaticMesh
+namespace PV::StaticMeshImport
 {
 	enum class EImportResult
 	{
@@ -16,13 +16,13 @@ namespace PV::Import::StaticMesh
 		InvalidOutputPath,
 		InvalidMeshData,
 		MeshNotReady,
-		FailedToGenerateGrowthData,
-		Canceled
+		InvalidMaterialFilter,
+		FailedToGenerateGrowthData
 	};
 
-	struct FImportMeshOptions
+	struct FImportOptions
 	{
-		FName MaterialFilter;
+		TArray<FName> MaterialsToKeep;
 
 		enum class EDebugState : uint8
 		{
@@ -30,21 +30,19 @@ namespace PV::Import::StaticMesh
 			InitialMesh,
 			KeepMaterial,
 			ComputeNormals,
-			AverageNormals,
-			SmoothNormalsWithLengthPreservation,
-			GetSpherical,
-			BlurRadius,
-			ExtractReductionCloud,
-			RadialFuse,
-			GroupPathPoints,
-			ComputeDist,
-			DrawBranches,
-			RemoveOverlaps,
+			ComputeVerticesCenterAndRadius,
+			CollapseBranchesToLine,
+			FindEndPoints,
+			ConnectBranchesToParent,
+			FindRootPoints,
+			ComputeLengthFromRoot,
+			DrawBranchCurves,
+			RemoveBranchCurveOverlaps,
 			ExtractBranches,
 			MAX
 		};
 		EDebugState DebugState = EDebugState::None;
 	};
 
-	EImportResult ImportGrowthDataFromStaticMesh(const TSoftObjectPtr<UStaticMesh>& InSourceStaticMesh, FManagedArrayCollection& OutCollection, const FImportMeshOptions& InOptions);
+	EImportResult ImportGrowthDataFromStaticMesh(const TObjectPtr<UStaticMesh>& InSourceStaticMesh, FManagedArrayCollection& OutCollection, const FImportOptions& InOptions);
 };
diff --git a/Engine/Plugins/Experimental/ProceduralVegetationEditor/Source/ProceduralVegetationEditor/Private/Importer/PVPlantDescription.cpp b/Engine/Plugins/Experimental/ProceduralVegetationEditor/Source/ProceduralVegetationEditor/Private/Importer/PVPlantDescription.cpp
--- a/Engine/Plugins/Experimental/ProceduralVegetationEditor/Source/ProceduralVegetationEditor/Private/Importer/PVPlantDescription.cpp
+++ b/Engine/Plugins/Experimental/ProceduralVegetationEditor/Source/ProceduralVegetationEditor/Private/Importer/PVPlantDescription.cpp
@@ -10,6 +10,53 @@
 
 FPVBranchHierarchyDescription::EValidateHierarchyResult FPVBranchHierarchyDescription::ValidateHierarchy() const
 {
+	if (!Branches.IsValidIndex(RootBranchIndex))
+	{
+		return EValidateHierarchyResult::InvalidRootBranchIndex;
+	}
+	
+	if (Branches[RootBranchIndex].ParentBranchIndex != INDEX_NONE)
+	{
+		return EValidateHierarchyResult::InvalidRootBranchParentIndex;
+	}
+	
+	for (int32 BranchIndex = 0; BranchIndex < Branches.Num(); BranchIndex++)
+	{
+		const FPVBranchDescription& Branch = Branches[BranchIndex];
+		if (Branch.BranchIndex != BranchIndex)
+		{
+			return EValidateHierarchyResult::BranchIndexMismatch;
+		}
+		
+		if (Branch.ParentBranchIndex != INDEX_NONE
+			&& !Branches.IsValidIndex(Branch.ParentBranchIndex))
+		{
+			return EValidateHierarchyResult::InvalidParentBranchIndex;
+		}
+		
+		if (Branch.ParentBranchIndex == INDEX_NONE
+			&& RootBranchIndex != Branch.BranchIndex)
+		{
+			return EValidateHierarchyResult::InvalidRootBranchIndex;
+		}
+		
+		for (int32 ChildBranchIndex : Branch.ChildBranchIndices)
+		{
+			if (!Branches.IsValidIndex(ChildBranchIndex))
+			{
+				return EValidateHierarchyResult::InvalidChildBranchIndex;
+			}
+		}
+		
+		for (int32 PointIndex : Branch.PointIndices)
+		{
+			if (!Points.IsValidIndex(PointIndex))
+			{
+				return EValidateHierarchyResult::InvalidPointIndex; 
+			}
+		}
+	}
+
 	TSet<int32> FoundBranches;
 	FoundBranches.Reserve(Branches.Num());
 
@@ -36,12 +83,12 @@ FPVBranchHierarchyDescription::EValidateHierarchyResult FPVBranchHierarchyDescri
 
 	if (!WalkHierarchyRecursive(RootBranchIndex))
 	{
-		return EValidateHierarchyResult::CircularHiearachy;
+		return EValidateHierarchyResult::CircularHierarchy;
 	}
 
 	if (FoundBranches.Num() != Branches.Num())
 	{
-		return EValidateHierarchyResult::DisjointedHiearachy;
+		return EValidateHierarchyResult::DisjointedHierarchy;
 	}
 
 	return EValidateHierarchyResult::Success;
@@ -65,7 +112,10 @@ FPrimitiveSceneProxy* UPVPlantDescriptionVisualizationComponent::CreateSceneProx
 
 				if (FoliageMesh)
 				{
-					MeshDescription = *FoliageMesh->GetMeshDescription(0);
+					if (FMeshDescription* FoliageMeshDescription = FoliageMesh->GetMeshDescription(0))
+					{
+						MeshDescription = *FoliageMeshDescription;
+					}
 				}
 			}
 		}
@@ -84,16 +134,13 @@ FPrimitiveSceneProxy* UPVPlantDescriptionVisualizationComponent::CreateSceneProx
 			Segments = FMath::Max(Segments, 4);
 			const float AngleStep = 2.f * UE_PI / float(Segments);
 
-			const FVector AxisY = Rotation.GetRightVector();
-			const FVector AxisZ = Rotation.GetUpVector();
-
 			float Angle = 0.f;
 			while (Segments--)
 			{
 				const FVector Vertex1 = Center + Radius * Rotation.RotateVector(FVector(0, FMath::Cos(Angle), 0) + FVector(0, 0, FMath::Sin(Angle)));
 				Angle += AngleStep;
 				const FVector Vertex2 = Center + Radius * Rotation.RotateVector(FVector(0, FMath::Cos(Angle), 0) + FVector(0, 0, FMath::Sin(Angle)));
-				PDI->DrawLine(Vertex1, Vertex2, Color, Thickness, DepthPriority);
+				PDI->DrawLine(Vertex1, Vertex2, Color, DepthPriority, Thickness);
 			}
 		}
 
@@ -150,11 +197,13 @@ FPrimitiveSceneProxy* UPVPlantDescriptionVisualizationComponent::CreateSceneProx
 					DrawBranchesRecursive(PDI, BranchHierarchy, ChildIndex, Depth + 1);
 				}
 			};
-
-			if (PlantDescription.BranchHierarchy.RootBranchIndex != INDEX_NONE)
+			
+			if (PlantDescription.BranchHierarchy.ValidateHierarchy() != FPVBranchHierarchyDescription::EValidateHierarchyResult::Success)
 			{
-				DrawBranchesRecursive(PDI, PlantDescription.BranchHierarchy, PlantDescription.BranchHierarchy.RootBranchIndex, 0);
+				return;
 			}
+
+			DrawBranchesRecursive(PDI, PlantDescription.BranchHierarchy, PlantDescription.BranchHierarchy.RootBranchIndex, 0);
 		}
 		
 		void DrawFoliage(FPrimitiveDrawInterface* PDI) const
@@ -256,4 +305,4 @@ FBoxSphereBounds UPVPlantDescriptionVisualizationComponent::CalcBounds(const FTr
 	}
 
 	return FBoxSphereBounds(BBox).TransformBy(LocalToWorld);
-}
\ No newline at end of file
+}
diff --git a/Engine/Plugins/Experimental/ProceduralVegetationEditor/Source/ProceduralVegetationEditor/Private/Importer/PVPlantDescription.h b/Engine/Plugins/Experimental/ProceduralVegetationEditor/Source/ProceduralVegetationEditor/Private/Importer/PVPlantDescription.h
--- a/Engine/Plugins/Experimental/ProceduralVegetationEditor/Source/ProceduralVegetationEditor/Private/Importer/PVPlantDescription.h
+++ b/Engine/Plugins/Experimental/ProceduralVegetationEditor/Source/ProceduralVegetationEditor/Private/Importer/PVPlantDescription.h
@@ -54,8 +54,15 @@ struct FPVBranchHierarchyDescription
 
 	enum class EValidateHierarchyResult
 	{
-		DisjointedHiearachy,
-		CircularHiearachy,
+		InvalidRootBranchIndex,
+		InvalidRootBranchParentIndex,
+		InvalidParentBranchIndex,
+		InvalidBranchIndex,
+		InvalidChildBranchIndex,
+		InvalidPointIndex,
+		BranchIndexMismatch,
+		DisjointedHierarchy,
+		CircularHierarchy,
 		Success
 	};
 	EValidateHierarchyResult ValidateHierarchy() const;
diff --git a/Engine/Plugins/Experimental/ProceduralVegetationEditor/Source/ProceduralVegetationEditor/Private/PVEditorInterfaceImpl.cpp b/Engine/Plugins/Experimental/ProceduralVegetationEditor/Source/ProceduralVegetationEditor/Private/PVEditorInterfaceImpl.cpp
--- a/Engine/Plugins/Experimental/ProceduralVegetationEditor/Source/ProceduralVegetationEditor/Private/PVEditorInterfaceImpl.cpp
+++ b/Engine/Plugins/Experimental/ProceduralVegetationEditor/Source/ProceduralVegetationEditor/Private/PVEditorInterfaceImpl.cpp
@@ -9,12 +9,12 @@
 
 bool PV::FPVEditorInterfaceImpl::ImportGrowthDataFromStaticMesh(const FPVImportStaticMeshParams& InParams, FManagedArrayCollection& OutCollection)
 {
-	PV::Import::StaticMesh::FImportMeshOptions ImportMeshOptions;
-	ImportMeshOptions.MaterialFilter = InParams.MaterialFilter;
+	PV::StaticMeshImport::FImportOptions ImportMeshOptions;
+	ImportMeshOptions.MaterialsToKeep = InParams.MaterialsToKeep;
 	
-	static_assert(static_cast<int32>(EPVImportStaticMeshDebugState::MAX) == static_cast<int32>(PV::Import::StaticMesh::FImportMeshOptions::EDebugState::MAX));
-	ImportMeshOptions.DebugState = static_cast<PV::Import::StaticMesh::FImportMeshOptions::EDebugState>(static_cast<int32>(InParams.DebugState));
+	static_assert(static_cast<int32>(EPVImportStaticMeshDebugState::MAX) == static_cast<int32>(PV::StaticMeshImport::FImportOptions::EDebugState::MAX));
+	ImportMeshOptions.DebugState = static_cast<PV::StaticMeshImport::FImportOptions::EDebugState>(static_cast<int32>(InParams.DebugState));
 
-	const PV::Import::StaticMesh::EImportResult ImportResult = PV::Import::StaticMesh::ImportGrowthDataFromStaticMesh(InParams.StaticMeshAsset, OutCollection, ImportMeshOptions);
-	return ImportResult == PV::Import::StaticMesh::EImportResult::Success;
-}
\ No newline at end of file
+	const PV::StaticMeshImport::EImportResult ImportResult = PV::StaticMeshImport::ImportGrowthDataFromStaticMesh(InParams.StaticMeshAsset, OutCollection, ImportMeshOptions);
+	return ImportResult == PV::StaticMeshImport::EImportResult::Success;
+}
