Create new PVE node for importing growth data from a 2D image, port 2D trace image algorithm from Houdini to PVE

#Synced-CL 48983295

diff --git a/Engine/Plugins/Experimental/ProceduralVegetationEditor/Source/ProceduralVegetation/Private/Nodes/PVTexture2DImporterSettings.cpp b/Engine/Plugins/Experimental/ProceduralVegetationEditor/Source/ProceduralVegetation/Private/Nodes/PVTexture2DImporterSettings.cpp
new file mode 100644
--- /dev/null
+++ b/Engine/Plugins/Experimental/ProceduralVegetationEditor/Source/ProceduralVegetation/Private/Nodes/PVTexture2DImporterSettings.cpp
@@ -0,0 +1,78 @@
+// Copyright Epic Games, Inc. All Rights Reserved.
+
+#include "PVTexture2DImporterSettings.h"
+#include "IPVEditorInterface.h"
+#include "PCGContext.h"
+#include "ProceduralVegetationModule.h"
+#include "DataTypes/PVGrowthData.h"
+#include "ProceduralVegetationPreset.h"
+#include "PVImportTexture2DParams.h"
+
+#define LOCTEXT_NAMESPACE "PVTexture2DImporterSettings"
+
+#if WITH_EDITOR
+FText UPVTexture2DImporterSettings::GetDefaultNodeTitle() const
+{ 
+	return LOCTEXT("NodeTitle", "Procedural Vegetation Texture 2D Importer"); 
+}
+
+FText UPVTexture2DImporterSettings::GetNodeTooltipText() const
+{	
+	return LOCTEXT("NodeTooltip", "");
+}
+#endif
+
+TArray<FPCGPinProperties> UPVTexture2DImporterSettings::InputPinProperties() const
+{
+	return TArray<FPCGPinProperties>();
+}
+
+FPCGDataTypeIdentifier UPVTexture2DImporterSettings::GetOutputPinTypeIdentifier() const
+{
+	return FPCGDataTypeIdentifier{ FPVDataTypeInfoGrowth::AsId() };
+}
+
+FPCGElementPtr UPVTexture2DImporterSettings::CreateElement() const
+{
+	return MakeShared<FPVTexture2DImporterElement>();
+}
+
+bool FPVTexture2DImporterElement::ExecuteInternal(FPCGContext* InContext) const
+{
+	TRACE_CPUPROFILER_EVENT_SCOPE(FPVTexture2DImporterElement::ExecuteInternal);
+
+	check(InContext);
+
+	const UPVTexture2DImporterSettings* Settings = InContext->GetInputSettings<UPVTexture2DImporterSettings>();
+	check(Settings);
+
+	if (!Settings->Params.Texture2DAsset)
+	{
+		PCGLog::LogWarningOnGraph(LOCTEXT("MissingTexture2DAsset", "Texture2D asset not set"), InContext);
+		return true;
+	}
+
+#if WITH_EDITOR
+	if (!PV::GPVEditorInterface)
+	{
+		return true;
+	}
+	
+	FManagedArrayCollection Collection;
+	if (!PV::GPVEditorInterface->ImportGrowthDataFromTexture2D(Settings->Params, Collection))
+	{
+		return false;
+	}
+
+	UPVGrowthData* OutVariantData = FPCGContext::NewObject_AnyThread<UPVGrowthData>(InContext);
+	OutVariantData->Initialize(MoveTemp(Collection));
+
+	FPCGTaggedData& CollectionOutput = InContext->OutputData.TaggedData.Emplace_GetRef();
+	CollectionOutput.Data = OutVariantData;
+	CollectionOutput.Pin = TEXT("Output");
+#endif
+
+	return true;
+}
+
+#undef LOCTEXT_NAMESPACE
\ No newline at end of file
diff --git a/Engine/Plugins/Experimental/ProceduralVegetationEditor/Source/ProceduralVegetation/Private/Nodes/PVTexture2DImporterSettings.h b/Engine/Plugins/Experimental/ProceduralVegetationEditor/Source/ProceduralVegetation/Private/Nodes/PVTexture2DImporterSettings.h
new file mode 100644
--- /dev/null
+++ b/Engine/Plugins/Experimental/ProceduralVegetationEditor/Source/ProceduralVegetation/Private/Nodes/PVTexture2DImporterSettings.h
@@ -0,0 +1,41 @@
+// Copyright Epic Games, Inc. All Rights Reserved.
+#pragma once
+
+#include "CoreMinimal.h"
+#include "Nodes/PVBaseSettings.h"
+#include "PVImportTexture2DParams.h"
+
+#include "PVTexture2DImporterSettings.generated.h"
+
+UCLASS(BlueprintType, ClassGroup = (Procedural))
+class UPVTexture2DImporterSettings : public UPVBaseSettings
+{
+	GENERATED_BODY()
+
+public:
+	
+#if WITH_EDITOR
+	virtual FName GetDefaultNodeName() const override { return FName(TEXT("ProceduralVegetationTexture2DImporter")); }
+	virtual FText GetDefaultNodeTitle() const override;
+	virtual FText GetNodeTooltipText() const override;
+	virtual EPCGSettingsType GetType() const override { return EPCGSettingsType::InputOutput; }
+	virtual TArray<EPVRenderType> GetDefaultRenderType() const override { return TArray{ EPVRenderType::PointData }; }
+	virtual TArray<EPVRenderType> GetSupportedRenderTypes() const override { return TArray{ EPVRenderType::PointData, EPVRenderType::Mesh }; }
+#endif
+
+protected:
+	virtual TArray<FPCGPinProperties> InputPinProperties() const override;
+	virtual FPCGDataTypeIdentifier GetOutputPinTypeIdentifier() const override;
+	virtual FPCGElementPtr CreateElement() const override;
+
+public:
+	UPROPERTY(EditAnywhere, Category="Importer Settings", meta = (ShowOnlyInnerProperties))
+	FPVImportTexture2DParams Params;
+
+};
+
+class FPVTexture2DImporterElement : public FPVBaseElement
+{
+protected:
+	virtual bool ExecuteInternal(FPCGContext* InContext) const override;
+};
diff --git a/Engine/Plugins/Experimental/ProceduralVegetationEditor/Source/ProceduralVegetation/Public/IPVEditorInterface.h b/Engine/Plugins/Experimental/ProceduralVegetationEditor/Source/ProceduralVegetation/Public/IPVEditorInterface.h
--- a/Engine/Plugins/Experimental/ProceduralVegetationEditor/Source/ProceduralVegetation/Public/IPVEditorInterface.h
+++ b/Engine/Plugins/Experimental/ProceduralVegetationEditor/Source/ProceduralVegetation/Public/IPVEditorInterface.h
@@ -3,6 +3,7 @@
 #pragma once
 
 #include "PVImportStaticMeshParams.h"
+#include "PVImportTexture2DParams.h"
 
 struct FManagedArrayCollection;
 struct FPVImportStaticMeshParams;
@@ -17,6 +18,7 @@ namespace PV
 		virtual ~IPVEditorInterface() = default;
 
 		virtual bool ImportGrowthDataFromStaticMesh(const FPVImportStaticMeshParams& InParams, FManagedArrayCollection& OutCollection) = 0;
+		virtual bool ImportGrowthDataFromTexture2D(const FPVImportTexture2DParams& InParams, FManagedArrayCollection& OutCollection) = 0;
 	};
 
 	extern PROCEDURALVEGETATION_API IPVEditorInterface* GPVEditorInterface;
diff --git a/Engine/Plugins/Experimental/ProceduralVegetationEditor/Source/ProceduralVegetation/Public/PVImportTexture2DParams.h b/Engine/Plugins/Experimental/ProceduralVegetationEditor/Source/ProceduralVegetation/Public/PVImportTexture2DParams.h
new file mode 100644
--- /dev/null
+++ b/Engine/Plugins/Experimental/ProceduralVegetationEditor/Source/ProceduralVegetation/Public/PVImportTexture2DParams.h
@@ -0,0 +1,36 @@
+// Copyright Epic Games, Inc. All Rights Reserved.
+#pragma once
+
+#include "Engine/EngineTypes.h"
+
+#include "PVImportTexture2DParams.generated.h"
+
+class UTexture2D;
+
+UENUM(BlueprintType)
+enum class EPVImportTexture2DDebugState : uint8
+{
+	None,
+	Texture,
+	SampleImage,
+	ExtractContourLines,
+	MAX UMETA(Hidden)
+};
+
+USTRUCT(BlueprintType)
+struct FPVImportTexture2DParams
+{
+	GENERATED_BODY()
+
+	UPROPERTY(EditAnywhere, Category = "Asset", meta=(ClampMin=0, UIMin=0))
+	int32 SampleWidth = 128;
+	
+	UPROPERTY(EditAnywhere, Category = "Asset", meta=(ClampMin=0, UIMin=0))
+    int32 SampleHeight = 128;
+
+	UPROPERTY(EditAnywhere, Category = "Asset")
+	TObjectPtr<UTexture2D> Texture2DAsset = nullptr;
+
+	UPROPERTY(EditAnywhere, Category = "Asset")
+	EPVImportTexture2DDebugState DebugState = EPVImportTexture2DDebugState::None;
+};
diff --git a/Engine/Plugins/Experimental/ProceduralVegetationEditor/Source/ProceduralVegetationEditor/Private/Importer/PVImporter_Texture2D.cpp b/Engine/Plugins/Experimental/ProceduralVegetationEditor/Source/ProceduralVegetationEditor/Private/Importer/PVImporter_Texture2D.cpp
new file mode 100644
--- /dev/null
+++ b/Engine/Plugins/Experimental/ProceduralVegetationEditor/Source/ProceduralVegetationEditor/Private/Importer/PVImporter_Texture2D.cpp
@@ -0,0 +1,283 @@
+// Copyright Epic Games, Inc. All Rights Reserved.
+
+#include "PVImporter_Texture2D.h"
+
+#include "IntVectorTypes.h"
+#include "GeometryCollection/ManagedArrayCollection.h"
+
+#include "Engine/Texture2D.h"
+#include "Materials/MaterialInstanceDynamic.h"
+#include "Materials/MaterialInterface.h"
+
+#include "GeometryCollection/GeometryCollection.h"
+#include "GeometryCollection/GeometryCollectionUtility.h"
+
+namespace Internal
+{
+	using FVector2i = UE::Geometry::FVector2i;
+	constexpr float DebugQuadSize = 100;
+
+	const static FName PixelsDebugGroupName = TEXT("PixelDebug");
+	const static FName PixelPositionsAttributeName = TEXT("Position");
+	const static FName PixelValueAttributeName = TEXT("Value");
+	const static FName LineSegmentsGroupName = TEXT("LineSegments");
+	const static FName LineSegmentsAttributeName = TEXT("LineSegment");
+	
+	static int32 GetPixelIndex(const FVector2i& PixelCoors, int32 ImageWidth)
+	{
+		return PixelCoors.X + PixelCoors.Y * ImageWidth;
+	}
+	
+	template<typename T>
+	static const T& SampleImageBuffer(const TArray<T>& Buffer, const int32 Stride, const FVector2i& XYCoord)
+	{
+		const int32 Index = GetPixelIndex(XYCoord, Stride);
+		check(Buffer.IsValidIndex(Index));
+		return Buffer[Index];
+	}
+
+	static FVector2f PixelCoordToUV(const FVector2i& PixelCoord, const FVector2i& ImageSize)
+	{
+		const float U = PixelCoord.X / static_cast<float>(ImageSize.X - 1);
+		const float V = PixelCoord.Y / static_cast<float>(ImageSize.Y - 1);
+		return FVector2f(U, V);
+	}
+
+	static void AddDebugPixelValuesToCollection(FManagedArrayCollection& InOutCollection, const TArray<bool>& ImageData, const FVector2i& ImageSize)
+	{
+		check(ImageData.Num() == ImageSize.X * ImageSize.Y);
+		
+		InOutCollection.AddGroup(PixelsDebugGroupName);
+		InOutCollection.AddElements(ImageData.Num(), PixelsDebugGroupName);
+		auto& PixelPositionsAttribute = InOutCollection.AddAttribute<FVector3f>(PixelPositionsAttributeName, PixelsDebugGroupName);
+		auto& PixelValueAttribute = InOutCollection.AddAttribute<bool>(PixelValueAttributeName, PixelsDebugGroupName);
+
+		for (int32 x = 0; x < ImageSize.X; x++)
+		{
+			for (int32 y = 0; y < ImageSize.Y; y++)
+			{
+				const FVector2i PixelCoords = FVector2i(x, y);
+				const FVector2f UV = PixelCoordToUV(PixelCoords, ImageSize);
+				
+				const int32 PixelIndex = GetPixelIndex(PixelCoords, ImageSize.X);
+				PixelPositionsAttribute[PixelIndex] = FVector3f(0, UV.X * DebugQuadSize, (1.0 - UV.Y) * DebugQuadSize);
+				PixelValueAttribute[PixelIndex] = ImageData[PixelIndex];
+			}
+		}
+	}
+
+	static void AddDebugLineSegmentsToCollection(FManagedArrayCollection& InOutCollection, const TArray<TPair<FVector2f, FVector2f>>& InLineSegments)
+	{
+		InOutCollection.AddGroup(LineSegmentsGroupName);
+		InOutCollection.AddElements(InLineSegments.Num(), LineSegmentsGroupName);
+		
+		auto& LineSegmentAttribute = InOutCollection.AddAttribute<TArray<FVector3f>>(LineSegmentsAttributeName, LineSegmentsGroupName);
+		for (int32 i = 0; i < InLineSegments.Num(); i++)
+		{
+			const auto& It = InLineSegments[i];
+			const FVector2f LineStart = It.Key * DebugQuadSize;
+			const FVector2f LineEnd = It.Value * DebugQuadSize;
+			LineSegmentAttribute[i] = { FVector3f(0, LineStart.X, DebugQuadSize - LineStart.Y), FVector3f(0, LineEnd.X, DebugQuadSize - LineEnd.Y) };
+		}
+	}
+
+	static TSharedRef<FGeometryCollection> GenerateDebugQuad(UTexture2D& Texture2D)
+	{
+		const TArray<FVector> Points = {
+			FVector(0, 0, 1),
+			FVector(0, 1, 1),
+			FVector(0, 1, 0),
+			FVector(0, 0, 0)
+		};
+		const TArray<FVector3f> Normals = {
+			FVector3f(-1, 0, 0),
+			FVector3f(-1, 0, 0),
+			FVector3f(-1, 0, 0),
+			FVector3f(-1, 0, 0)
+		};
+		const TArray<UE::Geometry::FVector3i> Triangles = {
+			{ 0, 1, 2 },
+			{ 2, 3, 0 }
+		};
+		const TArray<FVector2f> UVs = {
+			FVector2f(0, 0),
+			FVector2f(1, 0),
+			FVector2f(1, 1),
+			FVector2f(0, 1),
+		};
+		const TSharedPtr<FGeometryCollection> MeshGeomCollection = GeometryCollection::MakeMeshElement(
+			Points,
+			Normals,
+			Triangles,
+			UVs,
+			FTransform(FRotator::ZeroRotator, FVector::ZeroVector, FVector(Internal::DebugQuadSize)),
+			FTransform::Identity,
+			1
+		);
+
+		const FSoftObjectPath ParentMaterialPath = FSoftObjectPath(TEXT("/Script/Engine.Material'/Engine/Engine_MI_Shaders/M_Shader_SimpleTranslucent.M_Shader_SimpleTranslucent'"));
+		const TSoftObjectPtr<UMaterialInterface> ParentMaterialRef(ParentMaterialPath);
+		UMaterialInterface* ParentMaterial = ParentMaterialRef.LoadSynchronous();
+		if (ParentMaterial)
+		{
+			UMaterialInstanceDynamic* MaterialInstance = UMaterialInstanceDynamic::Create(ParentMaterial, nullptr);
+			MaterialInstance->SetFlags(RF_Transient);
+			FTaskTagScope Scope(ETaskTag::EParallelGameThread); // Needed for SetTextureParameterValue
+			MaterialInstance->SetTextureParameterValue(TEXT("Diffuse"), &Texture2D);
+			MaterialInstance->SetTextureParameterValue(TEXT("OpacityMask"), &Texture2D);
+
+			const static FName MaterialPathAttributeName("MaterialPath");
+			auto& MaterialArray = MeshGeomCollection->AddAttribute<FString>(MaterialPathAttributeName, FGeometryCollection::MaterialGroup);
+			MaterialArray[0] = MaterialInstance->GetPathName();
+		}
+
+		return MeshGeomCollection.ToSharedRef();
+	}
+
+	static void SampleImage(FTextureSource& InTextureSource, FVector2i& InOutSampleSize, TArray<bool>& OutBlackAndWhiteImageData)
+	{
+		const FTextureSource::FMipLock LockedMip0(FTextureSource::ELockState::ReadOnly, &InTextureSource, 0);
+		check(LockedMip0.IsValid());
+
+		const int32 ImageWith =  LockedMip0.Image.GetWidth();
+		const int32 ImageHeight = LockedMip0.Image.GetHeight();
+		InOutSampleSize.X = FMath::Min(InOutSampleSize.X, ImageWith);
+		InOutSampleSize.Y  = FMath::Min(InOutSampleSize.Y, ImageHeight);
+
+		OutBlackAndWhiteImageData.SetNum(InOutSampleSize.X * InOutSampleSize.Y);
+		for (int32 x = 0; x < InOutSampleSize.X; x++)
+		{
+			for (int32 y = 0; y < InOutSampleSize.Y; y++)
+			{
+				const FVector2i PixelCoords = FVector2i(x, y);
+				const FVector2f UV = PixelCoordToUV(PixelCoords, InOutSampleSize);
+
+				const int32 SamplePosX = FMath::RoundToInt(UV.X * (ImageWith - 1));
+				const int32 SamplePosY = FMath::RoundToInt(UV.Y * (ImageHeight - 1));
+				check(SamplePosX >= 0 && SamplePosX < ImageWith);
+				check(SamplePosY >= 0 && SamplePosY < ImageHeight);
+				
+				const FLinearColor PixelColor = LockedMip0.Image.GetOnePixelLinear(SamplePosX, SamplePosY);
+				const float ColorIntensity = PixelColor.R + PixelColor.G + PixelColor.B + PixelColor.A;
+
+				const int32 PixelIndex = GetPixelIndex(PixelCoords, InOutSampleSize.X);
+				OutBlackAndWhiteImageData[PixelIndex] = ColorIntensity > 0;
+			}
+		}
+	}
+
+	static void ExtractContourLines(const TArray<bool>& ImageData, const FVector2i& ImageSize, TArray<TPair<FVector2f, FVector2f>>& OutLineSegments)
+	{
+		// Extracts the contour lines of the image data using marching squares (https://en.wikipedia.org/wiki/Marching_squares)
+		
+		const FVector2f Positions[] = {
+			{ 0.0, 0.5 }, { 0.5, 0.0 }, 
+			{ 0.5, 0.0 }, { 1.0, 0.5 }, 
+			{ 0.5, 1.0 }, { 1.0, 0.5 }, 
+			{ 0.0, 0.5 }, { 0.5, 1.0 }, 
+			{ 0.5, 1.0 }, { 0.5, 0.0 }, 
+			{ 0.0, 0.5 }, { 1.0, 0.5 }, 
+			{ 0.5, 0.0 }, { 1.0, 0.5 }, { 0.0, 0.5 }, { 0.5, 1.0 }, 
+			{ 0.0, 0.5 }, { 0.5, 0.0 }, { 0.5, 1.0 }, { 1.0, 0.5 }
+		};
+	
+		constexpr int32 Starts[] = { 0, 0, 2, 10, 4, 16, 8, 6, 6, 8, 12, 4, 10, 2, 0, 0 };
+		constexpr int32 Sizes[] =  { 0, 2, 2, 2,  2,  4, 2, 2, 2, 2,  4, 2,  2, 2, 2, 0 };
+
+		const int32 Width = ImageSize.X;
+		const int32 Height = ImageSize.Y;
+		check(Width > 1 && Height > 1);
+		const FVector2f Scale = FVector2f(1.0 / (Width - 1), 1.0 / (Height - 1));
+
+		OutLineSegments.Reset(Width * Height);
+		for (int32 x = 0; x < Width - 1; x++)
+		{
+			for (int32 y = 0; y < Height - 1; y++)
+			{
+				const FVector2i PixelsToSample[] = {
+					{ x, y },
+					{ x + 1, y },
+					{ x + 1, y + 1 },
+					{ x, y + 1 }
+				};
+	
+				int32 Index = 0;
+				for (int32 i = 0; i < 4; i++)
+				{
+					if (!SampleImageBuffer(ImageData, Width, PixelsToSample[i]))
+					{
+						Index |= (1 << i); 
+					}
+				}
+		        
+				const int32 Start = Starts[Index];
+				const int32 End = Starts[Index] + Sizes[Index];
+				const FVector2f PixelPosition = FVector2f(x, y);
+		        
+				for (int i = Start; i < End; i += 2)
+				{
+					const FVector2f LineStart = (Positions[i] + PixelPosition) * Scale;
+					const FVector2f LineEnd = (Positions[i + 1] + PixelPosition) * Scale;
+					OutLineSegments.Emplace(LineStart, LineEnd);
+				}
+			}
+		}
+	}
+};
+
+PV::Texture2DImport::EImportResult PV::Texture2DImport::ImportGrowthDataFromTexture2D(const TObjectPtr<UTexture2D>& InSourceTexture2D,
+	FManagedArrayCollection& OutCollection, const FImportOptions& InOptions)
+{
+	using namespace Internal;
+	
+	if (!InSourceTexture2D || !InSourceTexture2D->Source.IsValid())
+	{
+		return EImportResult::InvalidSourceTexture;
+	}
+
+	const auto ReturnWithDebugMesh = [&]()->EImportResult
+	{
+		const TSharedRef<FGeometryCollection> DebugMeshGeomCollection = GenerateDebugQuad(*InSourceTexture2D);
+
+		FManagedArrayCollection MeshArrayCollection;
+		DebugMeshGeomCollection->CopyTo(&OutCollection);
+	
+		return EImportResult::Success;
+	};
+
+	if (InOptions.DebugState == FImportOptions::EDebugState::Texture)
+	{
+		return ReturnWithDebugMesh();
+	}
+
+	FVector2i SampleSize = FVector2i(InOptions.SampleWidth, InOptions.SampleHeight);
+	if (SampleSize.X <= 1 || SampleSize.Y <= 1)
+	{
+		return EImportResult::InvalidSampleSize;
+	}
+	
+	TArray<bool> BlackAndWhiteImageData;
+
+	SampleImage(InSourceTexture2D->Source, SampleSize, BlackAndWhiteImageData);
+	if (SampleSize.X <= 1 || SampleSize.Y <= 1)
+	{
+		return EImportResult::InvalidTextureSize;
+	}
+	
+	if (InOptions.DebugState == FImportOptions::EDebugState::SampleImage)
+	{
+		AddDebugPixelValuesToCollection(OutCollection, BlackAndWhiteImageData, SampleSize);
+		return ReturnWithDebugMesh();
+	}
+
+	TArray<TPair<FVector2f, FVector2f>> LineSegments;
+	ExtractContourLines(BlackAndWhiteImageData, SampleSize, LineSegments);
+	if (InOptions.DebugState == FImportOptions::EDebugState::ExtractContourLines)
+	{
+		AddDebugPixelValuesToCollection(OutCollection, BlackAndWhiteImageData, SampleSize);
+		AddDebugLineSegmentsToCollection(OutCollection, LineSegments);
+		return ReturnWithDebugMesh();
+	}
+	
+	return EImportResult::Success;
+}
diff --git a/Engine/Plugins/Experimental/ProceduralVegetationEditor/Source/ProceduralVegetationEditor/Private/Importer/PVImporter_Texture2D.h b/Engine/Plugins/Experimental/ProceduralVegetationEditor/Source/ProceduralVegetationEditor/Private/Importer/PVImporter_Texture2D.h
new file mode 100644
--- /dev/null
+++ b/Engine/Plugins/Experimental/ProceduralVegetationEditor/Source/ProceduralVegetationEditor/Private/Importer/PVImporter_Texture2D.h
@@ -0,0 +1,37 @@
+// Copyright Epic Games, Inc. All Rights Reserved.
+#pragma once
+
+#include "CoreMinimal.h"
+
+#include "GeometryCollection/ManagedArrayCollection.h"
+
+class UTexture2D;
+
+namespace PV::Texture2DImport
+{
+	enum class EImportResult
+	{
+		Success,
+		InvalidSourceTexture,
+		InvalidSampleSize,
+		InvalidTextureSize
+	};
+
+	struct FImportOptions
+	{
+		int32 SampleWidth = 128;
+		int32 SampleHeight = 128;
+		
+		enum class EDebugState : uint8
+		{
+			None,
+			Texture,
+			SampleImage,
+			ExtractContourLines,
+			MAX
+		};
+		EDebugState DebugState = EDebugState::None;
+	};
+
+	EImportResult ImportGrowthDataFromTexture2D(const TObjectPtr<UTexture2D>& InSourceTexture2D, FManagedArrayCollection& OutCollection, const FImportOptions& InOptions);
+};
diff --git a/Engine/Plugins/Experimental/ProceduralVegetationEditor/Source/ProceduralVegetationEditor/Private/PVEditorInterfaceImpl.cpp b/Engine/Plugins/Experimental/ProceduralVegetationEditor/Source/ProceduralVegetationEditor/Private/PVEditorInterfaceImpl.cpp
--- a/Engine/Plugins/Experimental/ProceduralVegetationEditor/Source/ProceduralVegetationEditor/Private/PVEditorInterfaceImpl.cpp
+++ b/Engine/Plugins/Experimental/ProceduralVegetationEditor/Source/ProceduralVegetationEditor/Private/PVEditorInterfaceImpl.cpp
@@ -3,9 +3,12 @@
 #include "PVEditorInterfaceImpl.h"
 
 #include "Importer/PVImporter_StaticMesh.h"
+#include "Importer/PVImporter_Texture2D.h"
 #include "PVImportStaticMeshParams.h"
+#include "PVImportTexture2DParams.h"
 
 #include "Engine/StaticMesh.h"
+#include "Engine/Texture2D.h"
 
 bool PV::FPVEditorInterfaceImpl::ImportGrowthDataFromStaticMesh(const FPVImportStaticMeshParams& InParams, FManagedArrayCollection& OutCollection)
 {
@@ -18,3 +21,16 @@ bool PV::FPVEditorInterfaceImpl::ImportGrowthDataFromStaticMesh(const FPVImportS
 	const PV::StaticMeshImport::EImportResult ImportResult = PV::StaticMeshImport::ImportGrowthDataFromStaticMesh(InParams.StaticMeshAsset, OutCollection, ImportMeshOptions);
 	return ImportResult == PV::StaticMeshImport::EImportResult::Success;
 }
+
+bool PV::FPVEditorInterfaceImpl::ImportGrowthDataFromTexture2D(const FPVImportTexture2DParams& InParams, FManagedArrayCollection& OutCollection)
+{
+	PV::Texture2DImport::FImportOptions ImportMeshOptions;
+	ImportMeshOptions.SampleWidth = InParams.SampleWidth;
+	ImportMeshOptions.SampleHeight = InParams.SampleHeight;
+	
+	static_assert(static_cast<int32>(EPVImportTexture2DDebugState::MAX) == static_cast<int32>(PV::Texture2DImport::FImportOptions::EDebugState::MAX));
+	ImportMeshOptions.DebugState = static_cast<PV::Texture2DImport::FImportOptions::EDebugState>(static_cast<int32>(InParams.DebugState));
+
+	const PV::Texture2DImport::EImportResult ImportResult = PV::Texture2DImport::ImportGrowthDataFromTexture2D(InParams.Texture2DAsset, OutCollection, ImportMeshOptions);
+	return ImportResult == PV::Texture2DImport::EImportResult::Success;
+}
\ No newline at end of file
diff --git a/Engine/Plugins/Experimental/ProceduralVegetationEditor/Source/ProceduralVegetationEditor/Private/PVEditorInterfaceImpl.h b/Engine/Plugins/Experimental/ProceduralVegetationEditor/Source/ProceduralVegetationEditor/Private/PVEditorInterfaceImpl.h
--- a/Engine/Plugins/Experimental/ProceduralVegetationEditor/Source/ProceduralVegetationEditor/Private/PVEditorInterfaceImpl.h
+++ b/Engine/Plugins/Experimental/ProceduralVegetationEditor/Source/ProceduralVegetationEditor/Private/PVEditorInterfaceImpl.h
@@ -15,6 +15,7 @@ namespace PV
 	{
 	public:
 		bool ImportGrowthDataFromStaticMesh(const FPVImportStaticMeshParams& InParams, FManagedArrayCollection& OutCollection) override;
+		bool ImportGrowthDataFromTexture2D(const FPVImportTexture2DParams& InParams, FManagedArrayCollection& OutCollection) override;
 	};
 };
 #endif
diff --git a/Engine/Plugins/Experimental/ProceduralVegetationEditor/Source/ProceduralVegetationEditor/ProceduralVegetationEditor.Build.cs b/Engine/Plugins/Experimental/ProceduralVegetationEditor/Source/ProceduralVegetationEditor/ProceduralVegetationEditor.Build.cs
--- a/Engine/Plugins/Experimental/ProceduralVegetationEditor/Source/ProceduralVegetationEditor/ProceduralVegetationEditor.Build.cs
+++ b/Engine/Plugins/Experimental/ProceduralVegetationEditor/Source/ProceduralVegetationEditor/ProceduralVegetationEditor.Build.cs
@@ -27,6 +27,7 @@ public class ProceduralVegetationEditor : ModuleRules
 				"GeometryFramework",
 				"GeometryScriptingCore",
 				"InputCore",
+				"ImageCore",
 				"Json",
 				"JsonUtilities",
 				"LevelEditor",
