Fix minor mistakes in comment

#Synced-CL 48983295

diff --git a/Engine/Plugins/Experimental/ProceduralVegetationEditor/Source/ProceduralVegetationEditor/Private/Importer/PVImporter_StaticMesh.cpp b/Engine/Plugins/Experimental/ProceduralVegetationEditor/Source/ProceduralVegetationEditor/Private/Importer/PVImporter_StaticMesh.cpp
--- a/Engine/Plugins/Experimental/ProceduralVegetationEditor/Source/ProceduralVegetationEditor/Private/Importer/PVImporter_StaticMesh.cpp
+++ b/Engine/Plugins/Experimental/ProceduralVegetationEditor/Source/ProceduralVegetationEditor/Private/Importer/PVImporter_StaticMesh.cpp
@@ -430,9 +430,12 @@ static void ComputeVerticesCenterAndRadius_Internal_V2(UE::Geometry::FDynamicMes
 	// indicates that at radius/center point was not found.
 
 	// Issues with this algorithm:
-	// This algorithm only works if the branch mesh is fully enclosed. If there are holes then the ray may exist through those holes and hit parts
-	// of the branch we do not desire for it to hit. It also does not work well if the mesh is built up using only faces (e.g flat geometry instead
-	// of a tube for the branches)
+	// 1. This algorithm only works if the branch mesh is fully enclosed. If there are holes then the ray may exit through those holes and hit parts
+	// of the branch we do not desire for it to hit.
+	// 2. This algorithm does not handle intersection geometry well as it cannot differentiate between hitting an intersection branch and the opposite
+	// side of the current branch. This is partially handled by running this algorithm per detached mesh (which for most meshes will mean per branch)
+	// 3. It also not work well if the mesh is built up using only faces (e.g flat geometry instead of a tube for the branches)
+	// 4. The normal of the vertex may not point directly towards the opposite side (see "Remove extreme radius values" section below for more information)
 
 	const TArray<int32> ComponentVertexIndices = GetUniqueVertexIndices(DynamicMesh, TriangleIds);
 
