[PVE] Minor improvement for Object Avoidance

#Synced-CL 49158918

diff --git a/Engine/Plugins/Experimental/ProceduralVegetationEditor/Source/ProceduralVegetation/Private/Nodes/PVObjectAvoidance.cpp b/Engine/Plugins/Experimental/ProceduralVegetationEditor/Source/ProceduralVegetation/Private/Nodes/PVObjectAvoidance.cpp
--- a/Engine/Plugins/Experimental/ProceduralVegetationEditor/Source/ProceduralVegetation/Private/Nodes/PVObjectAvoidance.cpp
+++ b/Engine/Plugins/Experimental/ProceduralVegetationEditor/Source/ProceduralVegetation/Private/Nodes/PVObjectAvoidance.cpp
@@ -132,6 +132,8 @@ bool FPVObjectAvoidanceElement::ExecuteInternal(FPCGContext* InContext) const
 				for (int32 BranchIndex : BranchIndices)
 				{
 					PVE_OUTER_LOOP_DEBUG_CHECK(break);
+					
+					const float MaxPointScale = *Algo::MaxElement(PointFacade.GetPointScales());
 
 					const TArray<int32> BranchPoints = BranchFacade.GetPoints(BranchIndex);
 					for (int32 BranchPointIndex = 1; BranchPointIndex < BranchPoints.Num(); ++BranchPointIndex)
@@ -154,7 +156,7 @@ bool FPVObjectAvoidanceElement::ExecuteInternal(FPCGContext* InContext) const
 							PreviousDirection = (NextPointPosition - PointPosition).GetSafeNormal(UE_SMALL_NUMBER, FVector::UpVector);
 						}
 
-						const float PointScale = PointFacade.GetPointScale(PointIndex);
+						const float PointScale = MaxPointScale;
 
 						FRay3d Ray(PointPosition, SegmentVector);
 						FRay3d InverseRay(Transform.InverseTransformPosition(Ray.Origin), Transform.InverseTransformVector(Ray.Direction));
