[PVE] Foliage condition system

- Added attributes in foliage palette
- Implemented the condition system in the foliage distributor
- Implemented reading foliage attributes and conditions from JSON

Attributes are added with each foliage palette entry. These are 0-1 values that user can set with foliage entry. For each, if you want a specific foliage to only appear on tip, then you just set the tip to true and enable the tip condition in the foliage distributor.

Main logic in files:
FPVFoliage

UI changes in:
PVFoliagePaletteSettings
PVFoliageDistributorSettings

Files added to avoid changes in existing public headers:
PVFoliageAttributeNames
PVFoliageJSONHelper

diff --git a/Engine/Plugins/Experimental/ProceduralVegetationEditor/Source/ProceduralVegetation/Private/Facades/PVBudVectorsFacade.h b/Engine/Plugins/Experimental/ProceduralVegetationEditor/Source/ProceduralVegetation/Private/Facades/PVBudVectorsFacade.h
--- a/Engine/Plugins/Experimental/ProceduralVegetationEditor/Source/ProceduralVegetation/Private/Facades/PVBudVectorsFacade.h
+++ b/Engine/Plugins/Experimental/ProceduralVegetationEditor/Source/ProceduralVegetation/Private/Facades/PVBudVectorsFacade.h
@@ -7,9 +7,10 @@
 
 namespace PV::Facades
 {
+	constexpr int BudDirectionsApical = 0;
 	constexpr int BudDirectionsLightOptimalIndex = 2;
 	constexpr int BudDirectionsLightSubOptimal = 3;
-
+	constexpr int BudLightDetectedAvailableIndex = 0;
 	/**
 	 * FBudVectorsFacade is used to access and manipulate the BudVectors data from Point Group with in the ProceduralVegetation's FManagedArrayCollection
 	 * BudVectors represent different growth related vectors in plant data and are associated with each point. 
diff --git a/Engine/Plugins/Experimental/ProceduralVegetationEditor/Source/ProceduralVegetation/Private/Facades/PVFoliageAttributesNames.h b/Engine/Plugins/Experimental/ProceduralVegetationEditor/Source/ProceduralVegetation/Private/Facades/PVFoliageAttributesNames.h
new file mode 100644
--- /dev/null
+++ b/Engine/Plugins/Experimental/ProceduralVegetationEditor/Source/ProceduralVegetation/Private/Facades/PVFoliageAttributesNames.h
@@ -0,0 +1,42 @@
+// Copyright Epic Games, Inc. All Rights Reserved.
+#pragma once
+
+#include "CoreMinimal.h"
+
+namespace PV::GroupNames
+{
+	inline static const FName FoliageConditionGroup = FName(TEXT("FoliageConditions"));
+}
+
+namespace PV::FoliageAttributeNames
+{
+
+	//Foliage Attributes
+	inline static const FName FoliageNameID = FName(TEXT("FoliageNameID"));
+	inline static const FName FoliageBranchID = FName(TEXT("FoliageBranchID"));
+	inline static const FName FoliagePivotPoint = FName(TEXT("FoliagePivotPoint"));
+	inline static const FName FoliageUPVector = FName(TEXT("FoliageUPVector"));
+	inline static const FName FoliageNormalVector = FName(TEXT("FoliageNormalVector"));
+	inline static const FName FoliageScale = FName(TEXT("FoliageScale"));
+	inline static const FName FoliageLengthFromRoot = FName(TEXT("LengthFromRoot"));
+	inline static const FName FoliageConditionLight = FName(TEXT("ConditionLight"));
+	inline static const FName FoliageConditionTip = FName(TEXT("ConditionTip"));
+	inline static const FName FoliageConditionScale = FName(TEXT("ConditionScale"));
+	inline static const FName FoliageConditionUpAlignment = FName(TEXT("ConditionUpAlignment"));
+	inline static const FName FoliageConditionHealth = FName(TEXT("ConditionHealth"));
+	inline static const FName FoliageParentBoneID = FName(TEXT("FoliageParentBoneID"));
+
+	//FoliageInfo Attributes
+	inline static const FName FoliageName = FName(TEXT("Name"));
+	inline static const FName FoliageAttributeLight = FName(TEXT("Light"));
+	inline static const FName FoliageAttributeScale = FName(TEXT("Scale"));
+	inline static const FName FoliageAttributeTip = FName(TEXT("Tip"));
+	inline static const FName FoliageAttributeUpAlignment = FName(TEXT("UpAlignment"));
+	inline static const FName FoliageAttributeHealth = FName(TEXT("Health"));
+
+	// Condition
+	inline static const FName ConditionName = FName(TEXT("Name"));
+	inline static const FName ConditionWeight = FName(TEXT("Weight"));
+	inline static const FName ConditionOffset = FName(TEXT("Offset"));
+}
+
diff --git a/Engine/Plugins/Experimental/ProceduralVegetationEditor/Source/ProceduralVegetation/Private/Facades/PVFoliageConditionFacade.cpp b/Engine/Plugins/Experimental/ProceduralVegetationEditor/Source/ProceduralVegetation/Private/Facades/PVFoliageConditionFacade.cpp
new file mode 100644
--- /dev/null
+++ b/Engine/Plugins/Experimental/ProceduralVegetationEditor/Source/ProceduralVegetation/Private/Facades/PVFoliageConditionFacade.cpp
@@ -0,0 +1,111 @@
+﻿// Copyright Epic Games, Inc. All Rights Reserved.
+
+#include "Facades/PVFoliageConditionFacade.h"
+#include "PVFoliageAttributesNames.h"
+
+namespace PV::Facades
+{
+	FFoliageConditionFacade::FFoliageConditionFacade(FManagedArrayCollection& InCollection)
+		: NameAttribute(InCollection, FoliageAttributeNames::ConditionName, GroupNames::FoliageConditionGroup)
+		, WeightAttribute(InCollection, FoliageAttributeNames::ConditionWeight, GroupNames::FoliageConditionGroup)
+		, OffsetAttribute(InCollection, FoliageAttributeNames::ConditionOffset, GroupNames::FoliageConditionGroup)
+	{
+		DefineSchema(InCollection);
+	}
+
+	FFoliageConditionFacade::FFoliageConditionFacade(const FManagedArrayCollection& InCollection)
+	: NameAttribute(InCollection, FoliageAttributeNames::ConditionName, GroupNames::FoliageConditionGroup)
+	, WeightAttribute(InCollection, FoliageAttributeNames::ConditionWeight, GroupNames::FoliageConditionGroup)
+	, OffsetAttribute(InCollection, FoliageAttributeNames::ConditionOffset, GroupNames::FoliageConditionGroup)
+	{
+	}
+
+	void FFoliageConditionFacade::DefineSchema(FManagedArrayCollection& InCollection)
+	{
+		if (!InCollection.HasGroup(GroupNames::FoliageConditionGroup))
+		{
+			InCollection.AddGroup(GroupNames::FoliageConditionGroup);
+		}
+
+		NameAttribute.Add();
+		WeightAttribute.Add();
+		OffsetAttribute.Add();
+	}
+
+	bool FFoliageConditionFacade::IsValid() const
+	{
+		return NameAttribute.IsValid() && WeightAttribute.IsValid() && OffsetAttribute.IsValid();
+	}
+
+	FFoliageConditonInfo FFoliageConditionFacade::GetEntry(const int32 Index) const
+	{
+		FFoliageConditonInfo ReturnData = {};
+		if (IsValid() && Index > -1)
+		{
+			if (NameAttribute.IsValidIndex(Index))
+			{
+				ReturnData.Name = NameAttribute.Get()[Index];
+			}
+			if (WeightAttribute.IsValidIndex(Index))
+			{
+				ReturnData.Weight = WeightAttribute.Get()[Index];
+			}
+			if (OffsetAttribute.IsValidIndex(Index))
+			{
+				ReturnData.Offset = OffsetAttribute.Get()[Index];
+			}
+		}
+		return ReturnData;
+	}
+
+	int32 FFoliageConditionFacade::Add(const FFoliageConditonInfo& InputData)
+	{
+		if (IsValid())
+		{
+			const int32 NewIndex = NameAttribute.AddElements(1);
+			NameAttribute.Modify()[NewIndex] = InputData.Name;
+			WeightAttribute.Modify()[NewIndex] = InputData.Weight;
+			OffsetAttribute.Modify()[NewIndex] = InputData.Offset;
+
+			return NewIndex;
+		}
+		return INDEX_NONE;
+	}
+
+	void FFoliageConditionFacade::Set(const int32 Index, const FFoliageConditonInfo& InputData)
+	{
+		if (IsValid() && Index >= 0)
+		{
+			NameAttribute.ModifyAt(Index, InputData.Name);
+			WeightAttribute.ModifyAt(Index, InputData.Weight);
+			OffsetAttribute.ModifyAt(Index, InputData.Offset);
+		}
+	}
+
+	TArray<FFoliageConditonInfo> FFoliageConditionFacade::GetData() const
+	{
+		TArray<FFoliageConditonInfo> Infos;
+		
+		if (NameAttribute.IsValid())
+		{
+			for (int32 i = 0; i < NameAttribute.Num(); i++)
+			{
+				Infos.Add(GetEntry(i));
+			}
+		}
+
+		return Infos;
+	}
+
+	void FFoliageConditionFacade::SetData(const TArray<FFoliageConditonInfo>& Infos)
+	{
+		int32 NumElements = NameAttribute.AddElements(Infos.Num());
+		
+		for (int32 Index = 0; Index < Infos.Num(); ++Index)
+		{
+			FFoliageConditonInfo Info = Infos[Index];
+
+			Set(Index, Info);
+		}
+	}
+}
diff --git a/Engine/Plugins/Experimental/ProceduralVegetationEditor/Source/ProceduralVegetation/Private/Facades/PVFoliageConditionFacade.h b/Engine/Plugins/Experimental/ProceduralVegetationEditor/Source/ProceduralVegetation/Private/Facades/PVFoliageConditionFacade.h
new file mode 100644
--- /dev/null
+++ b/Engine/Plugins/Experimental/ProceduralVegetationEditor/Source/ProceduralVegetation/Private/Facades/PVFoliageConditionFacade.h
@@ -0,0 +1,63 @@
+﻿// Copyright Epic Games, Inc. All Rights Reserved.
+
+#pragma once
+
+#include "CoreMinimal.h"
+#include "GeometryCollection/ManagedArrayAccessor.h"
+#include "Implementations/PVFoliage.h"
+
+struct FPVFoliageInfo;
+
+namespace PV::Facades
+{
+	struct PROCEDURALVEGETATION_API FFoliageConditonInfo
+	{
+		FString Name;
+		float Weight;
+		float Offset;
+
+		EFoliageDistributionCondition GetType() const
+		{
+			const UEnum* EnumPtr = StaticEnum<EFoliageDistributionCondition>();
+
+			int64 EnumValue = EnumPtr->GetValueByNameString(Name);
+
+			if (EnumValue == INDEX_NONE)
+				return EFoliageDistributionCondition::None;
+
+			return static_cast<EFoliageDistributionCondition>(EnumValue);
+		}
+
+		FString ToString() const
+		{
+			return FString::Printf(TEXT("Name = %s, Weight = %f, Reduction = %f"), *Name, Weight, Offset);
+		}
+	};
+
+	class PROCEDURALVEGETATION_API FFoliageConditionFacade
+	{
+	public:
+		FFoliageConditionFacade(FManagedArrayCollection& InCollection);
+		FFoliageConditionFacade(const FManagedArrayCollection& InCollection);
+
+		bool IsValid() const;
+
+		int32 NumEntries() const { return NameAttribute.Num(); };
+		
+		int32 Add(const FFoliageConditonInfo& InputData);
+
+		void Set(const int32 Index, const FFoliageConditonInfo& InputData);
+		void SetData(const TArray<FFoliageConditonInfo>& Infos);
+		
+		FFoliageConditonInfo GetEntry(const int32 Index) const;
+
+		TArray<FFoliageConditonInfo> GetData() const;
+
+	protected:
+		void DefineSchema(FManagedArrayCollection& InCollection);
+		
+		TManagedArrayAccessor<FString> NameAttribute;
+		TManagedArrayAccessor<float> WeightAttribute;
+		TManagedArrayAccessor<float> OffsetAttribute;
+	};
+}
diff --git a/Engine/Plugins/Experimental/ProceduralVegetationEditor/Source/ProceduralVegetation/Private/Facades/PVFoliageFacade.cpp b/Engine/Plugins/Experimental/ProceduralVegetationEditor/Source/ProceduralVegetation/Private/Facades/PVFoliageFacade.cpp
--- a/Engine/Plugins/Experimental/ProceduralVegetationEditor/Source/ProceduralVegetation/Private/Facades/PVFoliageFacade.cpp
+++ b/Engine/Plugins/Experimental/ProceduralVegetationEditor/Source/ProceduralVegetation/Private/Facades/PVFoliageFacade.cpp
@@ -1,53 +1,75 @@
 ﻿// Copyright Epic Games, Inc. All Rights Reserved.
-
+ 
 #include "Facades/PVFoliageFacade.h"
-
+ 
 #include "ProceduralVegetationModule.h"
+#include "PVFoliageAttributesNames.h"
 #include "VisualizeTexture.h"
 #include "Chaos/Deformable/MuscleActivationConstraints.h"
 #include "Facades/PVAttributesNames.h"
-
+#include "Implementations/PVFoliage.h"
+ 
 namespace PV::Facades
 {
 	FFoliageFacade::FFoliageFacade(FManagedArrayCollection& InCollection, int32 InitialSize)
 		: ConstCollection(InCollection)
 		, Collection(&InCollection)
-		, NamesAttribute(InCollection, AttributeNames::FoliageName, GroupNames::FoliageNamesGroup)
-		, NameIdsAttribute(InCollection, AttributeNames::FoliageNameID, GroupNames::FoliageGroup)
-		, BranchIdsAttribute(InCollection, AttributeNames::FoliageBranchID, GroupNames::FoliageGroup)
-		, PivotPointsAttribute(InCollection, AttributeNames::FoliagePivotPoint, GroupNames::FoliageGroup)
-		, UpVectorsAttribute(InCollection, AttributeNames::FoliageUPVector, GroupNames::FoliageGroup)
-		, NormalVectorsAttribute(InCollection, AttributeNames::FoliageNormalVector, GroupNames::FoliageGroup)
-		, ScalesAttribute(InCollection, AttributeNames::FoliageScale, GroupNames::FoliageGroup)
-		, LengthFromRootAttribute(InCollection, AttributeNames::FoliageLengthFromRoot, GroupNames::FoliageGroup)
-		, ParentBoneIdsAttribute(InCollection, AttributeNames::FoliageParentBoneID, GroupNames::FoliageGroup)
+		, FoliageInfoNamesAttribute(InCollection, FoliageAttributeNames::FoliageName, GroupNames::FoliageNamesGroup)
+		, FoliageInfoLightAttribute(InCollection, FoliageAttributeNames::FoliageAttributeLight, GroupNames::FoliageNamesGroup)
+		, FoliageInfoScaleAttribute(InCollection, FoliageAttributeNames::FoliageAttributeScale, GroupNames::FoliageNamesGroup)
+		, FoliageInfoTipAttribute(InCollection, FoliageAttributeNames::FoliageAttributeTip, GroupNames::FoliageNamesGroup)
+		, FoliageInfoUpAlignmentAttribute(InCollection, FoliageAttributeNames::FoliageAttributeUpAlignment, GroupNames::FoliageNamesGroup)
+		, FoliageInfoHealthAttribute(InCollection, FoliageAttributeNames::FoliageAttributeHealth, GroupNames::FoliageNamesGroup)
+		, NameIdsAttribute(InCollection, FoliageAttributeNames::FoliageNameID, GroupNames::FoliageGroup)
+		, BranchIdsAttribute(InCollection, FoliageAttributeNames::FoliageBranchID, GroupNames::FoliageGroup)
+		, PivotPointsAttribute(InCollection, FoliageAttributeNames::FoliagePivotPoint, GroupNames::FoliageGroup)
+		, UpVectorsAttribute(InCollection, FoliageAttributeNames::FoliageUPVector, GroupNames::FoliageGroup)
+		, NormalVectorsAttribute(InCollection, FoliageAttributeNames::FoliageNormalVector, GroupNames::FoliageGroup)
+		, ScalesAttribute(InCollection, FoliageAttributeNames::FoliageScale, GroupNames::FoliageGroup)
+		, LengthFromRootAttribute(InCollection, FoliageAttributeNames::FoliageLengthFromRoot, GroupNames::FoliageGroup)
+		, ConditionLightAttribute(InCollection, FoliageAttributeNames::FoliageConditionLight, GroupNames::FoliageGroup)
+		, ConditionScaleAttribute(InCollection, FoliageAttributeNames::FoliageConditionScale, GroupNames::FoliageGroup)
+		, ConditionTipAttribute(InCollection, FoliageAttributeNames::FoliageConditionTip, GroupNames::FoliageGroup)
+		, ConditionUpAlignmentAttribute(InCollection, FoliageAttributeNames::FoliageConditionUpAlignment, GroupNames::FoliageGroup)
+		, ConditionHealthAttribute(InCollection, FoliageAttributeNames::FoliageConditionHealth, GroupNames::FoliageGroup)
+		, ParentBoneIdsAttribute(InCollection, FoliageAttributeNames::FoliageParentBoneID, GroupNames::FoliageGroup)
 		, FoliageIdsAttribute(InCollection, AttributeNames::BranchFoliageIDs, GroupNames::BranchGroup)
 		, FoliagePathAttribute(InCollection, AttributeNames::FoliagePath, GroupNames::DetailsGroup)
 	{
 		DefineSchema(InitialSize);
 	}
-
+ 
 	FFoliageFacade::FFoliageFacade(const FManagedArrayCollection& InCollection)
 		: ConstCollection(InCollection)
 		, Collection(nullptr)
-		, NamesAttribute(InCollection, AttributeNames::FoliageName, GroupNames::FoliageNamesGroup)
-		, NameIdsAttribute(InCollection, AttributeNames::FoliageNameID, GroupNames::FoliageGroup)
-		, BranchIdsAttribute(InCollection, AttributeNames::FoliageBranchID, GroupNames::FoliageGroup)
-		, PivotPointsAttribute(InCollection, AttributeNames::FoliagePivotPoint, GroupNames::FoliageGroup)
-		, UpVectorsAttribute(InCollection, AttributeNames::FoliageUPVector, GroupNames::FoliageGroup)
-		, NormalVectorsAttribute(InCollection, AttributeNames::FoliageNormalVector, GroupNames::FoliageGroup)
-		, ScalesAttribute(InCollection, AttributeNames::FoliageScale, GroupNames::FoliageGroup)
-		, LengthFromRootAttribute(InCollection, AttributeNames::FoliageLengthFromRoot, GroupNames::FoliageGroup)
-		, ParentBoneIdsAttribute(InCollection, AttributeNames::FoliageParentBoneID, GroupNames::FoliageGroup)
+		, FoliageInfoNamesAttribute(InCollection, FoliageAttributeNames::FoliageName, GroupNames::FoliageNamesGroup)
+		, FoliageInfoLightAttribute(InCollection, FoliageAttributeNames::FoliageAttributeLight, GroupNames::FoliageNamesGroup)
+		, FoliageInfoScaleAttribute(InCollection, FoliageAttributeNames::FoliageAttributeScale, GroupNames::FoliageNamesGroup)
+		, FoliageInfoTipAttribute(InCollection, FoliageAttributeNames::FoliageAttributeTip, GroupNames::FoliageNamesGroup)
+		, FoliageInfoUpAlignmentAttribute(InCollection, FoliageAttributeNames::FoliageAttributeUpAlignment, GroupNames::FoliageNamesGroup)
+		, FoliageInfoHealthAttribute(InCollection, FoliageAttributeNames::FoliageAttributeHealth, GroupNames::FoliageNamesGroup)
+		, NameIdsAttribute(InCollection, FoliageAttributeNames::FoliageNameID, GroupNames::FoliageGroup)
+		, BranchIdsAttribute(InCollection, FoliageAttributeNames::FoliageBranchID, GroupNames::FoliageGroup)
+		, PivotPointsAttribute(InCollection, FoliageAttributeNames::FoliagePivotPoint, GroupNames::FoliageGroup)
+		, UpVectorsAttribute(InCollection, FoliageAttributeNames::FoliageUPVector, GroupNames::FoliageGroup)
+		, NormalVectorsAttribute(InCollection, FoliageAttributeNames::FoliageNormalVector, GroupNames::FoliageGroup)
+		, ScalesAttribute(InCollection, FoliageAttributeNames::FoliageScale, GroupNames::FoliageGroup)
+		, LengthFromRootAttribute(InCollection, FoliageAttributeNames::FoliageLengthFromRoot, GroupNames::FoliageGroup)
+		, ConditionLightAttribute(InCollection, FoliageAttributeNames::FoliageConditionLight, GroupNames::FoliageGroup)
+		, ConditionScaleAttribute(InCollection, FoliageAttributeNames::FoliageConditionScale, GroupNames::FoliageGroup)
+		, ConditionTipAttribute(InCollection, FoliageAttributeNames::FoliageConditionTip, GroupNames::FoliageGroup)
+		, ConditionUpAlignmentAttribute(InCollection, FoliageAttributeNames::FoliageConditionUpAlignment, GroupNames::FoliageGroup)
+		, ConditionHealthAttribute(InCollection, FoliageAttributeNames::FoliageConditionHealth, GroupNames::FoliageGroup)
+		, ParentBoneIdsAttribute(InCollection, FoliageAttributeNames::FoliageParentBoneID, GroupNames::FoliageGroup)
 		, FoliageIdsAttribute(InCollection, AttributeNames::BranchFoliageIDs, GroupNames::BranchGroup)
 		, FoliagePathAttribute(InCollection, AttributeNames::FoliagePath, GroupNames::DetailsGroup)
 	{
 	}
-
+ 
 	void FFoliageFacade::DefineSchema(int32 InitialSize)
 	{
 		check(!IsConst());
-
+ 
 		if (!Collection->HasGroup(GroupNames::FoliageGroup))
 		{
 			Collection->AddGroup(GroupNames::FoliageGroup);
@@ -56,18 +78,23 @@ namespace PV::Facades
 				Collection->AddElements(InitialSize, GroupNames::FoliageGroup);
 			}
 		}
-
+ 
 		if (!Collection->HasGroup(GroupNames::FoliageNamesGroup))
 		{
 			Collection->AddGroup(GroupNames::FoliageNamesGroup);
 		}
-
+ 
 		if (!Collection->HasGroup(GroupNames::DetailsGroup))
 		{
 			Collection->AddGroup(GroupNames::DetailsGroup);
 		}
-
-		NamesAttribute.Add();
+ 
+		FoliageInfoNamesAttribute.Add();
+		FoliageInfoLightAttribute.Add();
+		FoliageInfoScaleAttribute.Add();
+		FoliageInfoTipAttribute.Add();
+		FoliageInfoUpAlignmentAttribute.Add();
+		FoliageInfoHealthAttribute.Add();
 		NameIdsAttribute.Add();
 		BranchIdsAttribute.Add();
 		PivotPointsAttribute.Add();
@@ -75,18 +102,24 @@ namespace PV::Facades
 		NormalVectorsAttribute.Add();
 		ScalesAttribute.Add();
 		LengthFromRootAttribute.Add();
+		ConditionLightAttribute.Add();
+		ConditionTipAttribute.Add();
+		ConditionUpAlignmentAttribute.Add();
+		ConditionHealthAttribute.Add();
+		ConditionScaleAttribute.Add();
 		ParentBoneIdsAttribute.Add();
 		FoliageIdsAttribute.Add();
 		FoliagePathAttribute.Add();
 	}
-
+ 
 	bool FFoliageFacade::IsValid() const
 	{
-		return NamesAttribute.IsValid() && NameIdsAttribute.IsValid() && BranchIdsAttribute.IsValid() && PivotPointsAttribute.IsValid() &&
+		return FoliageInfoNamesAttribute.IsValid() && FoliageInfoTipAttribute.IsValid() && FoliageInfoUpAlignmentAttribute.IsValid() && FoliageInfoHealthAttribute.IsValid() && FoliageInfoLightAttribute.IsValid() && FoliageInfoScaleAttribute.IsValid() && 
+			NameIdsAttribute.IsValid() && BranchIdsAttribute.IsValid() && PivotPointsAttribute.IsValid() &&
 			UpVectorsAttribute.IsValid() && NormalVectorsAttribute.IsValid() && ScalesAttribute.IsValid() && LengthFromRootAttribute.IsValid() &&
 			FoliageIdsAttribute.IsValid() && ParentBoneIdsAttribute.IsValid();
 	}
-
+ 
 	FFoliageEntryData FFoliageFacade::GetFoliageEntry(int32 Index) const
 	{
 		FFoliageEntryData ReturnData = {};
@@ -120,6 +153,32 @@ namespace PV::Facades
 			{
 				ReturnData.LengthFromRoot = LengthFromRootAttribute.Get()[Index];
 			}
+ 
+			if (ConditionLightAttribute.IsValidIndex(Index))
+			{
+				ReturnData.ConditionLight = ConditionLightAttribute.Get()[Index];
+			}
+ 
+			if (ConditionScaleAttribute.IsValidIndex(Index))
+			{
+				ReturnData.ConditionScale = ConditionScaleAttribute.Get()[Index];
+			}
+ 
+			if (ConditionTipAttribute.IsValidIndex(Index))
+			{
+				ReturnData.ConditionTip = ConditionTipAttribute.Get()[Index];
+			}
+ 
+			if (ConditionUpAlignmentAttribute.IsValidIndex(Index))
+			{
+				ReturnData.ConditionUpAlignment = ConditionUpAlignmentAttribute.Get()[Index];
+			}
+ 
+			if (ConditionHealthAttribute.IsValidIndex(Index))
+			{
+				ReturnData.ConditionHealth = ConditionHealthAttribute.Get()[Index];
+			}
+			
 			if (ParentBoneIdsAttribute.IsValidIndex(Index))
 			{
 				ReturnData.ParentBoneID = ParentBoneIdsAttribute.Get()[Index];
@@ -127,18 +186,18 @@ namespace PV::Facades
 		}
 		return ReturnData;
 	}
-
+ 
 	float FFoliageFacade::GetLengthFromRoot(int32 Index) const
 	{
 		if (LengthFromRootAttribute.IsValid() && LengthFromRootAttribute.IsValidIndex(Index))
 		{
 			return LengthFromRootAttribute[Index];
 		}
-
+ 
 		static constexpr float DefaultLengthFromRoot = 0.0;
 		return DefaultLengthFromRoot;
 	}
-
+ 
 	void FFoliageFacade::SetLengthFromRoot(int32 Index, float Input)
 	{
 		check(!IsConst());
@@ -147,17 +206,17 @@ namespace PV::Facades
 			LengthFromRootAttribute.ModifyAt(Index, Input);
 		}
 	}
-
+ 
 	int32 FFoliageFacade::GetParentBoneID(int32 Index) const
 	{
 		if (ParentBoneIdsAttribute.IsValid() && ParentBoneIdsAttribute.IsValidIndex(Index))
 		{
 			return ParentBoneIdsAttribute[Index];
 		}
-
+ 
 		return INDEX_NONE;
 	}
-
+ 
 	void FFoliageFacade::SetParentBoneID(int32 Index, int32 Input)
 	{
 		check(!IsConst());
@@ -166,18 +225,18 @@ namespace PV::Facades
 			ParentBoneIdsAttribute.ModifyAt(Index, Input);
 		}
 	}
-
+ 
 	const TArray<int32>& FFoliageFacade::GetFoliageEntryIdsForBranch(int32 Index) const
 	{
 		if (FoliageIdsAttribute.IsValid() && FoliageIdsAttribute.IsValidIndex(Index))
 		{
 			return FoliageIdsAttribute[Index];
 		}
-
+ 
 		static const TArray<int32> EmptyArray;
 		return EmptyArray;
 	}
-
+ 
 	int32 FFoliageFacade::AddFoliageEntry(const FFoliageEntryData& InputData)
 	{
 		check(!IsConst());
@@ -191,13 +250,19 @@ namespace PV::Facades
 			NormalVectorsAttribute.Modify()[NewIndex] = InputData.NormalVector;
 			ScalesAttribute.Modify()[NewIndex] = InputData.Scale;
 			LengthFromRootAttribute.Modify()[NewIndex] = InputData.LengthFromRoot;
+			ConditionLightAttribute.Modify()[NewIndex] = InputData.ConditionLight;
+			ConditionScaleAttribute.Modify()[NewIndex] = InputData.ConditionScale;
+			ConditionUpAlignmentAttribute.Modify()[NewIndex] = InputData.ConditionUpAlignment;
+			ConditionHealthAttribute.Modify()[NewIndex] = InputData.ConditionHealth;
+			ConditionTipAttribute.Modify()[NewIndex] = InputData.ConditionTip;
+			
 			ParentBoneIdsAttribute.Modify()[NewIndex] = InputData.ParentBoneID;
-
+ 
 			return NewIndex;
 		}
 		return INDEX_NONE;
 	}
-
+ 
 	void FFoliageFacade::SetFoliageEntry(int32 Index, const FFoliageEntryData& InputData)
 	{
 		check(!IsConst());
@@ -210,12 +275,15 @@ namespace PV::Facades
 			NormalVectorsAttribute.ModifyAt(Index, InputData.NormalVector);
 			ScalesAttribute.ModifyAt(Index, InputData.Scale);
 			LengthFromRootAttribute.ModifyAt(Index, InputData.LengthFromRoot);
-			LengthFromRootAttribute.ModifyAt(Index, InputData.LengthFromRoot);
-			LengthFromRootAttribute.ModifyAt(Index, InputData.LengthFromRoot);
+			ConditionTipAttribute.ModifyAt(Index, InputData.ConditionTip);
+			ConditionScaleAttribute.ModifyAt(Index, InputData.ConditionScale);
+			ConditionLightAttribute.ModifyAt(Index, InputData.ConditionLight);
+			ConditionUpAlignmentAttribute.ModifyAt(Index, InputData.ConditionUpAlignment);
+			ConditionHealthAttribute.ModifyAt(Index, InputData.ConditionHealth);
 			ParentBoneIdsAttribute.ModifyAt(Index, InputData.ParentBoneID);
 		}
 	}
-
+ 
 	void FFoliageFacade::SetFoliageIdsArray(int32 Index, TArray<int32> InputIds)
 	{
 		check(!IsConst());
@@ -224,7 +292,7 @@ namespace PV::Facades
 			FoliageIdsAttribute.Modify()[Index] = MoveTemp(InputIds);
 		}
 	}
-
+ 
 	int32 FFoliageFacade::GetFoliageBranchId(int32 Index) const
 	{
 		if (BranchIdsAttribute.IsValid() && BranchIdsAttribute.IsValidIndex(Index))
@@ -233,7 +301,7 @@ namespace PV::Facades
 		}
 		return INDEX_NONE;
 	}
-
+ 
 	void FFoliageFacade::SetFoliageBranchId(int32 Index, int32 InputId)
 	{
 		check(!IsConst());
@@ -242,62 +310,93 @@ namespace PV::Facades
 			BranchIdsAttribute.ModifyAt(Index, InputId);
 		}
 	}
-
-	const FString& FFoliageFacade::GetFoliageName(int32 Index) const
+ 
+	FPVFoliageInfo FFoliageFacade::GetFoliageInfo(const int32 Index) const
 	{
-		if (NamesAttribute.IsValid() && NamesAttribute.IsValidIndex(Index))
+		FPVFoliageInfo ReturnData = {};
+		
+		if (IsValid() && Index > -1)
 		{
-			return NamesAttribute[Index];
+			if (FoliageInfoNamesAttribute.IsValidIndex(Index))
+			{
+				ReturnData.Mesh = FoliageInfoNamesAttribute.Get()[Index];
+			}
+ 
+			if (FoliageInfoLightAttribute.IsValidIndex(Index))
+			{
+				ReturnData.Attributes.Light = FoliageInfoLightAttribute.Get()[Index];
+			}
+ 
+			if (FoliageInfoScaleAttribute.IsValidIndex(Index))
+			{
+				ReturnData.Attributes.Scale = FoliageInfoScaleAttribute.Get()[Index];
+			}
+ 
+			if (FoliageInfoTipAttribute.IsValidIndex(Index))
+			{
+				ReturnData.Attributes.Tip = (bool)FoliageInfoTipAttribute.Get()[Index];
+			}
+ 
+			if (FoliageInfoHealthAttribute.IsValidIndex(Index))
+			{
+				ReturnData.Attributes.Health = FoliageInfoHealthAttribute.Get()[Index];
+			}
+ 
+			if (FoliageInfoUpAlignmentAttribute.IsValidIndex(Index))
+			{
+				ReturnData.Attributes.UpAlignment = FoliageInfoUpAlignmentAttribute.Get()[Index];
+			}
 		}
-
-		const static FString DefaultFoliageName = FString();
-		return DefaultFoliageName;
+		return ReturnData;
 	}
-
-	const TArray<FString>& FFoliageFacade::GetFoliageNames() const
+ 
+	TArray<FPVFoliageInfo> FFoliageFacade::GetFoliageInfos() const
 	{
-		if (NamesAttribute.IsValid())
+		TArray<FPVFoliageInfo> FoliageInfos;
+		
+		if (FoliageInfoNamesAttribute.IsValid())
 		{
-			return NamesAttribute.Get().GetConstArray();
+			for (int32 i = 0; i < FoliageInfoNamesAttribute.Num(); i++)
+			{
+				FoliageInfos.Add(GetFoliageInfo(i));
+			}
 		}
-
-		const static TArray<FString> DefaultFoliageNames;
-		return DefaultFoliageNames;
+		return FoliageInfos;
 	}
-
+ 
 	const FVector3f& FFoliageFacade::GetPivotPoint(int32 Index) const
 	{
 		if (PivotPointsAttribute.IsValid() && PivotPointsAttribute.IsValidIndex(Index))
 		{
 			return PivotPointsAttribute[Index];
 		}
-
+ 
 		static const FVector3f DefaultPivotPoint = FVector3f::Zero();
 		return DefaultPivotPoint;
 	}
-
+ 
 	const FVector3f& FFoliageFacade::GetUpVector(const int32 Index) const
 	{
 		if (UpVectorsAttribute.IsValid() && UpVectorsAttribute.IsValidIndex(Index))
 		{
 			return UpVectorsAttribute[Index];
 		}
-
+ 
 		static const FVector3f DefaultUpVector = FVector3f::Zero();
 		return DefaultUpVector;
 	}
-
+ 
 	const FVector3f& FFoliageFacade::GetNormalVector(const int32 Index) const
 	{
 		if (NormalVectorsAttribute.IsValid() && NormalVectorsAttribute.IsValidIndex(Index))
 		{
 			return NormalVectorsAttribute[Index];
 		}
-
+ 
 		static const FVector3f DefaultNormalVector = FVector3f::Zero();
 		return DefaultNormalVector;
 	}
-
+ 
 	void FFoliageFacade::SetPivotPoint(const int32 Index, const FVector3f& Input)
 	{
 		check(!IsConst());
@@ -306,7 +405,7 @@ namespace PV::Facades
 			PivotPointsAttribute.ModifyAt(Index, Input);
 		}
 	}
-
+ 
 	void FFoliageFacade::SetUpVector(int32 Index, const FVector3f& Input)
 	{
 		check(!IsConst());
@@ -315,7 +414,7 @@ namespace PV::Facades
 			UpVectorsAttribute.ModifyAt(Index, Input);
 		}
 	}
-
+ 
 	void FFoliageFacade::SetNormalVector(int32 Index, const FVector3f& Input)
 	{
 		check(!IsConst());
@@ -324,7 +423,7 @@ namespace PV::Facades
 			NormalVectorsAttribute.ModifyAt(Index, Input);
 		}
 	}
-
+ 
 	void FFoliageFacade::SetScale(int32 Index, float Input)
 	{
 		check(!IsConst());
@@ -333,16 +432,38 @@ namespace PV::Facades
 			ScalesAttribute.ModifyAt(Index, Input);
 		}
 	}
-
+ 
 	void FFoliageFacade::SetFoliageName(int32 Index, const FString& Input)
 	{
 		check(!IsConst());
-		if (IsValid() && NamesAttribute.IsValidIndex(Index))
+		if (IsValid() && FoliageInfoNamesAttribute.IsValidIndex(Index))
 		{
-			NamesAttribute.ModifyAt(Index, Input);
+			FoliageInfoNamesAttribute.ModifyAt(Index, Input);
 		}
 	}
-
+ 
+	const FString FFoliageFacade::GetFoliageName(const int32 Index) const
+	{
+		if (FoliageInfoNamesAttribute.IsValid() && FoliageInfoNamesAttribute.IsValidIndex(Index))
+		{
+			return FoliageInfoNamesAttribute[Index];
+		}
+ 
+		return FString();
+	}
+ 
+	const TArray<FString>& FFoliageFacade::GetFoliageNames() const
+	{
+		TArray<FString> FoliageNames;
+		
+		if (FoliageInfoNamesAttribute.IsValid())
+		{
+			return FoliageInfoNamesAttribute.Get().GetConstArray();
+		}
+ 
+		return FoliageNames;
+	}
+	
 	void FFoliageFacade::SetFoliageNames(const TArray<FString>& InputNames)
 	{
 		check(!IsConst());
@@ -350,39 +471,67 @@ namespace PV::Facades
 		Collection->AddElements(InputNames.Num(), GroupNames::FoliageNamesGroup);
 		for (int32 Index = 0; Index < InputNames.Num(); ++Index)
 		{
-			NamesAttribute.Modify()[Index] = InputNames[Index];
+			FoliageInfoNamesAttribute.Modify()[Index] = InputNames[Index];
+		}
+	}
+ 
+	void FFoliageFacade::SetFoliageInfo(const int32 Index, const FString& Input)
+	{
+		check(!IsConst());
+		if (IsValid() && FoliageInfoNamesAttribute.IsValidIndex(Index))
+		{
+			FoliageInfoNamesAttribute.ModifyAt(Index, Input);
+		}
+	}
+ 
+	void FFoliageFacade::SetFoliageInfos(const TArray<FPVFoliageInfo>& Infos)
+	{
+		check(!IsConst());
+		Collection->EmptyGroup(GroupNames::FoliageNamesGroup);
+		Collection->AddElements(Infos.Num(), GroupNames::FoliageNamesGroup);
+ 
+		for (int32 Index = 0; Index < Infos.Num(); ++Index)
+		{
+			FPVFoliageInfo Info = Infos[Index];
+			
+			FoliageInfoNamesAttribute.Modify()[Index] = Info.Mesh.ToString();
+			FoliageInfoScaleAttribute.Modify()[Index] = Info.Attributes.Scale;
+			FoliageInfoLightAttribute.Modify()[Index] = Info.Attributes.Light;
+			FoliageInfoTipAttribute.Modify()[Index] = Info.Attributes.Tip;
+			FoliageInfoUpAlignmentAttribute.Modify()[Index] = Info.Attributes.UpAlignment;
+			FoliageInfoHealthAttribute.Modify()[Index] = Info.Attributes.Health;
 		}
 	}
-
-	void FFoliageFacade::SetFoliagePath(FString InPath)
+ 
+	void FFoliageFacade::SetFoliagePath(const FString InPath)
 	{
 		check(!IsConst());
-
-		const int32 NumElements = Collection->NumElements(GroupNames::DetailsGroup);
-
+ 
+		int32 NumElements = Collection->NumElements(GroupNames::DetailsGroup);
+ 
 		if (NumElements == 0)
 		{
 			Collection->AddElements(1, GroupNames::DetailsGroup);
 		}
-
-		FoliagePathAttribute.Modify()[0] = MoveTemp(InPath);
+ 
+		FoliagePathAttribute.Modify()[0] = InPath;
 	}
-
+ 
 	FString FFoliageFacade::GetFoliagePath() const
 	{
 		if (FoliagePathAttribute.IsValid() && FoliagePathAttribute.IsValidIndex(0))
 		{
 			return FoliagePathAttribute[0];
 		}
-
+ 
 		return FString();
 	}
-
+ 
 	int32 FFoliageFacade::GetElementCount() const
 	{
 		return NameIdsAttribute.Num();
 	}
-
+ 
 	void FFoliageFacade::CopyEntry(int32 FromIndex, int32 ToIndex)
 	{
 		if (IsValid() && NameIdsAttribute.IsValidIndex(FromIndex) && NameIdsAttribute.IsValidIndex(ToIndex))
@@ -394,10 +543,15 @@ namespace PV::Facades
 			NormalVectorsAttribute.ModifyAt(ToIndex, NormalVectorsAttribute[FromIndex]);
 			ScalesAttribute.ModifyAt(ToIndex, ScalesAttribute[FromIndex]);
 			LengthFromRootAttribute.ModifyAt(ToIndex, LengthFromRootAttribute[FromIndex]);
+			ConditionTipAttribute.ModifyAt(ToIndex, ConditionTipAttribute[FromIndex]);
+			ConditionLightAttribute.ModifyAt(ToIndex, ConditionLightAttribute[FromIndex]);
+			ConditionScaleAttribute.ModifyAt(ToIndex, ConditionScaleAttribute[FromIndex]);
+			ConditionUpAlignmentAttribute.ModifyAt(ToIndex, ConditionUpAlignmentAttribute[FromIndex]);
+			ConditionHealthAttribute.ModifyAt(ToIndex, ConditionHealthAttribute[FromIndex]);
 			ParentBoneIdsAttribute.ModifyAt(ToIndex, ParentBoneIdsAttribute[FromIndex]);
 		}
 	}
-
+ 
 	void FFoliageFacade::RemoveEntries(int32 NumEntries, int32 StartIndex)
 	{
 		if (IsValid() && NameIdsAttribute.IsValidIndex(StartIndex) && StartIndex + NumEntries <= NameIdsAttribute.Num())
@@ -405,19 +559,19 @@ namespace PV::Facades
 			NameIdsAttribute.RemoveElements(NumEntries, StartIndex);
 		}
 	}
-
+ 
 	FTransform FFoliageFacade::GetFoliageTransform(int32 Id) const
 	{
 		const FFoliageEntryData Data = GetFoliageEntry(Id);
-
+ 
 		const FVector UpVectorN = FVector(Data.UpVector).GetSafeNormal();
 		if (UpVectorN.IsNearlyZero())
 		{
 			return FTransform(FQuat::Identity, FVector(Data.PivotPoint), FVector(Data.Scale));
 		}
-
+ 
 		const FVector NormalVectorN = FVector(Data.NormalVector).GetSafeNormal();
-
+ 
 		FVector OrthogonalizedNormalVectorN =
 			(NormalVectorN - FVector::DotProduct(NormalVectorN, UpVectorN) * UpVectorN).GetSafeNormal();
 		if (OrthogonalizedNormalVectorN.IsNearlyZero())
@@ -425,18 +579,18 @@ namespace PV::Facades
 			FVector TempVector;
 			UpVectorN.FindBestAxisVectors(TempVector, OrthogonalizedNormalVectorN);
 		}
-
+ 
 		const FQuat RotationQuat = FQuat(FRotationMatrix::MakeFromYZ(OrthogonalizedNormalVectorN, UpVectorN));
 		return FTransform(RotationQuat, FVector(Data.PivotPoint), FVector(Data.Scale));
 	}
-
+ 
 	const TManagedArray<FVector3f>& FFoliageFacade::GetPivotPositions() const
 	{
 		return PivotPointsAttribute.Get();
 	}
-
+ 
 	TManagedArray<FVector3f>& FFoliageFacade::ModifyPivotPositions()
 	{
 		return PivotPointsAttribute.Modify();
 	}
-}
+}
\ No newline at end of file
diff --git a/Engine/Plugins/Experimental/ProceduralVegetationEditor/Source/ProceduralVegetation/Private/Helpers/PVAnalyticsHelper.cpp b/Engine/Plugins/Experimental/ProceduralVegetationEditor/Source/ProceduralVegetation/Private/Helpers/PVAnalyticsHelper.cpp
--- a/Engine/Plugins/Experimental/ProceduralVegetationEditor/Source/ProceduralVegetation/Private/Helpers/PVAnalyticsHelper.cpp
+++ b/Engine/Plugins/Experimental/ProceduralVegetationEditor/Source/ProceduralVegetation/Private/Helpers/PVAnalyticsHelper.cpp
@@ -1,7 +1,7 @@
 ﻿// Copyright Epic Games, Inc. All Rights Reserved.
-
+ 
 #include "Helpers/PVAnalyticsHelper.h"
-
+ 
 #include "EngineAnalytics.h"
 #include "ProceduralVegetationModule.h"
 #include "PVExportParams.h"
@@ -10,8 +10,9 @@
 #include "GeometryCollection/GeometryCollection.h"
 #include "GeometryCollection/ManagedArrayCollection.h"
 #include "Helpers/PVUtilities.h"
+#include "Implementations/PVFoliage.h"
 #include "UObject/UnrealTypePrivate.h"
-
+ 
 namespace PV::Analytics
 {
 	void LogEvent(const FString& InEventName, const TArray<FAnalyticsEventAttribute>& InAttributes = TArray<FAnalyticsEventAttribute>())
@@ -33,7 +34,7 @@ namespace PV::Analytics
 			FEngineAnalytics::GetProvider().RecordEvent(EventName, InAttributes);
 		}
 	}
-
+ 
 	void SendSessionStartedEvent()
 	{
 		if (FEngineAnalytics::IsAvailable())
@@ -41,7 +42,7 @@ namespace PV::Analytics
 			SendEvent(TEXT("SessionStarted"));
 		}
 	}
-
+ 
 	void SendSessionEndedEvent(const double InTimeInSeconds)
 	{
 		if (FEngineAnalytics::IsAvailable())
@@ -52,7 +53,7 @@ namespace PV::Analytics
 			SendEvent(TEXT("SessionEnded"), Attributes);
 		}
 	}
-
+ 
 	void SendNodeAddedEvent(const FString& InNodeType)
 	{
 		if (FEngineAnalytics::IsAvailable())
@@ -64,13 +65,13 @@ namespace PV::Analytics
 			LogEvent(TEXT("NodeAdded"), Attributes);
 		}
 	}
-
+ 
 	void SendNodeTweakedEvent(const FString& InNodeType, const FProperty* InProperty)
 	{
 		if (FEngineAnalytics::IsAvailable() && InProperty)
 		{
 			FString PropertyName = InProperty->GetName();
-
+ 
 			FString PropertyCategory = "N/A";
 #if WITH_METADATA
 			PropertyCategory = InProperty->GetMetaData(TEXT("Category"));
@@ -85,7 +86,7 @@ namespace PV::Analytics
 			LogEvent(TEXT("NodeTweaked"), Attributes);
 		}
 	}
-
+ 
 	void SendMaterialChangeEvent(FString InMatPath)
 	{
 		if (FEngineAnalytics::IsAvailable())
@@ -97,7 +98,7 @@ namespace PV::Analytics
 			LogEvent(TEXT("MeshMaterialChanged"), Attributes);
 		}
 	}
-
+ 
 	void SendFoliageMeshChangeEvent(FString InFoliageMeshPath)
 	{
 		if (FEngineAnalytics::IsAvailable())
@@ -109,7 +110,7 @@ namespace PV::Analytics
 			LogEvent(TEXT("FoliageMeshChanged"), Attributes);
 		}
 	}
-
+ 
 	void SendWindSettingsChangeEvent(FString InPath)
 	{
 		if (FEngineAnalytics::IsAvailable())
@@ -121,30 +122,30 @@ namespace PV::Analytics
 			LogEvent(TEXT("WindPresetChanged"), Attributes);
 		}
 	}
-
+ 
 	void GatherFoliageMeshesTriangles(const Facades::FFoliageFacade& InFoliageFacade, uint32& OutUniqueFoliageTriangles, uint32& OutTotalFoliageTriangles)
 	{
 		OutUniqueFoliageTriangles = 0;
 		OutTotalFoliageTriangles = 0;
 		
-		for (int32 FoliageNameIndex = 0; FoliageNameIndex < InFoliageFacade.NumFoliageNames(); ++FoliageNameIndex)
+		for (int32 FoliageIndex = 0; FoliageIndex < InFoliageFacade.NumFoliageInfo(); ++FoliageIndex)
 		{
-			FString FoliagePath = InFoliageFacade.GetFoliageName(FoliageNameIndex);
+			FPVFoliageInfo Foliage = InFoliageFacade.GetFoliageInfo(FoliageIndex);
 			
-			OutUniqueFoliageTriangles += PV::Utilities::GetMeshTriangles(FoliagePath);
+			OutUniqueFoliageTriangles += PV::Utilities::GetMeshTriangles(Foliage.Mesh.ToString());
 		}
-
+ 
 		int32 NumInstances = InFoliageFacade.NumFoliageEntries();
 		
 		for (int32 Id = 0; Id < NumInstances; Id++)
 		{
 			Facades::FFoliageEntryData Data = InFoliageFacade.GetFoliageEntry(Id);
-			FString FoliagePath = InFoliageFacade.GetFoliageName(Data.NameId);
-
-			OutTotalFoliageTriangles += PV::Utilities::GetMeshTriangles(FoliagePath);
+			FPVFoliageInfo FoliageInfo = InFoliageFacade.GetFoliageInfo(Data.NameId);
+ 
+			OutTotalFoliageTriangles += PV::Utilities::GetMeshTriangles(FoliageInfo.Mesh.ToString());
 		}
 	}
-
+ 
 	TArray<FAnalyticsEventAttribute> GatherExportCommonAttributes(const FManagedArrayCollection& InCollection, const FPVExportParams& InExportParam)
 	{
 		TArray<FAnalyticsEventAttribute> ExportAttributes;
@@ -154,12 +155,12 @@ namespace PV::Analytics
 		ExportAttributes.Add(FAnalyticsEventAttribute(TEXT("ExportSettings.NaniteShapePreservation"),  StaticEnum<ENaniteShapePreservation>()->GetNameStringByValue((uint8)InExportParam.NaniteShapePreservation)));
 		ExportAttributes.Add(FAnalyticsEventAttribute(TEXT("NaniteSettings.NaniteFoliageEnabled"),  GetDefault<URendererSettings>()->bEnableNaniteFoliage));
 		ExportAttributes.Add(FAnalyticsEventAttribute(TEXT("NaniteSettings.Enabled"),  GetDefault<URendererSettings>()->bNanite));
-
+ 
 		Facades::FFoliageFacade FoliageFacade(InCollection);
 		
 		const int FoliageInstances = FoliageFacade.GetElementCount();
 		const int TrunkFaces = InCollection.NumElements(FGeometryCollection::FacesGroup);
-
+ 
 		uint32 UniqueFoliageTriangles = 0;
 		uint32 TotalFoliageTriangles = 0;
 		GatherFoliageMeshesTriangles(FoliageFacade, UniqueFoliageTriangles, TotalFoliageTriangles);
@@ -168,10 +169,10 @@ namespace PV::Analytics
 		ExportAttributes.Add(FAnalyticsEventAttribute(TEXT("TrunkTriangles"),  TrunkFaces));
 		ExportAttributes.Add(FAnalyticsEventAttribute(TEXT("FoliageTrianglesUnique"),  UniqueFoliageTriangles));
 		ExportAttributes.Add(FAnalyticsEventAttribute(TEXT("FoliageTrianglesTotal"),  TotalFoliageTriangles));
-
+ 
 		return ExportAttributes;
 	}
-
+ 
 	void SendExportStartEvent(const TArray<FAnalyticsEventAttribute>& InCommonAttributes)
 	{
 		if (FEngineAnalytics::IsAvailable())
@@ -180,13 +181,13 @@ namespace PV::Analytics
 			LogEvent(TEXT("ExportStarted"), InCommonAttributes);
 		}
 	}
-
+ 
 	void SendExportFinishedEvent(TArray<FAnalyticsEventAttribute> InCommonAttributes, float InTotalTimeInSeconds, EExportResult ExportResult)
 	{
 		if (FEngineAnalytics::IsAvailable())
 		{
 			InCommonAttributes.Add(FAnalyticsEventAttribute(TEXT("TotalTimeSeconds"), InTotalTimeInSeconds));
-
+ 
 			const static auto ExportResultEnumToString = [](EExportResult InExportResult)->FString {
 				switch (InExportResult)
 				{
@@ -197,14 +198,14 @@ namespace PV::Analytics
 				case EExportResult::Skipped:
 					return TEXT("Skipped");
 				}
-
+ 
 				return TEXT("Unknown");
 			};
 			
 			InCommonAttributes.Add(FAnalyticsEventAttribute(TEXT("ExportResult"), ExportResultEnumToString(ExportResult)));
-
+ 
 			SendEvent(TEXT("ExportFinished"), InCommonAttributes);
 			LogEvent(TEXT("ExportFinished"), InCommonAttributes);
 		}
 	}
-}
+}
\ No newline at end of file
diff --git a/Engine/Plugins/Experimental/ProceduralVegetationEditor/Source/ProceduralVegetation/Private/Helpers/PVFoliageJSONHelper.cpp b/Engine/Plugins/Experimental/ProceduralVegetationEditor/Source/ProceduralVegetation/Private/Helpers/PVFoliageJSONHelper.cpp
new file mode 100644
--- /dev/null
+++ b/Engine/Plugins/Experimental/ProceduralVegetationEditor/Source/ProceduralVegetation/Private/Helpers/PVFoliageJSONHelper.cpp
@@ -0,0 +1,407 @@
+﻿// Copyright Epic Games, Inc. All Rights Reserved.
+
+#include "Helpers/PVFoliageJSONHelper.h"
+
+#include "ProceduralVegetationModule.h"
+#include "Implementations/PVFoliage.h"
+#include "Facades/PVAttributesNames.h"
+#include "Facades/PVFoliageFacade.h"
+#include "Dom/JsonObject.h"
+#include "Dom/JsonValue.h"
+#include "Facades/PVFoliageConditionFacade.h"
+#include "Misc/FileHelper.h"
+#include "Serialization/JsonSerializer.h"
+#include "Misc/Paths.h"
+#include "Misc/PackageName.h"
+
+namespace PV
+{
+	bool PVFoliageJSONHelper::LoadJSON(const FString& InPath, TSharedPtr<FJsonObject>& OutData, FString& OutErrorMessage)
+	{
+		FString JsonString;
+		if (!FFileHelper::LoadFileToString(JsonString, *InPath))
+		{
+			OutErrorMessage = TEXT("Failed to open file");
+			return false;
+		}
+
+		OutData = MakeShareable(new FJsonObject);
+		if (!FJsonSerializer::Deserialize(TJsonReaderFactory<>::Create(JsonString), OutData))
+		{
+			OutErrorMessage = TEXT("Failed to deserialize JSON file");
+			return false;
+		}
+
+		return true;
+	}
+
+	bool GetConditionFromString(const FString& InString, EFoliageDistributionCondition& OutCondition)
+	{
+		if (InString.Equals("light", ESearchCase::IgnoreCase))
+		{
+			OutCondition = EFoliageDistributionCondition::Light;
+			return true;
+		}
+		else if (InString.Equals("scale", ESearchCase::IgnoreCase))
+		{
+			OutCondition = EFoliageDistributionCondition::Scale;
+			return true;
+		}
+		else if (InString.Equals("upAlignment", ESearchCase::IgnoreCase))
+		{
+			OutCondition = EFoliageDistributionCondition::UpAlignment;
+			return true;
+		}
+		else if (InString.Equals("health", ESearchCase::IgnoreCase))
+		{
+			OutCondition = EFoliageDistributionCondition::Health;
+			return true;
+		}
+		else if (InString.Equals("tip", ESearchCase::IgnoreCase))
+		{
+			OutCondition = EFoliageDistributionCondition::Tip;
+			return true;
+		}
+			
+		return false;
+	}
+
+	bool PVFoliageJSONHelper::LoadFoliageData(const FString& InPath, const FString& InPackagePath, FoliageVariationsMap& OutFoliageVariationsMap, FString& OutErrorMessage)
+	{
+		TSharedPtr<FJsonObject> OutData = MakeShareable(new FJsonObject);
+		
+		OutFoliageVariationsMap.Empty();
+
+		if (LoadJSON(InPath, OutData, OutErrorMessage))
+		{
+			if (OutData->HasTypedField(TEXT("Variations"), EJson::Array))
+			{
+				const TArray<TSharedPtr<FJsonValue>>& Variations = OutData->GetArrayField(TEXT("Variations"));
+
+				for (TSharedPtr<FJsonValue> FoliageInfoVariation : Variations)
+				{
+					if (FoliageInfoVariation->Type == EJson::Object)
+					{
+						const TSharedPtr<FJsonObject> VariationData = FoliageInfoVariation->AsObject();
+
+						FString VariationName = VariationData->GetStringField(TEXT("name"));
+
+						const TArray<TSharedPtr<FJsonValue>>* OutVariationData;
+
+						if (!VariationData->TryGetArrayField(TEXT("Data"), OutVariationData))
+						{
+							OutErrorMessage = FString::Printf(TEXT("Unable to load foliage data : Error: Missing data for variation : %s"), *VariationName);
+							return false;
+						}
+
+						TArray<FPVFoliageInfo> FoliageInfos;
+						
+						for (TSharedPtr<FJsonValue> FoliageInfoData : *OutVariationData)
+						{
+							if (const TSharedPtr<FJsonObject> FoliageInfoObject = FoliageInfoData->AsObject())
+							{
+								FPVFoliageInfo FoliageInfo;
+								FString BranchName;
+
+								if (!FoliageInfoObject->TryGetStringField(TEXT("Name"), BranchName))
+								{
+									OutErrorMessage = FString::Printf(TEXT("FoliageInfo Doesnt contain valid branch name for variation %s"), *VariationName);
+									return false;
+								}
+								
+								FoliageInfo.Mesh = InPackagePath / BranchName + "." + BranchName;
+								
+								if (!FoliageInfoObject->TryGetNumberField(TEXT("scale"), FoliageInfo.Attributes.Scale))
+								{
+									FoliageInfo.Attributes.Scale = 0;
+								}
+								if (!FoliageInfoObject->TryGetNumberField(TEXT("upAlignment"), FoliageInfo.Attributes.UpAlignment))
+								{
+									FoliageInfo.Attributes.UpAlignment = 0;
+								}
+
+								if (!FoliageInfoObject->TryGetNumberField(TEXT("light"), FoliageInfo.Attributes.Light))
+								{
+									FoliageInfo.Attributes.Light = 0;
+								}
+
+								if (!FoliageInfoObject->TryGetNumberField(TEXT("health"), FoliageInfo.Attributes.Health))
+								{
+									FoliageInfo.Attributes.Health = 0;
+								}
+
+								float OutTip; 
+								if (FoliageInfoObject->TryGetNumberField(TEXT("tip"), OutTip))
+								{
+									FoliageInfo.Attributes.Tip = static_cast<bool>(OutTip);	
+								}
+								else
+								{
+									FoliageInfo.Attributes.Tip = false;
+								}
+								
+
+								FoliageInfos.Add(FoliageInfo);
+							}
+						}
+
+						const TArray<TSharedPtr<FJsonValue>>* OutVariationRules;
+
+						if (!VariationData->TryGetArrayField(TEXT("Rules"), OutVariationRules))
+						{
+							OutErrorMessage = FString::Printf(TEXT("Unable to load foliage rules : Error: Missing data for variation : %s"), *VariationName);
+							return false;
+						}
+
+						TMap<EFoliageDistributionCondition, Facades::FFoliageConditonInfo> FoliageDistributionConditions;
+						
+						for (TSharedPtr<FJsonValue> FoliageRule : *OutVariationRules)
+						{
+							if (const TSharedPtr<FJsonObject> Rule = FoliageRule->AsObject())
+							{
+								FString RuleName;
+								
+								if (!Rule->TryGetStringField(TEXT("name"), RuleName))
+								{
+									OutErrorMessage = FString::Printf(TEXT("FoliageData: Rule doesnt contain valid name for variation %s"), *VariationName);
+									continue;
+								}
+
+								EFoliageDistributionCondition OutCondition;
+								
+								if (GetConditionFromString(RuleName, OutCondition))
+								{
+									Facades::FFoliageConditonInfo ConditionSettings;
+
+									const UEnum* EnumPtr = StaticEnum<EFoliageDistributionCondition>();
+									ConditionSettings.Name = EnumPtr->GetNameStringByValue(static_cast<int64>(OutCondition));
+									
+									if (!Rule->TryGetNumberField(TEXT("weight"), ConditionSettings.Weight))
+									{
+										ConditionSettings.Weight = 0;
+									}
+
+									if (!Rule->TryGetNumberField(TEXT("offset"), ConditionSettings.Offset))
+									{
+										ConditionSettings.Offset = 0;
+									}
+
+									// float OutNormalize; 
+									// if (Rule->TryGetNumberField(TEXT("normalize"), OutNormalize))
+									// {
+									// 	ConditionSettings.Normalize = static_cast<bool>(OutNormalize);	
+									// }
+									// else
+									// {
+									// 	ConditionSettings.Normalize = false;
+									// }
+
+									if (!FoliageDistributionConditions.Contains(OutCondition))
+									{
+										FoliageDistributionConditions.Add(OutCondition);
+									}
+									
+									FoliageDistributionConditions[OutCondition] = ConditionSettings;
+								}
+							}
+						}
+
+						if (!OutFoliageVariationsMap.Contains(VariationName))
+						{
+							OutFoliageVariationsMap.Add(VariationName);
+						}
+
+						OutFoliageVariationsMap[VariationName].VariationName = VariationName;
+						OutFoliageVariationsMap[VariationName].FoliageInfos = FoliageInfos;
+						OutFoliageVariationsMap[VariationName].Conditions = FoliageDistributionConditions;
+					}
+				}
+			}
+			else
+			{
+				OutErrorMessage = FString::Printf(TEXT("Unable to load foliage data : Error: Missing variation data"));
+				
+				return false;	
+			}
+		}
+
+		for (auto Mapping : OutFoliageVariationsMap)
+		{
+			UE_LOG(LogProceduralVegetation, Log, TEXT("Variation name : %s : Start"), *Mapping.Key);
+
+			for (const FPVFoliageInfo& FoliageInfo : Mapping.Value.FoliageInfos)
+			{
+				UE_LOG(LogProceduralVegetation, Log, TEXT("FoliageInfo : Mesh : %s : Scale : %f : Light : %f : Tip : %i : UpAlignment :  %f : Health : %f"),
+					*FoliageInfo.Mesh.ToString(),
+					FoliageInfo.Attributes.Scale,
+					FoliageInfo.Attributes.Light,
+					FoliageInfo.Attributes.Tip,
+					FoliageInfo.Attributes.UpAlignment,
+					FoliageInfo.Attributes.Health);	
+			}
+
+			for (const auto& ConditionEntry : Mapping.Value.Conditions)
+			{
+				UE_LOG(LogProceduralVegetation, Log, TEXT("FoliageCondition : %s"), *ConditionEntry.Value.ToString());	
+			}
+
+			UE_LOG(LogProceduralVegetation, Log, TEXT("Variation name : %s : End"), *Mapping.Key);
+		}
+
+		return true;
+	}
+
+	bool PVFoliageJSONHelper::LoadFoliageDataInCollection(FManagedArrayCollection& Collection, const FString& InPresetPath, const FPVFoliageVariationData& VariationData, FString& OutErrorMessage)
+	{
+		TSharedPtr<FJsonObject> OutData = MakeShareable(new FJsonObject);
+
+		FString MegaPlantsData;
+		if (!FFileHelper::LoadFileToString(MegaPlantsData, *InPresetPath))
+		{
+			OutErrorMessage = TEXT("Failed to open skeleton file");
+			return false;
+		}
+
+		TSharedPtr<FJsonObject> LoadedData = MakeShareable(new FJsonObject);
+		if (!FJsonSerializer::Deserialize(TJsonReaderFactory<>::Create(MegaPlantsData), LoadedData))
+		{
+			OutErrorMessage = TEXT("Failed to deserialize JSON file for skeleton");
+			return false;
+		}
+
+		Collection.RemoveGroup(GroupNames::FoliageGroup);
+		Collection.RemoveGroup(GroupNames::FoliageNamesGroup);
+		
+		const TSharedPtr<FJsonObject>& Primitives = LoadedData->GetObjectField(TEXT("primitives"));
+		const TSharedPtr<FJsonObject>& PrimitiveAttributes = Primitives->GetObjectField(TEXT("attributes"));
+
+		const TSharedPtr<FJsonObject>& InstancerNameObject = PrimitiveAttributes->GetObjectField(TEXT("instancer_name"));
+		const TArray<TSharedPtr<FJsonValue>>& InstancerNameValues = InstancerNameObject->GetArrayField(TEXT("values"));
+
+		const TSharedPtr<FJsonObject>& InstancerPivotObject = PrimitiveAttributes->GetObjectField(TEXT("instancer_pivot"));
+		const TArray<TSharedPtr<FJsonValue>>& InstancerPivotValues = InstancerPivotObject->GetArrayField(TEXT("values"));
+
+		const TSharedPtr<FJsonObject>& InstancerUpVectorObject = PrimitiveAttributes->GetObjectField(TEXT("instancer_UP"));
+		const TArray<TSharedPtr<FJsonValue>>& InstancerUpVectorValues = InstancerUpVectorObject->GetArrayField(TEXT("values"));
+
+		const TSharedPtr<FJsonObject>& InstancerNormalVectorObject = PrimitiveAttributes->GetObjectField(TEXT("instancer_N"));
+		const TArray<TSharedPtr<FJsonValue>>& InstancerNormalVectorValues = InstancerNormalVectorObject->GetArrayField(TEXT("values"));
+
+		const TSharedPtr<FJsonObject>& InstancerScaleObject = PrimitiveAttributes->GetObjectField(TEXT("instancer_scale"));
+		const TArray<TSharedPtr<FJsonValue>>& InstancerScaleValues = InstancerScaleObject->GetArrayField(TEXT("values"));
+
+		const TSharedPtr<FJsonObject>& InstancerLFRObject = PrimitiveAttributes->GetObjectField(TEXT("instancer_LFR"));
+		const TArray<TSharedPtr<FJsonValue>>& InstancerLFRValues = InstancerLFRObject->GetArrayField(TEXT("values"));
+
+		
+		TMap<FString, int32> FoliageNamesToIndex;
+		for (int FoliageIndex = 0; FoliageIndex < VariationData.FoliageInfos.Num(); ++FoliageIndex)
+		{
+			auto FoliageInfo = VariationData.FoliageInfos[FoliageIndex];
+			FoliageNamesToIndex.Emplace(FoliageInfo.Mesh.GetAssetName(), FoliageIndex);
+		}
+		
+		// Compute unique names and number of elements for group
+		int32 TotalValues = 0;
+		TArray<FString> FoliageNames;
+		for (int32 i = 0; i < InstancerNameValues.Num(); ++i)
+		{
+			const TArray<TSharedPtr<FJsonValue>>& InstancerNames = InstancerNameValues[i]->AsArray();
+			for (int32 Index = 0; Index < InstancerNames.Num(); ++Index)
+			{
+				const FString& Name = InstancerNames[Index]->AsString();
+
+				if (!FoliageNamesToIndex.Contains(Name))
+				{
+					OutErrorMessage = FString::Printf(TEXT("FoliageData doesnt contain branch info : %s"), *Name);
+					return false;
+				}
+			}
+			TotalValues += InstancerNames.Num();
+		}
+		
+		Facades::FFoliageFacade Facade(Collection, TotalValues);
+		Facade.SetFoliageInfos(VariationData.FoliageInfos);
+
+		TArray<Facades::FFoliageConditonInfo> Conditions;
+
+		for (auto ConditionEntry : VariationData.Conditions)
+		{
+			Conditions.Add(ConditionEntry.Value);
+		}
+
+		Facades::FFoliageConditionFacade ConditionFacade(Collection);
+		ConditionFacade.SetData(Conditions);
+		
+		int32 FoliageEntryIndex = 0;
+		for (int32 Id = 0; Id < InstancerNameValues.Num(); ++Id)
+		{
+			const TArray<TSharedPtr<FJsonValue>>& InstancerNames = InstancerNameValues[Id]->AsArray();
+			const TArray<TSharedPtr<FJsonValue>>& PivotPoints = InstancerPivotValues[Id]->AsArray();
+			const TArray<TSharedPtr<FJsonValue>>& UpVectors = InstancerUpVectorValues[Id]->AsArray();
+			const TArray<TSharedPtr<FJsonValue>>& NormalVectors = InstancerNormalVectorValues[Id]->AsArray();
+			const TArray<TSharedPtr<FJsonValue>>& Scales = InstancerScaleValues[Id]->AsArray();
+			const TArray<TSharedPtr<FJsonValue>>& LFRs = InstancerLFRValues[Id]->AsArray();
+		
+			TArray<int32> FoliageEntryIds;
+			for (int32 j = 0; j < InstancerNames.Num(); ++j)
+			{
+				const FString Name = InstancerNames[j]->AsString();
+				const FVector3f PivotPoint = FVector3f(
+					PivotPoints[j * 3]->AsNumber(),
+					PivotPoints[(j * 3) + 2]->AsNumber(),
+					PivotPoints[(j * 3) + 1]->AsNumber()) * 100.0f;
+				const FVector3f UpVector = FVector3f(
+					UpVectors[j * 3]->AsNumber(),
+					UpVectors[(j * 3) + 2]->AsNumber(),
+					UpVectors[(j * 3) + 1]->AsNumber());
+				const FVector3f NormalVector = FVector3f(
+					NormalVectors[j * 3]->AsNumber(),
+					NormalVectors[(j * 3) + 2]->AsNumber(),
+					NormalVectors[(j * 3) + 1]->AsNumber());
+				const float Scale = Scales[j]->AsNumber();
+				const float LengthFromRoot = LFRs[j]->AsNumber();
+		
+				Facade.SetFoliageEntry(FoliageEntryIndex, {
+					.NameId = FoliageNamesToIndex[Name],
+					.BranchId = Id,
+					.PivotPoint = PivotPoint,
+					.UpVector = UpVector,
+					.NormalVector = NormalVector,
+					.Scale = Scale,
+					.LengthFromRoot = LengthFromRoot,
+				});
+		
+				FoliageEntryIds.Add(FoliageEntryIndex);
+				FoliageEntryIndex++;
+			}
+		
+			Facade.SetFoliageIdsArray(Id, FoliageEntryIds);
+		}
+
+		return true;
+	}
+
+
+	void PVFoliageJSONHelper::SetFoliagePaths(FManagedArrayCollection& Collection, const FString& FilePath)
+	{
+		PV::Facades::FFoliageFacade Facade(Collection);
+
+		FString AbsolutePath = FPaths::ConvertRelativePathToFull(FilePath);
+		FString PackagePath;
+
+		if (FPackageName::TryConvertFilenameToLongPackageName(AbsolutePath, PackagePath))
+		{
+			TArray<FPVFoliageInfo> Infos = Facade.GetFoliageInfos();
+			FString PackageFolder = FPackageName::GetLongPackagePath(PackagePath);
+			Facade.SetFoliagePath(PackageFolder);
+
+			for (FPVFoliageInfo& Info : Infos)
+			{
+				Info.Mesh = PackageFolder / Info.Mesh.ToString();
+			}
+			
+			Facade.SetFoliageInfos(Infos);
+		}
+	}
+}
diff --git a/Engine/Plugins/Experimental/ProceduralVegetationEditor/Source/ProceduralVegetation/Private/Helpers/PVFoliageJSONHelper.h b/Engine/Plugins/Experimental/ProceduralVegetationEditor/Source/ProceduralVegetation/Private/Helpers/PVFoliageJSONHelper.h
new file mode 100644
--- /dev/null
+++ b/Engine/Plugins/Experimental/ProceduralVegetationEditor/Source/ProceduralVegetation/Private/Helpers/PVFoliageJSONHelper.h
@@ -0,0 +1,34 @@
+// Copyright Epic Games, Inc. All Rights Reserved.
+#pragma once
+
+#include "CoreMinimal.h"
+#include "Facades/PVFoliageConditionFacade.h"
+#include "Templates/SharedPointer.h"
+
+struct FPVFoliageDistributionConditionSettings;
+enum class EFoliageDistributionCondition : uint8;
+class FJsonObject;
+struct FPVFoliageInfo;
+struct FManagedArrayCollection;
+
+namespace PV
+{
+	typedef TMap<EFoliageDistributionCondition, Facades::FFoliageConditonInfo> FoliageConditionMap;
+	
+	struct FPVFoliageVariationData
+	{
+		FString VariationName;
+		TArray<FPVFoliageInfo> FoliageInfos;
+		FoliageConditionMap Conditions;
+	};
+
+	typedef TMap<FString, FPVFoliageVariationData> FoliageVariationsMap;
+	
+	struct PVFoliageJSONHelper
+	{
+		static bool LoadJSON(const FString& InPath, TSharedPtr<FJsonObject>& OutData, FString& OutErrorMessage);
+		static bool LoadFoliageData(const FString& InPath, const FString& InPackagePath, FoliageVariationsMap& OutFoliageVariationsMap, FString& OutErrorMessage);
+		static bool LoadFoliageDataInCollection(FManagedArrayCollection& Collection, const FString& InPresetPath, const FPVFoliageVariationData& VariationData, FString& OutErrorMessage);
+		static void SetFoliagePaths(FManagedArrayCollection& Collection, const FString& FilePath);
+	};	
+}
\ No newline at end of file
diff --git a/Engine/Plugins/Experimental/ProceduralVegetationEditor/Source/ProceduralVegetation/Private/Helpers/PVUtilities.cpp b/Engine/Plugins/Experimental/ProceduralVegetationEditor/Source/ProceduralVegetation/Private/Helpers/PVUtilities.cpp
--- a/Engine/Plugins/Experimental/ProceduralVegetationEditor/Source/ProceduralVegetation/Private/Helpers/PVUtilities.cpp
+++ b/Engine/Plugins/Experimental/ProceduralVegetationEditor/Source/ProceduralVegetation/Private/Helpers/PVUtilities.cpp
@@ -1,5 +1,5 @@
 ﻿// Copyright Epic Games, Inc. All Rights Reserved.
-
+ 
 #include "Helpers/PVUtilities.h"
 #include "ProceduralVegetationModule.h"
 #include "ProceduralVegetationPreset.h"
@@ -15,15 +15,16 @@
 #include "AssetRegistry/AssetData.h"
 #include "Facades/PVBranchFacade.h"
 #include "Facades/PVPointFacade.h"
+#include "Implementations/PVFoliage.h"
 #include "Misc/FileHelper.h"
 #include "Rendering/SkeletalMeshRenderData.h"
-
+ 
 namespace PV::Utilities
 {
 #if WITH_EDITOR
 	static TAutoConsoleVariable<bool> CVarEnablePVDebugMode(
 		TEXT("PV.DebugMode.Enabled"),
-		false,
+		true,
 		TEXT("Enables debug mode for the Procedural Vegetation Editor"), 
 		FConsoleVariableDelegate::CreateLambda([](IConsoleVariable* ConsoleVariable)
 			{
@@ -32,7 +33,7 @@ namespace PV::Utilities
 			})
 	);
 #endif
-
+ 
 	bool DebugModeEnabled()
 	{
 #if WITH_EDITOR
@@ -41,12 +42,12 @@ namespace PV::Utilities
 		return false;
 #endif
 	}
-
+ 
 	void AddFoliageInstancesToActor(const Facades::FFoliageFacade& FoliageFacade, AActor* Actor,
 	                                TMap<FString, TObjectPtr<UMeshComponent>>& OutInstancedComponentMap)
 	{
 		TRACE_CPUPROFILER_EVENT_SCOPE(PV::Utilities::AddFoliageInstancesToActor());
-
+ 
 		auto OnComponentCreated = [&Actor](UMeshComponent* InInstancedComponent)
 			{
 				InInstancedComponent->RegisterComponentWithWorld(Actor->GetWorld());
@@ -56,25 +57,26 @@ namespace PV::Utilities
 		
 		AddFoliageInstances(FoliageFacade, Actor, FFoliageComponentCreatedCallback::CreateLambda(OnComponentCreated),OutInstancedComponentMap);
 	}
-
+ 
 	void AddFoliageInstances(const Facades::FFoliageFacade& InFoliageFacade, const UObject* InParent, FFoliageComponentCreatedCallback InCallback, TMap<FString, TObjectPtr<UMeshComponent>>& OutInstancedComponentMap)
 	{
 		TRACE_CPUPROFILER_EVENT_SCOPE(PV::Utilities::AddFoliageInstances());
-
+ 
 		int32 NumInstances = InFoliageFacade.NumFoliageEntries();
-
+ 
 		OutInstancedComponentMap.Empty();
-
+ 
 		if (NumInstances > 0)
 		{
-			for (int32 FoliageNameIndex = 0; FoliageNameIndex < InFoliageFacade.NumFoliageNames(); ++FoliageNameIndex)
+			for (int32 FoliageNameIndex = 0; FoliageNameIndex < InFoliageFacade.NumFoliageInfo(); ++FoliageNameIndex)
 			{
-				FString FoliageName = InFoliageFacade.GetFoliageName(FoliageNameIndex);
+				FPVFoliageInfo FoliageInfo = InFoliageFacade.GetFoliageInfo(FoliageNameIndex);
+				FString FoliageName = FoliageInfo.Mesh.ToString();
 				FSoftObjectPath MeshPath(FoliageName);
-
+ 
 				UStaticMesh* StaticFoliageMesh = PackageExists(MeshPath.GetLongPackageName(), UStaticMesh::StaticClass()) ? LoadObject<UStaticMesh>(nullptr, *FoliageName) : nullptr;
 				USkeletalMesh* SkeletalFoliageMesh = PackageExists(MeshPath.GetLongPackageName(), USkeletalMesh::StaticClass()) ? LoadObject<USkeletalMesh>(nullptr, *FoliageName) : nullptr;
-
+ 
 				if(StaticFoliageMesh)
 				{
 					UInstancedStaticMeshComponent* InstancedComponent = NewObject<UInstancedStaticMeshComponent>(InParent != nullptr ? InParent->GetPackage() : GetTransientPackage(), StaticFoliageMesh->GetFName());
@@ -102,16 +104,17 @@ namespace PV::Utilities
 					InCallback.ExecuteIfBound(InstancedComponent);
 				}
 			}
-
+ 
 			TMap<FString, TArray<FTransform>> InstanceTransformsMap;
-
+ 
 			for (int32 Id = 0; Id < NumInstances; Id++)
 			{
 				Facades::FFoliageEntryData Data = InFoliageFacade.GetFoliageEntry(Id);
-				FString FoliageName = InFoliageFacade.GetFoliageName(Data.NameId);
-
+				FPVFoliageInfo FoliageInfo = InFoliageFacade.GetFoliageInfo(Data.NameId);
+				FString FoliageName = FoliageInfo.Mesh.ToString();
+				
 				const FTransform Transform = InFoliageFacade.GetFoliageTransform(Id);
-
+ 
 				InstanceTransformsMap.FindOrAdd(FoliageName).Add(Transform);
 			}
 		
@@ -136,7 +139,7 @@ namespace PV::Utilities
 			}
 		}
 	}
-
+ 
 	FString FormatLongPackageNameErrorCode(const FPackageName::EErrorCode ErrorCode)
 	{
 		switch (ErrorCode)
@@ -164,7 +167,7 @@ namespace PV::Utilities
 		}
 		return FString();
 	}
-
+ 
 	bool ValidateAssetPathAndName(const FString& AssetName, const FString& Path, UClass* InClass, FString& OutError)
 	{
 		if (FPackageName::EErrorCode ErrorCode; !FPackageName::IsValidLongPackageName(Path, false, &ErrorCode))
@@ -172,24 +175,24 @@ namespace PV::Utilities
 			OutError =  FormatLongPackageNameErrorCode(ErrorCode);
         	return false;
         }
-
+ 
 		if (FText Reason; !IsFileNameValid(FName(AssetName), Reason))
         {
 			OutError = Reason.ToString();
         	return false;
         }
-
+ 
 		return true;
 	}
-
+ 
 	bool IsFileNameValid(FName FileName, FText& Reason)
 	{
 		FileName.IsValidObjectName(Reason);
 		FFileHelper::IsFilenameValidForSaving(FileName.ToString(), Reason);
-
+ 
 		return Reason.IsEmpty();
 	}
-
+ 
 	bool DoesConflictingPackageExist(const FString& LongPackageName, UClass* AssetClass)
 	{
 		UPackage* Package = FindPackage(nullptr, *LongPackageName);
@@ -203,22 +206,22 @@ namespace PV::Utilities
 			TArray<FAssetData> OutAssets;
 			IAssetRegistry& AssetRegistry = FAssetRegistryModule::GetRegistry();
 			AssetRegistry.GetAssetsByPackageName(*LongPackageName, OutAssets, true);
-
+ 
 			if (OutAssets.Num() == 0)
 			{
 				return false;
 			}
-
+ 
 			if (OutAssets.Num() != 1
 				|| OutAssets[0].AssetClassPath != AssetClass->GetClassPathName())
 			{
 				return true;
 			}
 		}
-
+ 
 		return false;
 	}
-
+ 
 	bool PackageExists(const FString& LongPackageName, UClass* AssetClass)
 	{
 		UPackage* Package = FindPackage(nullptr, *LongPackageName);
@@ -232,26 +235,26 @@ namespace PV::Utilities
 			TArray<FAssetData> OutAssets;
 			IAssetRegistry& AssetRegistry = FAssetRegistryModule::GetRegistry();
 			AssetRegistry.GetAssetsByPackageName(*LongPackageName, OutAssets, true);
-
+ 
 			if (OutAssets.Num() == 0)
 			{
 				return false;
 			}
-
+ 
 			if (OutAssets[0].AssetClassPath == AssetClass->GetClassPathName())
 			{
 				return true;
 			}
 		}
-
+ 
 		return false;
 	}
-
+ 
 	int32 GetMeshTriangles(const FString InMeshPath)
 	{
 		UStaticMesh* StaticFoliageMesh = LoadObject<UStaticMesh>(nullptr, *InMeshPath);
 		USkeletalMesh* SkeletalFoliageMesh = LoadObject<USkeletalMesh>(nullptr, *InMeshPath);
-
+ 
 		int32 NumTriangles = 0;
 		
 		if (StaticFoliageMesh)
@@ -266,13 +269,13 @@ namespace PV::Utilities
 				const FSkeletalMeshLODRenderData& LODData = SkelMeshRenderData->LODRenderData[0];
 				NumTriangles = LODData.GetTotalFaces();
 			}
-
+ 
 			return NumTriangles;
 		}
-
+ 
 		return NumTriangles;
 	}
-
+ 
 	FLinearColor GetRandomHueColor(float Alpha)
 	{
 		Alpha = FMath::Clamp(Alpha, 0.f, 1.f);
@@ -285,7 +288,7 @@ namespace PV::Utilities
 		);
 		return FLinearColor(HueDeg, 0.9, V).HSVToLinearRGB();
 	};
-
+ 
 	bool IsValidPVData(const FManagedArrayCollection& Collection)
 	{
 		Facades::FPointFacade PointFacade = Facades::FPointFacade(Collection);
@@ -293,4 +296,4 @@ namespace PV::Utilities
 	
 		return PointFacade.IsValid() && BranchFacade.IsValid();
 	}
-}
+}
\ No newline at end of file
diff --git a/Engine/Plugins/Experimental/ProceduralVegetationEditor/Source/ProceduralVegetation/Private/Implementations/PVFoliage.cpp b/Engine/Plugins/Experimental/ProceduralVegetationEditor/Source/ProceduralVegetation/Private/Implementations/PVFoliage.cpp
--- a/Engine/Plugins/Experimental/ProceduralVegetationEditor/Source/ProceduralVegetation/Private/Implementations/PVFoliage.cpp
+++ b/Engine/Plugins/Experimental/ProceduralVegetationEditor/Source/ProceduralVegetation/Private/Implementations/PVFoliage.cpp
@@ -1,49 +1,114 @@
 // Copyright Epic Games, Inc. All Rights Reserved.
-
+ 
 #include "PVFoliage.h"
-
+ 
 #include "ProceduralVegetationModule.h"
 #include "PVFloatRamp.h"
-
+#include "Algo/RandomShuffle.h"
+ 
 #include "Facades/PVAttributesNames.h"
 #include "Facades/PVBranchFacade.h"
+#include "Facades/PVBudVectorsFacade.h"
 #include "Facades/PVFoliageFacade.h"
 #include "Facades/PVPointFacade.h"
-
+ 
+struct ConditionAttributes
+{
+	TArray<float> UpAlignment;
+	TArray<float> Light;
+	TArray<float> Scale;
+	TArray<float> Tip;
+	TArray<float> Health;
+	
+ 
+	float GetConditionValue(EFoliageDistributionCondition InCondition, const int Index) const
+	{
+		TArray<float> Result;
+		switch (InCondition)
+		{
+		case EFoliageDistributionCondition::Light:
+			{
+				Result = Light;
+				break;
+			}
+		case EFoliageDistributionCondition::Scale:
+			{
+				Result = Scale;
+				break;;
+			}
+		case EFoliageDistributionCondition::Tip:
+			{
+				Result = Tip;
+				break;
+			}
+		case EFoliageDistributionCondition::UpAlignment:
+			{
+				Result = UpAlignment;
+				break;
+			}
+		case EFoliageDistributionCondition::Health:
+			{
+				Result = Health;
+				break;
+			}
+		}
+ 
+		if (Result.IsValidIndex(Index))
+		{
+			return Result[Index];
+		}
+ 
+		return 0;
+	}
+};
+ 
+// Condition Weight
+struct FFoliageWeightInfo
+{
+	int FoliageIndex = 0;
+	float Weight = 0;
+};
+ 
 void FPVFoliage::DistributeFoliage(FManagedArrayCollection& OutCollection, const FManagedArrayCollection& SourceCollection,
                                    const FDistributionSettings& DistributionSettings, const FScaleSettings& ScaleSettings,
                                    const FVectorSettings& VectorSettings, const FPhyllotaxySettings& PhyllotaxySettings,
-                                   const FMiscSettings& MiscSettings)
+                                   const FMiscSettings& MiscSettings, const FPVFoliageDistributionConditions& DistributionConditions)
 {
 	PV::Facades::FFoliageFacade FoliageFacadeOutput(OutCollection);
-
-	if (const int32 NumFoliageMeshes = FoliageFacadeOutput.NumFoliageNames(); NumFoliageMeshes == 0)
+ 
+	if (const int32 NumFoliageMeshes = FoliageFacadeOutput.NumFoliageInfo(); NumFoliageMeshes == 0)
 	{
 		UE_LOG(LogProceduralVegetation, Warning, TEXT("There are no foliage meshes available in input for distribution"));
 		return;
 	}
-
+ 
 	FRandomStream RandomStream;
 	RandomStream.Initialize(MiscSettings.RandomSeed);
-
+ 
 	PV::Facades::FBranchFacade BranchFacadeSource(SourceCollection);
 	PV::Facades::FPointFacade PointFacadeSource(SourceCollection);
-
+ 
 	TManagedArrayAccessor<TArray<FVector3f>> BudDirectionAttribute(OutCollection, PV::AttributeNames::BudDirection,
 		PV::GroupNames::PointGroup);
+	
 	TManagedArrayAccessor<float> BranchGradientsAttribute(OutCollection, PV::AttributeNames::BranchGradient,
 		PV::GroupNames::PointGroup);
+	
 	TManagedArrayAccessor<TArray<float>> LeafPhyllotaxyAttribute(OutCollection, PV::AttributeNames::LeafPhyllotaxy,
 		PV::GroupNames::DetailsGroup);
+	
 	TManagedArrayAccessor<TArray<float>> BudHormoneLevels(OutCollection, PV::AttributeNames::BudHormoneLevels,
 		PV::GroupNames::PointGroup);
-
+ 
+	TManagedArrayAccessor<TArray<int>> BudDevelopment(OutCollection, PV::AttributeNames::BudDevelopment,
+		PV::GroupNames::PointGroup);
+ 
 	float PhyllotaxyFormationLeaf = LeafPhyllotaxyAttribute[0][1];
 	float MinBudsLeaf = LeafPhyllotaxyAttribute[0][3];
 	float MaxBudsLeaf = LeafPhyllotaxyAttribute[0][4];
 	int32 ResetPhyllotaxyLeaf = static_cast<int>(LeafPhyllotaxyAttribute[0][6]);
 	float PhyllotaxyOffsetLeaf = LeafPhyllotaxyAttribute[0][7];
-
+ 
 	if (PhyllotaxySettings.OverridePhyllotaxy)
 	{
 		switch (PhyllotaxySettings.PhyllotaxyType)
@@ -77,19 +142,87 @@ void FPVFoliage::DistributeFoliage(FManagedArrayCollection& OutCollection, const
 				case EPhyllotaxyFormation::Parastichous: PhyllotaxyFormationLeaf = PhyllotaxySettings.PhyllotaxyAdditionalAngle;
 					break;
 				}
-
+ 
 				MinBudsLeaf = MaxBudsLeaf = 1;
 				break;
 			}
 		}
 	}
-
+ 
 	const int32 NumBranches = BranchFacadeSource.GetElementCount();
-
+ 
 	if (DistributionSettings.OverrideDistribution)
 	{
 		TRACE_CPUPROFILER_EVENT_SCOPE(PVFoliageDistributorNode::OverrideDistribution);
-
+ 
+		TArray<float> NormalisedPScales;
+		NormalisedPScales.Init(0, PointFacadeSource.GetElementCount());
+ 
+		TArray<float> NormalisedBudLightDetected;
+		NormalisedBudLightDetected.Init(0, PointFacadeSource.GetElementCount());
+ 
+		TArray<float> NormalisedHealth;
+		NormalisedHealth.Init(0, PointFacadeSource.GetElementCount());
+ 
+		TArray<PV::Facades::FFoliageEntryData> FoliageEntries;
+ 
+		// TODO : Optimize
+		// if (DistributionConditions.IsActiveCondition(EFoliageDistributionCondition::Scale) || DistributionConditions.IsActiveCondition(EFoliageDistributionCondition::Light))
+		// {
+		float MinPScale = std::numeric_limits<float>::max();
+		float MaxPScale = std::numeric_limits<float>::min();
+			
+		float MinBudLightDetected = std::numeric_limits<float>::max();
+		float MaxBudLightDetected = std::numeric_limits<float>::min();
+ 
+		float MinHealth = std::numeric_limits<float>::max();
+		float MaxHealth = std::numeric_limits<float>::min();
+ 
+		for (int PointIndex = 0; PointIndex < PointFacadeSource.GetElementCount(); PointIndex++)
+		{
+			// Scale
+			float PScale = PointFacadeSource.GetPointScale(PointIndex);
+			MinPScale = FMath::Min(MinPScale, PScale);	
+			MaxPScale = FMath::Max(MaxPScale, PScale);
+ 
+			// Light
+			const TArray<float>& CurrentBudLightDetected = PointFacadeSource.GetBudLightDetected(PointIndex);
+			const float CurrentLightAvailable = CurrentBudLightDetected.IsValidIndex(PV::Facades::BudLightDetectedAvailableIndex) ? CurrentBudLightDetected[PV::Facades::BudLightDetectedAvailableIndex] : 0.0f;
+				
+			MinBudLightDetected = FMath::Min(MinBudLightDetected, CurrentLightAvailable);	
+			MaxBudLightDetected = FMath::Max(MaxBudLightDetected, CurrentLightAvailable);
+ 
+			// Health
+			float Ethylene = BudHormoneLevels[PointIndex][4];
+			int BudAgeSenescence = BudDevelopment[PointIndex][3];
+			int BudLightSenescence = BudDevelopment[PointIndex][4];
+			
+			int Senescence = FMath::Max(BudAgeSenescence,BudLightSenescence) + 1;
+			NormalisedHealth[PointIndex] = 1 - (Ethylene * (1/Senescence));
+			
+			MinHealth = FMath::Min(MinHealth, NormalisedHealth[PointIndex]);
+			MaxHealth = FMath::Max(MaxHealth, NormalisedHealth[PointIndex]);
+		}
+ 
+		for (int PointIndex = 0; PointIndex < PointFacadeSource.GetElementCount(); PointIndex++)
+		{
+			// Scale
+			float PScale = PointFacadeSource.GetPointScale(PointIndex);
+			NormalisedPScales[PointIndex] = FMath::GetMappedRangeValueClamped(FVector2f(MinPScale, MaxPScale),FVector2f(0, 1),PScale);
+ 
+			// UE_LOG(LogTemp, Log, TEXT("Normalized PScale : Index : %i : Value : %f"),
+			// 					PointIndex, NormalisedPScales[PointIndex]);
+ 
+			// Light
+			const TArray<float>& CurrentBudLightDetected = PointFacadeSource.GetBudLightDetected(PointIndex);
+			const float CurrentLightAvailable = CurrentBudLightDetected.IsValidIndex(PV::Facades::BudLightDetectedAvailableIndex) ? CurrentBudLightDetected[PV::Facades::BudLightDetectedAvailableIndex] : 0.0f;
+			NormalisedBudLightDetected[PointIndex] = FMath::GetMappedRangeValueClamped(FVector2f(MinBudLightDetected, MaxBudLightDetected),FVector2f(0, 1), CurrentLightAvailable);
+ 
+			// Health
+			NormalisedHealth[PointIndex] = FMath::GetMappedRangeValueClamped(FVector2f(MinHealth, MaxHealth),FVector2f(0, 1),NormalisedHealth[PointIndex]);
+		}
+		// }
+		
 		// Compute attachment points in terms of normalized distances
 		OutCollection.EmptyGroup(PV::GroupNames::FoliageGroup);
  
@@ -113,6 +246,7 @@ void FPVFoliage::DistributeFoliage(FManagedArrayCollection& OutCollection, const
 		
 		for (int32 BranchIndex = 0; BranchIndex < NumBranches; ++BranchIndex)
 		{
+			ConditionAttributes ConditionAttributes;
 			TArray<float> NormalizedAttachmentPoints;
 			TArray<float> NormalizedAttachmentPointsFinal;
 			TArray<FVector3f> AttachmentPointsPositions;
@@ -120,8 +254,11 @@ void FPVFoliage::DistributeFoliage(FManagedArrayCollection& OutCollection, const
 			TArray<FVector3f> AttachmentPointsUpVectors;
 			TArray<FVector3f> AttachmentPointsNormalVectors;
 			TArray<float> AttachmentPointsScales;
+			
+			const int32 NumFoliageInfo = FoliageFacadeOutput.NumFoliageInfo();
+			
 			TArray<int32> BranchPoints = BranchFacadeSource.GetPoints(BranchIndex);
-
+ 
 			if (BranchPoints.Num() < 2)
 			{
 				continue;
@@ -135,124 +272,172 @@ void FPVFoliage::DistributeFoliage(FManagedArrayCollection& OutCollection, const
 				? BranchPoints.Num()
 				: DistributionSettings.MaxPerBranch;
 			float LookUp = 0.0f;
-
+ 
 			for (int32 i = 0; i < AdjustedMaxPerBranch; ++i)
 			{
 				if (LookUp > 1.0f)
 				{
 					break;
 				}
-
+ 
 				float InstanceSpacingRampValue = DistributionSettings.InstanceSpacingRamp
 					? DistributionSettings.InstanceSpacingRamp->GetRichCurveConst()->Eval(1.0f - LookUp)
 					: Increment;
 				// The increment is nudged slightly before blending with the weighted ramp value, the number is chosen based on observation
 				float AdjustedIncrement = Increment + (InstanceSpacingRampValue * (Increment * 10.0f));
 				AdjustedIncrement = FMath::Lerp(Increment, AdjustedIncrement, DistributionSettings.InstanceSpacingRampEffect);
-
+ 
 				NormalizedAttachmentPoints.Add(1.0f - LookUp);
 				LookUp += AdjustedIncrement;
 			}
-
+ 
 			// Compute Position, Angles, Scale and LengthFromRoot for attachment points
-
+ 
 			// Compute based on first point attributes
 			const FVector3f InitialApicalDirection(BudDirectionAttribute[BranchPoints[0]][0]);
 			FVector3f InitialAxillaryDirection(BudDirectionAttribute[BranchPoints[0]][1]);
 			const FVector3f InitialUpVector(BudDirectionAttribute[BranchPoints[0]][5]);
-
+ 
 			if (ResetPhyllotaxyLeaf == 1)
 			{
 				InitialAxillaryDirection = FVector3f::CrossProduct(InitialApicalDirection, InitialUpVector);
 			}
-
+ 
 			if (PhyllotaxyOffsetLeaf > 0.01f)
 			{
 				const FQuat4f RotationQuat(InitialApicalDirection.GetSafeNormal(), FMath::DegreesToRadians(PhyllotaxyOffsetLeaf));
 				InitialAxillaryDirection = RotationQuat.RotateVector(InitialAxillaryDirection);
 			}
-
+ 
 			int32 AxillaryRotationIteration = 0;
 			for (int32 Idx = 0; Idx < NormalizedAttachmentPoints.Num(); ++Idx)
 			{
 				for (int32 BranchPointIndex = 0; BranchPointIndex < BranchPoints.Num() - 1; ++BranchPointIndex)
 				{
 					float AxillaryRotationDegree = FMath::DegreesToRadians(PhyllotaxyFormationLeaf * AxillaryRotationIteration);
-
+ 
 					const int32 CurrentPointIndex = BranchPoints[BranchPointIndex];
 					const int32 NextPointIndex = BranchPoints[BranchPointIndex + 1];
+ 
 					const FVector3f CurrentPointPosition = PointFacadeSource.GetPosition(CurrentPointIndex);
 					const FVector3f NextPointPosition = PointFacadeSource.GetPosition(NextPointIndex);
+ 
 					float EthyleneLevelAtCurrentPoint = BudHormoneLevels[CurrentPointIndex][4];
 					float EthyleneLevelAtNextPoint = BudHormoneLevels[NextPointIndex][4];
+ 
 					const float CurrentPointLengthFromRoot = PointFacadeSource.GetLengthFromRoot(CurrentPointIndex);
 					const float NextPointLengthFromRoot = PointFacadeSource.GetLengthFromRoot(NextPointIndex);
+					
 					const float CurrentPointLfrNormalized = FMath::GetMappedRangeValueClamped(
 						FVector2f(FirstPointLengthFromRoot, LastPointLengthFromRoot),
 						FVector2f(0.0f, 1.0f), CurrentPointLengthFromRoot);
+ 
 					const float NextPointLfrNormalized = FMath::GetMappedRangeValueClamped(
 						FVector2f(FirstPointLengthFromRoot, LastPointLengthFromRoot),
 						FVector2f(0.0f, 1.0f),
 						NextPointLengthFromRoot);
-
+ 
 					if (NormalizedAttachmentPoints[Idx] >= CurrentPointLfrNormalized && NormalizedAttachmentPoints[Idx] <=
 						NextPointLfrNormalized)
 					{
 						const float DistanceAlpha = (NormalizedAttachmentPoints[Idx] -
 							CurrentPointLfrNormalized) / (NextPointLfrNormalized - CurrentPointLfrNormalized);
-
+ 
 						const float EthyleneLevel = FMath::Lerp(EthyleneLevelAtCurrentPoint, EthyleneLevelAtNextPoint, DistanceAlpha);
-
+ 
 						if ((EthyleneLevel - 0.001f) >= DistributionSettings.EthyleneThreshold)
 						{
 							continue;
 						}
-
+ 
 						const FVector3f Position = FMath::Lerp(CurrentPointPosition, NextPointPosition, DistanceAlpha);
 						AttachmentPointsPositions.Add(Position);
-
+ 
 						const float AttachmentPointLfr = FMath::GetMappedRangeValueClamped(
 							FVector2f(0.0f, 1.0f),
 							FVector2f(FirstPointLengthFromRoot, LastPointLengthFromRoot),
 							NormalizedAttachmentPoints[Idx]);
 						AttachmentPointsLengthsFromRoot.Add(AttachmentPointLfr);
-
+ 
 						// Compute Scale
 						const float CurrentPointBranchGradient = BranchGradientsAttribute[CurrentPointIndex];
 						const float NextPointBranchGradient = BranchGradientsAttribute[NextPointIndex];
-
+ 
 						const float BranchGradient = FMath::Clamp(
 							FMath::Lerp(CurrentPointBranchGradient, NextPointBranchGradient, DistanceAlpha), 0.0f, 1.0f);
 						float BranchGradientRampValue = ScaleSettings.ScaleRamp
 							? ScaleSettings.ScaleRamp->GetRichCurveConst()->Eval(1.0f - BranchGradient)
 							: 1.0f;
-
+ 
 						const float CurrentPointScale = PointFacadeSource.GetPointScale(CurrentPointIndex);
 						const float NextPointScale = PointFacadeSource.GetPointScale(NextPointIndex);
 						const float BranchScale = FMath::Lerp(CurrentPointScale, NextPointScale, DistanceAlpha);
-
+ 
 						const float BranchScaledValue = BranchGradientRampValue * (BranchScale + 0.1f);
 						BranchGradientRampValue = FMath::Lerp(BranchGradientRampValue, BranchScaledValue, ScaleSettings.BranchScaleImpact);
 						BranchGradientRampValue = FMath::Clamp(BranchGradientRampValue, 0.0f, 1.0f);
-
+ 
 						float Scale = FMath::GetMappedRangeValueClamped(
 							FVector2f(0.0f, 1.0f),
 							FVector2f(ScaleSettings.MinScale, ScaleSettings.MaxScale),
 							BranchGradientRampValue
 						) * ScaleSettings.BaseScale;
-
+ 
 						const float RandomValue = RandomStream.FRandRange(ScaleSettings.RandomScaleMin, ScaleSettings.RandomScaleMax);
 						Scale = Scale * RandomValue;
-
+ 
 						AttachmentPointsScales.Add(Scale);
-
+						
 						const FVector3f CurrentPointApicalDirectionVector(BudDirectionAttribute[CurrentPointIndex][0]);
 						const FVector3f NextPointApicalDirectionVector(BudDirectionAttribute[NextPointIndex][0]);
 						const FVector3f BudApicalDirectionVector = FMath::Lerp(CurrentPointApicalDirectionVector,
 							NextPointApicalDirectionVector, DistanceAlpha);
-
+ 
+						// Condition - Y Dot
+						float YDot = FVector3f::DotProduct(BudApicalDirectionVector, FVector3f::UpVector);
+						YDot = FMath::GetMappedRangeValueClamped(FVector2f(-1, 1),FVector2f(0, 1),YDot);
+						ConditionAttributes.UpAlignment.Add(YDot);
+ 
+						// Condition - Light
+						// const TArray<float>& CurrentBudLightDetected = PointFacadeSource.GetBudLightDetected(CurrentPointIndex);
+						// const float CurrentLightAvailable = CurrentBudLightDetected.IsValidIndex(PV::Facades::BudLightDetectedAvailableIndex) ? CurrentBudLightDetected[PV::Facades::BudLightDetectedAvailableIndex] : 0.0f;
+						//
+						// const TArray<float>& NextBudLightDetected = PointFacadeSource.GetBudLightDetected(NextPointIndex);
+						// const float NextLightAvailable = NextBudLightDetected.IsValidIndex(PV::Facades::BudLightDetectedAvailableIndex) ? NextBudLightDetected[PV::Facades::BudLightDetectedAvailableIndex] : 0.0f;	
+ 
+						const float CurrentLightAvailable = NormalisedBudLightDetected[CurrentPointIndex];
+						const float NextLightAvailable = NormalisedBudLightDetected[NextPointIndex];
+						const float BudLightReceived = FMath::Lerp(CurrentLightAvailable,NextLightAvailable, DistanceAlpha);
+						ConditionAttributes.Light.Add(BudLightReceived);
+ 
+						// Condition - Scale
+						// const float CurrentPointScaleGradient = PointFacadeSource.GetPointScaleGradient(CurrentPointIndex);
+						// const float NextPointScaleGradient = PointFacadeSource.GetPointScaleGradient(NextPointIndex);
+						// const float PointScaleGradient = FMath::Lerp(CurrentPointScaleGradient, NextPointScaleGradient, DistanceAlpha);
+ 
+						const float CurrentPScale = NormalisedPScales[CurrentPointIndex];
+						const float NextPScale = NormalisedPScales[NextPointIndex];
+						const float PointPScale = FMath::Lerp(CurrentPScale, NextPScale, DistanceAlpha);
+						ConditionAttributes.Scale.Add(PointPScale);
+ 
+						// Health
+						const float CurrentHealth = NormalisedHealth[CurrentPointIndex];
+						const float NextHealth = NormalisedHealth[NextPointIndex];
+						const float BudHealth = FMath::Lerp(CurrentHealth,NextHealth, DistanceAlpha);
+						ConditionAttributes.Health.Add(BudHealth);
+ 
+						// Condition - Tip
+						if (BranchPointIndex == BranchPoints.Num() - 2 && DistanceAlpha == 1)
+						{
+							ConditionAttributes.Tip.Add(1);
+						}
+						else
+						{
+							ConditionAttributes.Tip.Add(0);
+						}
+ 
 						int32 BudsToAdd = 1;
-
+ 
 						// Compute Up Vector
 						// Last point (order of the array is last to first)
 						if (Idx == 0)
@@ -260,74 +445,90 @@ void FPVFoliage::DistributeFoliage(FManagedArrayCollection& OutCollection, const
 							AttachmentPointsUpVectors.Add(BudApicalDirectionVector);
 						}
 						else // Points other than last
-						[[likely]]
-						{
-							const FQuat4f RotationQuat(InitialApicalDirection.GetSafeNormal(), AxillaryRotationDegree);
-							FVector3f NewAxillaryDirection = RotationQuat.RotateVector(InitialAxillaryDirection.GetSafeNormal());
-							FQuat4f ApicalCorrection = FQuat4f::FindBetweenNormals(
-								InitialApicalDirection.GetSafeNormal(), BudApicalDirectionVector.GetSafeNormal());
-							NewAxillaryDirection = ApicalCorrection.RotateVector(NewAxillaryDirection).GetSafeNormal();
-
-							int32 NumBuds = RandomStream.RandRange(MinBudsLeaf, MaxBudsLeaf);
-							float RotationPerBud = 360.0f / static_cast<float>(NumBuds);
-
-							for (int32 i = 0; i < NumBuds; ++i)
+							[[likely]]
 							{
-								float RotationAngle = FMath::DegreesToRadians(RotationPerBud * i);
-								const FQuat4f AxillaryRotation(InitialApicalDirection.GetSafeNormal(), RotationAngle);
-								FVector3f AxillaryDirectionModified = AxillaryRotation.RotateVector(NewAxillaryDirection.GetSafeNormal());
-
-								if (VectorSettings.OverrideAxilAngle)
-								{
-									float AxilAngleRampValue = VectorSettings.AxilAngleRamp
-										? FMath::Clamp(
-											VectorSettings.AxilAngleRamp->GetRichCurveConst()->Eval(1.0f - BranchGradient),
-											0.0f,
-											1.0f)
-										: 1.0f;
-
-									AxilAngleRampValue = AxilAngleRampValue * VectorSettings.AxilAngleRampUpperValue;
-									const float AxilAngleBlended = FMath::Lerp(VectorSettings.AxilAngle, AxilAngleRampValue,
-										VectorSettings.AxilAngleRampEffect);
-
-									const FVector3f RotationAxis = FVector3f::CrossProduct(
-										BudApicalDirectionVector.GetSafeNormal(), AxillaryDirectionModified.GetSafeNormal());
-
-									const FQuat4f Rotation(RotationAxis.GetSafeNormal(), FMath::DegreesToRadians(AxilAngleBlended));
-									const FVector3f AxilVector = Rotation.RotateVector(BudApicalDirectionVector.GetSafeNormal());
-
-									AttachmentPointsUpVectors.Add(AxilVector);
-								}
-								else
-								{
-									AttachmentPointsUpVectors.Add(AxillaryDirectionModified);
-								}
-
-								if (i > 0)
+								const FQuat4f RotationQuat(InitialApicalDirection.GetSafeNormal(), AxillaryRotationDegree);
+								FVector3f NewAxillaryDirection = RotationQuat.RotateVector(InitialAxillaryDirection.GetSafeNormal());
+								FQuat4f ApicalCorrection = FQuat4f::FindBetweenNormals(
+									InitialApicalDirection.GetSafeNormal(), BudApicalDirectionVector.GetSafeNormal());
+								NewAxillaryDirection = ApicalCorrection.RotateVector(NewAxillaryDirection).GetSafeNormal();
+ 
+								int32 NumBuds = RandomStream.RandRange(MinBudsLeaf, MaxBudsLeaf);
+								float RotationPerBud = 360.0f / static_cast<float>(NumBuds);
+ 
+								for (int32 i = 0; i < NumBuds; ++i)
 								{
-									FVector3f LastPositionAdded = AttachmentPointsPositions.Last();
-									float LastScaleAdded = AttachmentPointsScales.Last();
-									float LastLengthFromRootAdded = AttachmentPointsLengthsFromRoot.Last();
-									AttachmentPointsPositions.Add(LastPositionAdded);
-									AttachmentPointsScales.Add(LastScaleAdded);
-									AttachmentPointsLengthsFromRoot.Add(LastLengthFromRootAdded);
+									float RotationAngle = FMath::DegreesToRadians(RotationPerBud * i);
+									const FQuat4f AxillaryRotation(InitialApicalDirection.GetSafeNormal(), RotationAngle);
+									FVector3f AxillaryDirectionModified = AxillaryRotation.RotateVector(NewAxillaryDirection.GetSafeNormal());
+ 
+									if (VectorSettings.OverrideAxilAngle)
+									{
+										float AxilAngleRampValue = VectorSettings.AxilAngleRamp
+											? FMath::Clamp(
+												VectorSettings.AxilAngleRamp->GetRichCurveConst()->Eval(1.0f - BranchGradient),
+												0.0f,
+												1.0f)
+											: 1.0f;
+ 
+										AxilAngleRampValue = AxilAngleRampValue * VectorSettings.AxilAngleRampUpperValue;
+										const float AxilAngleBlended = FMath::Lerp(VectorSettings.AxilAngle, AxilAngleRampValue,
+											VectorSettings.AxilAngleRampEffect);
+ 
+										const FVector3f RotationAxis = FVector3f::CrossProduct(
+											BudApicalDirectionVector.GetSafeNormal(), AxillaryDirectionModified.GetSafeNormal());
+ 
+										const FQuat4f Rotation(RotationAxis.GetSafeNormal(), FMath::DegreesToRadians(AxilAngleBlended));
+										const FVector3f AxilVector = Rotation.RotateVector(BudApicalDirectionVector.GetSafeNormal());
+ 
+										AttachmentPointsUpVectors.Add(AxilVector);
+									}
+									else
+									{
+										AttachmentPointsUpVectors.Add(AxillaryDirectionModified);
+									}
+ 
+									if (i > 0)
+									{
+										FVector3f LastPositionAdded = AttachmentPointsPositions.Last();
+										float LastScaleAdded = AttachmentPointsScales.Last();
+										float LastLengthFromRootAdded = AttachmentPointsLengthsFromRoot.Last();
+										AttachmentPointsPositions.Add(LastPositionAdded);
+										AttachmentPointsScales.Add(LastScaleAdded);
+										AttachmentPointsLengthsFromRoot.Add(LastLengthFromRootAdded);
+ 
+										// Conditions
+										float LastConditionScale = ConditionAttributes.Scale.Last();
+										ConditionAttributes.Scale.Add(LastConditionScale);
+									
+										float LastConditionLight = ConditionAttributes.Light.Last();
+										ConditionAttributes.Light.Add(LastConditionLight);
+									
+										float LastConditionUpAlignment  = ConditionAttributes.UpAlignment.Last();
+										ConditionAttributes.UpAlignment.Add(LastConditionUpAlignment);
+ 
+										float LastConditionTip = ConditionAttributes.Tip.Last();
+										ConditionAttributes.Tip.Add(LastConditionTip);
+ 
+										float LastConditionHealth = ConditionAttributes.Health.Last();
+										ConditionAttributes.Health.Add(LastConditionHealth);
+									}
 								}
+ 
+								BudsToAdd = NumBuds;
+								AxillaryRotationIteration++;
 							}
-
-							BudsToAdd = NumBuds;
-							AxillaryRotationIteration++;
-						}
-
+ 
 						// Compute Normal vectors based on light optimal direction
 						const FVector3f CurrentPointLightOptimalVector(BudDirectionAttribute[CurrentPointIndex][2]);
 						const FVector3f NextPointLightOptimalVector(BudDirectionAttribute[NextPointIndex][2]);
 						const FVector3f BudLightOptimalVector = FMath::Lerp(CurrentPointLightOptimalVector, NextPointLightOptimalVector,
 							DistanceAlpha);
-
+ 
 						FQuat4f QuatRotation = FQuat4f::FindBetweenNormals(BudApicalDirectionVector.GetSafeNormal(),
 							BudLightOptimalVector.GetSafeNormal());
 						const FVector3f ApicalDirectionRotated = QuatRotation.RotateVector(BudApicalDirectionVector);
-
+ 
 						for (int32 i = 0; i < BudsToAdd; ++i)
 						{
 							AttachmentPointsNormalVectors.Add(ApicalDirectionRotated);
@@ -335,28 +536,145 @@ void FPVFoliage::DistributeFoliage(FManagedArrayCollection& OutCollection, const
 					}
 				}
 			}
-
+ 
 			// Set foliage entries
 			TManagedArrayAccessor<FVector3f> PivotPointsAttribute(OutCollection, PV::AttributeNames::FoliagePivotPoint,
 				PV::GroupNames::FoliageGroup);
-			const int32 NumFoliageMeshes = NonEmptyFoliageMeshNameIndices.Num();
-			int32 CurrentFoliageMeshIndex = BranchIndex % NumFoliageMeshes;
+ 
+			const TArray<FPVFoliageInfo> FoliageInfos = FoliageFacadeOutput.GetFoliageInfos();
+			
+ 
+			const int32 BranchNumber = BranchFacadeSource.GetBranchNumber(BranchIndex);
+ 
+			FRandomStream FoliageRandomStream;
+			FoliageRandomStream.Initialize(BranchNumber);
+			
 			for (int32 Idx = 0; Idx < AttachmentPointsLengthsFromRoot.Num(); ++Idx)
 			{
+				TArray<FFoliageWeightInfo> AvailableFoliageWeights;
+				
+				for (int i = 0; i < NumFoliageInfo; ++i)
+				{
+					auto FoliageInfo = FoliageInfos[i];
+ 
+					if (FoliageInfo.IsValid())
+					{
+						AvailableFoliageWeights.Add(FFoliageWeightInfo(i, 0));
+					}
+				}
+ 
+				// Shuffle foliage
+				for (int i = AvailableFoliageWeights.Num() - 1; i > 0; --i)
+				{
+					int Index = RandomStream.RandRange(0, i);
+					AvailableFoliageWeights.Swap(i, Index);
+				}
+ 
+				if(DistributionConditions.HasActiveCondition())
+				{
+					for (EFoliageDistributionCondition Condition : TEnumRange<EFoliageDistributionCondition>())
+					{
+						FPVFoliageDistributionConditionSettings OutConditionSettings;
+						
+						if (DistributionConditions.IsActiveCondition(Condition) && DistributionConditions.GetSettings(Condition, OutConditionSettings))
+						{
+							float PointAttributeValue = ConditionAttributes.GetConditionValue(Condition,Idx);
+					
+							const float AttributeValue = PointAttributeValue + OutConditionSettings.Offset;
+ 
+							for (int Index = 0; Index < AvailableFoliageWeights.Num(); ++Index)
+							{
+								int FoliageIndex = AvailableFoliageWeights[Index].FoliageIndex;
+ 
+								if (FoliageInfos.IsValidIndex(FoliageIndex))
+								{
+									auto FoliageInfo = FoliageInfos[FoliageIndex];
+									float FoliageValue = FoliageInfo.Attributes.GetAttributeValue(Condition);
+									float FoliageDistribution = FMath::Abs(FoliageValue - AttributeValue) * OutConditionSettings.Weight;
+									
+									// UE_LOG(LogProceduralVegetation, Log, TEXT("FoliageWeight : Index : %i : FoliageValue : %f : AttributeValue : %f : Weight : %f : Final Value: %f"),
+									// 	FoliageIndex,
+									// 	FoliageValue,
+									// 	AttributeValue,
+									// 	OutConditionSettings.Weight,
+									// 	FoliageDistribution);	
+ 
+									AvailableFoliageWeights[Index].Weight += FoliageDistribution;
+								}
+								else
+								{
+									check(false);
+								}
+							}	
+						}
+					}
+ 
+					float MinFoliageWeight = Algo::MinElementBy(AvailableFoliageWeights, [](const FFoliageWeightInfo& Info) { return Info.Weight; })->Weight;
+					float MaxFoliageWeight = Algo::MaxElementBy(AvailableFoliageWeights, [](const FFoliageWeightInfo& Info) { return Info.Weight; })->Weight;
+ 
+					for (FFoliageWeightInfo& Info : AvailableFoliageWeights)
+					{
+						Info.Weight = FMath::GetMappedRangeValueClamped(
+									FVector2f(MinFoliageWeight, MaxFoliageWeight),
+									FVector2f(0, 1),
+									Info.Weight);
+					}
+				}
+ 
+				// Pick foliage
+				AvailableFoliageWeights.Sort([&](const FFoliageWeightInfo& A, const FFoliageWeightInfo& B)
+					{
+						return A.Weight < B.Weight; // Ascending
+					});
+ 
+				
+				TArray<int> PickedFoliage;
+ 
+				
+				for (const FFoliageWeightInfo& Info : AvailableFoliageWeights)
+				{
+					if (Info.Weight < DistributionConditions.CutoffThreshold)
+					{
+						PickedFoliage.Add(Info.FoliageIndex);
+					}
+					else if (PickedFoliage.Num() < DistributionConditions.MinimumCandidates)
+					{
+						PickedFoliage.Add(Info.FoliageIndex);
+					}
+					else
+					{
+						break;
+					}
+				}
+ 
+				const int PickedIndex = RandomStream.RandRange(0, PickedFoliage.Num() - 1);
+				int PickedFoliageIndex = PickedFoliage[PickedIndex];
+ 
 				const int32 NewIndex = PivotPointsAttribute.AddElements(1);
 				FoliageFacadeOutput.SetFoliageEntry(NewIndex, {
-					.NameId = NonEmptyFoliageMeshNameIndices[CurrentFoliageMeshIndex],
+					.NameId = PickedFoliageIndex,
 					.BranchId = BranchIndex,
 					.PivotPoint = AttachmentPointsPositions[Idx],
 					.UpVector = AttachmentPointsUpVectors[Idx],
 					.NormalVector = AttachmentPointsNormalVectors[Idx],
 					.Scale = AttachmentPointsScales[Idx],
-					.LengthFromRoot = AttachmentPointsLengthsFromRoot[Idx]
+					.LengthFromRoot = AttachmentPointsLengthsFromRoot[Idx],
+					.ConditionUpAlignment = ConditionAttributes.UpAlignment[Idx],
+					.ConditionTip = ConditionAttributes.Tip[Idx],
+					.ConditionLight = ConditionAttributes.Light[Idx],
+					.ConditionScale = ConditionAttributes.Scale[Idx],
+					.ConditionHealth = ConditionAttributes.Health[Idx]
 				});
-				CurrentFoliageMeshIndex = (CurrentFoliageMeshIndex + 1) % NumFoliageMeshes;
 			}
 		}
-
+ 
+		// TArray<float> Scales;
+		//
+		// for (int32 i = 0; i < FoliageFacadeOutput.NumFoliageEntries(); ++i)
+		// {
+		// 	FoliageFacadeOutput.GetFoliageEntry(i).Scale;
+		// }
+		
 		// Re-calculate Foliage IDs for each branch
 		TMap<int32, TArray<int32>> BranchIndexToFoliageIDs;
 		for (int32 i = 0; i < FoliageFacadeOutput.NumFoliageEntries(); ++i)
@@ -372,29 +690,29 @@ void FPVFoliage::DistributeFoliage(FManagedArrayCollection& OutCollection, const
 			FoliageFacadeOutput.SetFoliageIdsArray(BranchId, FoliageIDs);
 		}
 	}
-
+ 
 	{
 		TRACE_CPUPROFILER_EVENT_SCOPE(PVFoliageDistributorNode::EthyleneRemove);
-
+ 
 		// No need to apply Ethylene threshold if the distribution was overriden since that's already been accounted for
 		if (DistributionSettings.OverrideDistribution)
 		{
 			return;
 		}
-
+ 
 		// Compute Foliage instances to remove based on Ethylene threshold
 		TManagedArrayAccessor<FVector3f> PivotPointsAttribute(
 			OutCollection,
 			PV::AttributeNames::FoliagePivotPoint,
 			PV::GroupNames::FoliageGroup);
 		TSet<int32> FoliageEntriesToRemove;
-
+ 
 		for (int32 BranchIndex = 0; BranchIndex < NumBranches; ++BranchIndex)
 		{
 			TArray<int32> BranchPoints = BranchFacadeSource.GetPoints(BranchIndex);
 			TArray<int32> FoliageIDs = FoliageFacadeOutput.GetFoliageEntryIdsForBranch(BranchIndex);
 			TSet<int32> FoliageIDsToRemove;
-
+ 
 			for (int32 BranchPointIndex = 0; BranchPointIndex < BranchPoints.Num() - 1; ++BranchPointIndex)
 			{
 				for (int32 FoliageIDsIndex = 0; FoliageIDsIndex < FoliageIDs.Num(); ++FoliageIDsIndex)
@@ -402,17 +720,17 @@ void FPVFoliage::DistributeFoliage(FManagedArrayCollection& OutCollection, const
 					const float CurrentPointLengthFromRoot = PointFacadeSource.GetLengthFromRoot(BranchPoints[BranchPointIndex]);
 					const float NextPointLengthFromRoot = PointFacadeSource.GetLengthFromRoot(BranchPoints[BranchPointIndex + 1]);
 					const float FoliageLengthFromRoot = FoliageFacadeOutput.GetFoliageEntry(FoliageIDs[FoliageIDsIndex]).LengthFromRoot;
-
+ 
 					if (FoliageLengthFromRoot >= CurrentPointLengthFromRoot && FoliageLengthFromRoot <= NextPointLengthFromRoot)
 					{
 						float EthyleneLevelAtCurrentPoint = BudHormoneLevels[BranchPoints[BranchPointIndex]][4];
 						float EthyleneLevelAtNextPoint = BudHormoneLevels[BranchPoints[BranchPointIndex + 1]][4];
-
+ 
 						float EthyleneLevel = FMath::GetMappedRangeValueClamped(
 							FVector2f(CurrentPointLengthFromRoot, NextPointLengthFromRoot),
 							FVector2f(EthyleneLevelAtNextPoint, EthyleneLevelAtCurrentPoint),
 							FoliageLengthFromRoot);
-
+ 
 						if ((EthyleneLevel - 0.001f) >= DistributionSettings.EthyleneThreshold)
 						{
 							FoliageIDsToRemove.Add(FoliageIDs[FoliageIDsIndex]);
@@ -420,13 +738,13 @@ void FPVFoliage::DistributeFoliage(FManagedArrayCollection& OutCollection, const
 					}
 				}
 			}
-
+ 
 			for (int32 Fid : FoliageIDsToRemove)
 			{
 				FoliageEntriesToRemove.Add(Fid);
 			}
 		}
-
+ 
 		// Remove foliage instances
 		for (int32 i = PivotPointsAttribute.Num() - 1; i >= 0; --i)
 		{
@@ -435,7 +753,7 @@ void FPVFoliage::DistributeFoliage(FManagedArrayCollection& OutCollection, const
 				PivotPointsAttribute.RemoveElements(1, i);
 			}
 		}
-
+ 
 		// Re-calculate Foliage IDs for each branch
 		TMap<int32, TArray<int32>> BranchIndexToFoliageIDs;
 		for (int32 i = 0; i < FoliageFacadeOutput.NumFoliageEntries(); ++i)
@@ -451,4 +769,4 @@ void FPVFoliage::DistributeFoliage(FManagedArrayCollection& OutCollection, const
 			FoliageFacadeOutput.SetFoliageIdsArray(BranchId, FoliageIDs);
 		}
 	}
-}
+}
\ No newline at end of file
diff --git a/Engine/Plugins/Experimental/ProceduralVegetationEditor/Source/ProceduralVegetation/Private/Implementations/PVFoliage.h b/Engine/Plugins/Experimental/ProceduralVegetationEditor/Source/ProceduralVegetation/Private/Implementations/PVFoliage.h
--- a/Engine/Plugins/Experimental/ProceduralVegetationEditor/Source/ProceduralVegetation/Private/Implementations/PVFoliage.h
+++ b/Engine/Plugins/Experimental/ProceduralVegetationEditor/Source/ProceduralVegetation/Private/Implementations/PVFoliage.h
@@ -1,17 +1,16 @@
 // Copyright Epic Games, Inc. All Rights Reserved.
-
+ 
 #pragma once
-
+ 
 #include "CoreMinimal.h"
-
-#include "Curves/CurveFloat.h"
-
+#include "Misc/EnumRange.h"
+ 
 #include "GeometryCollection/ManagedArrayCollection.h"
-
+ 
 #include "PVFoliage.generated.h"
-
+ 
 struct FPVFloatRamp;
-
+ 
 UENUM(BlueprintType)
 enum class EPhyllotaxyType : uint8
 {
@@ -21,7 +20,7 @@ enum class EPhyllotaxyType : uint8
 	Whorled UMETA(DisplayName = "Whorled"),
 	Spiral UMETA(DisplayName = "Spiral")
 };
-
+ 
 UENUM(BlueprintType)
 enum class EPhyllotaxyFormation : uint8
 {
@@ -31,7 +30,7 @@ enum class EPhyllotaxyFormation : uint8
 	Octastichous UMETA(DisplayName = "Octastichous"),
 	Parastichous UMETA(DisplayName = "Parastichous")
 };
-
+ 
 struct FDistributionSettings
 {
 	const float EthyleneThreshold;
@@ -41,7 +40,7 @@ struct FDistributionSettings
 	const float InstanceSpacingRampEffect;
 	const int32 MaxPerBranch;
 };
-
+ 
 struct FScaleSettings
 {
 	const float BaseScale;
@@ -52,7 +51,7 @@ struct FScaleSettings
 	const float RandomScaleMax;
 	const FPVFloatRamp* ScaleRamp;
 };
-
+ 
 struct FVectorSettings
 {
 	const bool OverrideAxilAngle;
@@ -61,7 +60,7 @@ struct FVectorSettings
 	const float AxilAngleRampUpperValue;
 	const float AxilAngleRampEffect;
 };
-
+ 
 struct FPhyllotaxySettings
 {
 	const bool OverridePhyllotaxy;
@@ -71,16 +70,289 @@ struct FPhyllotaxySettings
 	const int32 MaximumNodeBuds;
 	const float PhyllotaxyAdditionalAngle;
 };
-
+ 
+USTRUCT()
+struct FPVFoliageDistributionConditionSettings
+{
+	GENERATED_BODY()
+ 
+	UPROPERTY(EditAnywhere, Category="ConditionSettings", meta=(ClampMin=0.0f, ClampMax=1.0f, UIMin=0.0f, UIMax=1.0f, Tooltip=""))
+	float Weight = 1.0f;
+ 
+	UPROPERTY(EditAnywhere, Category="ConditionSettings", meta=(ClampMin=-1.0f, ClampMax=1.0f, UIMin=-1.0f, UIMax=1.0f, Tooltip=""))
+	float Offset = 0.0f;
+ 
+	FString ToString() const
+	{
+		return FString::Printf(TEXT("Weight = %f, Offset = %f"), Weight, Offset);
+	}
+};
+ 
+UENUM()
+enum class EFoliageDistributionCondition : uint8
+{
+	Light,
+	Scale,
+	UpAlignment,
+	Tip,
+	Health,
+	Count UMETA(Hidden),
+	None UMETA(Hidden),
+};
+ 
+ENUM_RANGE_BY_COUNT(EFoliageDistributionCondition, EFoliageDistributionCondition::Count)
+ 
+ 
+USTRUCT()
+struct FPVFoliageDistributionConditions
+{
+	GENERATED_BODY()
+ 
+	UPROPERTY(EditAnywhere, Category="ConditionSettings", meta=(ClampMin=0.0f, ClampMax=1.0f, UIMin=0.0f, UIMax=1.0f, Tooltip=""))
+	float CutoffThreshold = 0.3f;
+ 
+	UPROPERTY(EditAnywhere, Category="ConditionSettings", meta=(ClampMin=1, ClampMax=10, UIMin=0, UIMax=10, Tooltip=""))
+	int32 MinimumCandidates = 1;
+	
+	UPROPERTY(EditAnywhere, Category="ConditionSettings", meta=(PinHiddenByDefault, InlineEditConditionToggle))
+	bool bActivateLight = false;
+ 
+	UPROPERTY(EditAnywhere, Category="ConditionSettings", meta=(EditCondition="bActivateLight", Tooltip=""))
+	FPVFoliageDistributionConditionSettings Light;
+ 
+	UPROPERTY()
+	bool bActivateScale = false;
+ 
+	UPROPERTY()
+	FPVFoliageDistributionConditionSettings Scale;
+ 
+	UPROPERTY(EditAnywhere, Category="ConditionSettings", meta=(PinHiddenByDefault, InlineEditConditionToggle))
+	bool bActivateUpAlignment = false;
+ 
+	UPROPERTY(EditAnywhere, Category="ConditionSettings", meta=(EditCondition="bActivateUpAlignment", Tooltip=""))
+	FPVFoliageDistributionConditionSettings UpAlignment;
+ 
+	UPROPERTY(EditAnywhere, Category="ConditionSettings", meta=(PinHiddenByDefault, InlineEditConditionToggle))
+	bool bActivateTip = false;
+	
+	UPROPERTY(EditAnywhere, Category="ConditionSettings", meta=(EditCondition="bActivateTip", Tooltip=""))
+	FPVFoliageDistributionConditionSettings Tip;
+ 
+	UPROPERTY(EditAnywhere, Category="ConditionSettings", meta=(PinHiddenByDefault, InlineEditConditionToggle))
+	bool bActivateHealth = false;
+	
+	UPROPERTY(EditAnywhere, Category="ConditionSettings", meta=(EditCondition="bActivateHealth", Tooltip=""))
+	FPVFoliageDistributionConditionSettings Health;
+ 
+	bool HasActiveCondition() const
+	{
+		for (EFoliageDistributionCondition Condition : TEnumRange<EFoliageDistributionCondition>())
+		{
+			if (IsActiveCondition(Condition))
+			{
+				return true;
+			}
+		}
+		return false;
+	}
+ 
+	bool IsActiveCondition(const EFoliageDistributionCondition InCondition) const 
+	{
+		switch (InCondition)
+		{
+		case EFoliageDistributionCondition::Light:
+			return bActivateLight && Light.Weight > 0.0f;
+			
+		case EFoliageDistributionCondition::Scale:
+			return bActivateScale && Scale.Weight > 0.0f; 
+			
+		case EFoliageDistributionCondition::UpAlignment:
+			return bActivateUpAlignment && UpAlignment.Weight > 0.0f;
+			
+		case EFoliageDistributionCondition::Tip:
+			return bActivateTip && Tip.Weight > 0.0f;
+ 
+		case EFoliageDistributionCondition::Health:
+			return bActivateHealth && Health.Weight > 0.0f;
+ 
+		default:
+			return false;
+		}
+	}
+ 
+	void SetConditionState(const EFoliageDistributionCondition InCondition, const bool InNewState) 
+	{
+		switch (InCondition)
+		{
+		case EFoliageDistributionCondition::Light:
+			bActivateLight = InNewState;
+			break;
+			
+			
+		case EFoliageDistributionCondition::Scale:
+			bActivateScale = InNewState;
+			break;
+			
+		case EFoliageDistributionCondition::UpAlignment:
+			bActivateUpAlignment = InNewState;
+			break;
+			
+		case EFoliageDistributionCondition::Tip:
+			bActivateTip = InNewState;
+			break;
+ 
+		case EFoliageDistributionCondition::Health:
+			bActivateHealth = InNewState;
+			break;
+ 
+			default:
+			check(false);
+		}
+	}
+ 
+	FPVFoliageDistributionConditionSettings* GetSettings(const EFoliageDistributionCondition InCondition)
+	{
+		switch (InCondition)
+		{
+			case EFoliageDistributionCondition::Light:
+				return &Light;
+				
+			case EFoliageDistributionCondition::Scale:
+				return &Scale;
+			
+				
+			case EFoliageDistributionCondition::UpAlignment:
+				return &UpAlignment;
+				
+			case EFoliageDistributionCondition::Tip:
+				return &Tip;
+ 
+			case EFoliageDistributionCondition::Health:
+				return &Health;
+ 
+			default:
+			return nullptr;
+		}
+	}
+ 
+	bool GetSettings(const EFoliageDistributionCondition InCondition, FPVFoliageDistributionConditionSettings& OutSettings) const
+	{
+		switch (InCondition)
+		{
+			case EFoliageDistributionCondition::Light:
+			{
+				OutSettings =  Light;
+				return true;
+			}
+			
+			case EFoliageDistributionCondition::Scale:
+			{
+				OutSettings = Scale;
+				return true;
+			}
+			
+			case EFoliageDistributionCondition::UpAlignment:
+			{
+				OutSettings = UpAlignment;
+				return true;
+			}
+			
+			case EFoliageDistributionCondition::Tip:
+			{
+				OutSettings = Tip;
+				return true;
+			}
+ 
+			case EFoliageDistributionCondition::Health:
+			{
+				OutSettings = Health;
+				return true;
+			}
+		}
+ 
+		return false;
+	}
+};
+ 
 struct FMiscSettings
 {
 	const int32 RandomSeed;
 };
-
+ 
+ 
+USTRUCT()
+struct FPVFoliageAttributes
+{
+	GENERATED_BODY()
+ 
+	UPROPERTY(EditAnywhere, Category="Info", meta=(ClampMin=0.0f, ClampMax=1.0f, UIMin=0.0f, UIMax=1.0f, Tooltip=""))
+	float Light = 0.0f;
+ 
+	UPROPERTY()
+	float Scale = 0.0f;
+ 
+	UPROPERTY(EditAnywhere, Category="Info", meta=(ClampMin=0.0f, ClampMax=1.0f, UIMin=0.0f, UIMax=1.0f, Tooltip=""))
+	float UpAlignment = 0.0f;
+ 
+	UPROPERTY(EditAnywhere, Category="Info", meta=(ClampMin=0.0f, ClampMax=1.0f, UIMin=0.0f, UIMax=1.0f, Tooltip=""))
+	float Health = 0.0f;
+ 
+	UPROPERTY(EditAnywhere, Category="Info", meta=(Tooltip=""))
+	bool Tip = false;
+ 
+	float GetAttributeValue(const EFoliageDistributionCondition InCondition) const
+	{
+		switch (InCondition)
+		{
+			case EFoliageDistributionCondition::Light:
+			return Light;
+			
+			case EFoliageDistributionCondition::Scale:
+			return Scale;
+			
+			case EFoliageDistributionCondition::UpAlignment:
+			return UpAlignment;
+			
+			case EFoliageDistributionCondition::Tip:
+			return Tip ? 1.0f : 0.0f;
+ 
+			case EFoliageDistributionCondition::Health:
+			return Health;
+ 
+			default:
+			return 0.0f;
+		}
+	}
+ 
+	FString ToString() const
+	{
+		return FString::Printf(TEXT("Light = %f, Scale = %f, UpAlignment = %f, Health = %f, Tip = %s"), Light, Scale, UpAlignment, Health, Tip ? TEXT("True") : TEXT("False"));
+	}
+};
+ 
+USTRUCT()
+struct FPVFoliageInfo
+{
+	GENERATED_BODY()
+	
+	UPROPERTY(EditAnywhere, Category = "Info")
+	TSoftObjectPtr<UObject> Mesh;
+ 
+	UPROPERTY(EditAnywhere, Category="Info", meta=(ShowInnerProperties, FullyExpand="true", Tooltip=""))
+	FPVFoliageAttributes Attributes;
+ 
+	bool IsValid() const
+	{
+		// Asset is loaded and pointer is valid
+		// Path is valid, but asset is NOT loaded yet
+		return Mesh.IsValid() || Mesh.ToSoftObjectPath().IsValid(); 
+	}
+};
+ 
 struct FPVFoliage
 {
 	static void DistributeFoliage(FManagedArrayCollection& OutCollection, const FManagedArrayCollection& SourceCollection,
 	                              const FDistributionSettings& DistributionSettings, const FScaleSettings& ScaleSettings,
 	                              const FVectorSettings& VectorSettings, const FPhyllotaxySettings& PhyllotaxySettings,
-	                              const FMiscSettings& MiscSettings);
-};
+	                              const FMiscSettings& MiscSettings, const FPVFoliageDistributionConditions& DistributionConditions);
+};
\ No newline at end of file
diff --git a/Engine/Plugins/Experimental/ProceduralVegetationEditor/Source/ProceduralVegetation/Private/Nodes/PVFoliageDistributorSettings.cpp b/Engine/Plugins/Experimental/ProceduralVegetationEditor/Source/ProceduralVegetation/Private/Nodes/PVFoliageDistributorSettings.cpp
--- a/Engine/Plugins/Experimental/ProceduralVegetationEditor/Source/ProceduralVegetation/Private/Nodes/PVFoliageDistributorSettings.cpp
+++ b/Engine/Plugins/Experimental/ProceduralVegetationEditor/Source/ProceduralVegetation/Private/Nodes/PVFoliageDistributorSettings.cpp
@@ -1,22 +1,23 @@
 ﻿// Copyright Epic Games, Inc. All Rights Reserved.
-
+ 
 #include "PVFoliageDistributorSettings.h"
 #include "DataTypes/PVData.h"
 #include "DataTypes/PVMeshData.h"
 #include "DataTypes/PVFoliageMeshData.h"
 #include "Data/PCGBasePointData.h"
+#include "Facades/PVFoliageConditionFacade.h"
 #include "GeometryCollection/ManagedArrayCollection.h"
 #include "Helpers/PCGHelpers.h"
 #include "Helpers/PCGSettingsHelpers.h"
-
+ 
 #define LOCTEXT_NAMESPACE "PVFoliageDistributorSettings"
-
+ 
 #if WITH_EDITOR
 FText UPVFoliageDistributorSettings::GetDefaultNodeTitle() const 
 { 
 	return LOCTEXT("NodeTitle", "Foliage Distributor"); 
 }
-
+ 
 FText UPVFoliageDistributorSettings::GetNodeTooltipText() const 
 { 
 	return LOCTEXT("NodeTooltip", 
@@ -28,8 +29,26 @@ FText UPVFoliageDistributorSettings::GetNodeTooltipText() const
 		"\n\nPress Ctrl + L to lock/unlock node output"
 	); 
 }
+ 
+void UPVFoliageDistributorSettings::PostEditChangeProperty(FPropertyChangedEvent& PropertyChangedEvent)
+{
+	Super::PostEditChangeProperty(PropertyChangedEvent);
+ 
+	if (PropertyChangedEvent.Property
+		&& PropertyChangedEvent.Property->GetFName() == GET_MEMBER_NAME_CHECKED(UPVFoliageDistributorSettings, OverrideDistribution))
+	{
+		if (OverrideDistribution)
+		{
+			Conditions = CustomConditions;
+		}
+		else
+		{
+			Conditions = DefaultConditions;
+		}	
+	}
+}
 #endif
-
+ 
 UPVFoliageDistributorSettings::UPVFoliageDistributorSettings()
 {
 	if (PCGHelpers::IsNewObjectAndNotDefault(this))
@@ -40,14 +59,14 @@ UPVFoliageDistributorSettings::UPVFoliageDistributorSettings()
 		const FKeyHandle Handle1 = RichCurve->AddKey(1.0f, 1.0f);
 		RichCurve->SetKeyInterpMode(Handle0, RCIM_Linear);
 		RichCurve->SetKeyInterpMode(Handle1, RCIM_Linear);
-
+ 
 		RichCurve = ScaleRamp.GetRichCurve();
 		RichCurve->Reset();
 		const FKeyHandle Handle2 = RichCurve->AddKey(0.0f, 0.0f);
 		const FKeyHandle Handle3 = RichCurve->AddKey(1.0f, 1.0f);
 		RichCurve->SetKeyInterpMode(Handle2, RCIM_Linear);
 		RichCurve->SetKeyInterpMode(Handle3, RCIM_Linear);
-
+ 
 		RichCurve = AxilAngleRamp.GetRichCurve();
 		RichCurve->Reset();
 		const FKeyHandle Handle4 = RichCurve->AddKey(0.0f, 1.0f);
@@ -56,31 +75,34 @@ UPVFoliageDistributorSettings::UPVFoliageDistributorSettings()
 		RichCurve->SetKeyInterpMode(Handle5, RCIM_Linear);
 	}
 }
-
+ 
 FPCGDataTypeIdentifier UPVFoliageDistributorSettings::GetInputPinTypeIdentifier() const
 {
 	return FPCGDataTypeIdentifier{ FPVDataTypeInfoMesh::AsId() };
 }
-
+ 
 FPCGDataTypeIdentifier UPVFoliageDistributorSettings::GetOutputPinTypeIdentifier() const
 {
 	return FPCGDataTypeIdentifier{ FPVDataTypeInfoFoliageMesh::AsId() };
 }
-
+ 
 FPCGElementPtr UPVFoliageDistributorSettings::CreateElement() const
 {
 	return MakeShared<FPVFoliageDistributorElement>();
 }
-
+ 
 bool FPVFoliageDistributorElement::ExecuteInternal(FPCGContext* InContext) const
 {
 	TRACE_CPUPROFILER_EVENT_SCOPE(FPVFoliageDistributorElement::Execute);
-
+ 
 	check(InContext);
-
+ 
 	const UPVFoliageDistributorSettings* Settings = InContext->GetInputSettings<UPVFoliageDistributorSettings>();
 	check(Settings);
-
+ 
+	UPVFoliageDistributorSettings* MutableSettings = const_cast<UPVFoliageDistributorSettings*>(Settings);
+	check(MutableSettings);
+ 
 	for (const FPCGTaggedData& Input : InContext->InputData.GetInputsByPin(PCGPinConstants::DefaultInputLabel))
 	{
 		if(const UPVMeshData* InputData = Cast<UPVMeshData>(Input.Data))
@@ -88,7 +110,7 @@ bool FPVFoliageDistributorElement::ExecuteInternal(FPCGContext* InContext) const
 			FManagedArrayCollection Collection = InputData->GetCollection();
 			
 			UPVFoliageMeshData* OutManagedArrayCollectionData = FPCGContext::NewObject_AnyThread<UPVFoliageMeshData>(InContext);
-
+ 
 			const FDistributionSettings DistributionSettings = {
 				.EthyleneThreshold = Settings->EthyleneThreshold,
 				.OverrideDistribution = Settings->OverrideDistribution,
@@ -97,7 +119,7 @@ bool FPVFoliageDistributorElement::ExecuteInternal(FPCGContext* InContext) const
 				.InstanceSpacingRampEffect = Settings->InstanceSpacingRampEffect,
 				.MaxPerBranch = Settings->MaxPerBranch
 			};
-
+ 
 			const FScaleSettings ScaleSettings = {
 				.BaseScale = Settings->BaseScale,
 				.BranchScaleImpact = Settings->BranchScaleImpact,
@@ -107,7 +129,7 @@ bool FPVFoliageDistributorElement::ExecuteInternal(FPCGContext* InContext) const
 				.RandomScaleMax = Settings->RandomScaleMax,
 				.ScaleRamp = &Settings->ScaleRamp
 			};
-
+ 
 			const FVectorSettings VectorSettings = {
 				.OverrideAxilAngle = Settings->OverrideAxilAngle,
 				.AxilAngle = Settings->AxilAngle,
@@ -115,7 +137,7 @@ bool FPVFoliageDistributorElement::ExecuteInternal(FPCGContext* InContext) const
 				.AxilAngleRampUpperValue = Settings->AxilAngleRampUpperValue,
 				.AxilAngleRampEffect = Settings->AxilAngleRampEffect
 			};
-
+ 
 			const FPhyllotaxySettings PhyllotaxySettings = {
 				.OverridePhyllotaxy = Settings->OverridePhyllotaxy,
 				.PhyllotaxyType = Settings->PhyllotaxyType,
@@ -124,14 +146,53 @@ bool FPVFoliageDistributorElement::ExecuteInternal(FPCGContext* InContext) const
 				.MaximumNodeBuds = Settings->MaximumNodeBuds,
 				.PhyllotaxyAdditionalAngle = Settings->PhyllotaxyAdditionalAngle
 			};
-
+ 
 			const FMiscSettings MiscSettings = {
 				.RandomSeed = Settings->RandomSeed
 			};
-
+ 
+			if (!Settings->OverrideDistribution)
+			{
+				PV::Facades::FFoliageConditionFacade FoliageConditionFacade(Collection);
+ 
+				auto DefaultConditionsData = FoliageConditionFacade.GetData();
+ 
+				for (auto ConditonInfo : DefaultConditionsData)
+				{
+					auto ConditionType = ConditonInfo.GetType();
+ 
+					if (ConditionType != EFoliageDistributionCondition::None)
+					{
+						
+						if (auto ConditionSettings = MutableSettings->DefaultConditions.GetSettings(ConditionType))
+						{
+							ConditionSettings->Weight = ConditonInfo.Weight;
+							ConditionSettings->Offset = ConditonInfo.Offset;
+ 
+							MutableSettings->DefaultConditions.SetConditionState(ConditionType, true);	
+						}
+					}
+				}
+ 
+				if (!MutableSettings->bCustomConditionSet)
+				{
+					MutableSettings->CustomConditions = MutableSettings->DefaultConditions;
+					MutableSettings->bCustomConditionSet = true;
+				}
+			}
+ 
+			if (!MutableSettings->OverrideDistribution)
+			{
+				MutableSettings->Conditions = MutableSettings->DefaultConditions;
+			}
+			else
+			{
+				MutableSettings->CustomConditions = MutableSettings->Conditions;
+			}	
+ 
 			FPVFoliage::DistributeFoliage(Collection, Collection, DistributionSettings, ScaleSettings, VectorSettings,
-												 PhyllotaxySettings, MiscSettings);
-
+												 PhyllotaxySettings, MiscSettings, Settings->Conditions);
+ 
 			OutManagedArrayCollectionData->Initialize(MoveTemp(Collection));
 			InContext->OutputData.TaggedData.Emplace(OutManagedArrayCollectionData);
 		}
@@ -144,5 +205,5 @@ bool FPVFoliageDistributorElement::ExecuteInternal(FPCGContext* InContext) const
 	
 	return true;
 }
-
+ 
 #undef LOCTEXT_NAMESPACE
\ No newline at end of file
diff --git a/Engine/Plugins/Experimental/ProceduralVegetationEditor/Source/ProceduralVegetation/Private/Nodes/PVFoliageDistributorSettings.h b/Engine/Plugins/Experimental/ProceduralVegetationEditor/Source/ProceduralVegetation/Private/Nodes/PVFoliageDistributorSettings.h
--- a/Engine/Plugins/Experimental/ProceduralVegetationEditor/Source/ProceduralVegetation/Private/Nodes/PVFoliageDistributorSettings.h
+++ b/Engine/Plugins/Experimental/ProceduralVegetationEditor/Source/ProceduralVegetation/Private/Nodes/PVFoliageDistributorSettings.h
@@ -29,6 +29,7 @@ public:
 	virtual FLinearColor GetNodeTitleColor() const override { return FLinearColor::Blue; }
 
 	virtual TArray<EPVRenderType> GetDefaultRenderType() const override { return TArray{ EPVRenderType::Mesh, EPVRenderType::Foliage }; }
+	virtual void PostEditChangeProperty(FPropertyChangedEvent& PropertyChangedEvent) override;
 	virtual TArray<EPVRenderType> GetSupportedRenderTypes() const override { return TArray{ EPVRenderType::Mesh, EPVRenderType::Foliage, EPVRenderType::FoliageAttachments }; }
 	#endif
 
@@ -45,7 +46,7 @@ private:
 	float EthyleneThreshold = 0.85f;
 
 	UPROPERTY(EditAnywhere, Category="Distribution Settings", meta=(Tooltip="Override default placement with custom distribution rules.\n\nOverides the distribution settings loaded from the procedural vegetation preset."))
-	bool OverrideDistribution = true;
+	bool OverrideDistribution = false;
 
 	UPROPERTY(EditAnywhere, Category="Distribution Settings", meta=(EditCondition="OverrideDistribution", ClampMin=0.0f, ClampMax=1.0f, UIMin=0.0f, UIMax=1.0f, Tooltip="Minimum distance between placed instances.\n\nHigher spacing yields a more open distribution; lower spacing allows denser placement."))
 	float InstanceSpacing = 0.3f;
@@ -115,6 +116,18 @@ private:
 
 	UPROPERTY(EditAnywhere, Category="Misc Settings", meta=(EditCondition="OverrideDistribution"))
 	int32 RandomSeed = 123456;
+
+	UPROPERTY(EditAnywhere, Category="ConditionSettings", meta = (ShowOnlyInnerProperties, EditCondition="OverrideDistribution"))
+	FPVFoliageDistributionConditions Conditions;
+	
+	UPROPERTY()
+	FPVFoliageDistributionConditions DefaultConditions;
+	
+	UPROPERTY()
+	FPVFoliageDistributionConditions CustomConditions;
+
+	UPROPERTY()
+	bool bCustomConditionSet = false;
 };
 
 class FPVFoliageDistributorElement : public FPVBaseElement
diff --git a/Engine/Plugins/Experimental/ProceduralVegetationEditor/Source/ProceduralVegetation/Private/Nodes/PVFoliagePaletteSettings.cpp b/Engine/Plugins/Experimental/ProceduralVegetationEditor/Source/ProceduralVegetation/Private/Nodes/PVFoliagePaletteSettings.cpp
--- a/Engine/Plugins/Experimental/ProceduralVegetationEditor/Source/ProceduralVegetation/Private/Nodes/PVFoliagePaletteSettings.cpp
+++ b/Engine/Plugins/Experimental/ProceduralVegetationEditor/Source/ProceduralVegetation/Private/Nodes/PVFoliagePaletteSettings.cpp
@@ -1,7 +1,7 @@
 ﻿// Copyright Epic Games, Inc. All Rights Reserved.
-
+ 
 #include "PVFoliagePaletteSettings.h"
-
+ 
 #include "ProceduralVegetationModule.h"
 #include "DataTypes/PVData.h"
 #include "DataTypes/PVMeshData.h"
@@ -12,15 +12,16 @@
 #include "Helpers/PCGHelpers.h"
 #include "Helpers/PCGSettingsHelpers.h"
 #include "Helpers/PVAnalyticsHelper.h"
-
+#include "Implementations/PVFoliage.h"
+ 
 #define LOCTEXT_NAMESPACE "PVFoliagePaletteSettings"
-
+ 
 #if WITH_EDITOR
 FText UPVFoliagePaletteSettings::GetDefaultNodeTitle() const 
 { 
 	return LOCTEXT("NodeTitle", "Foliage Palette");
 }
-
+ 
 FText UPVFoliagePaletteSettings::GetNodeTooltipText() const 
 { 
 	return LOCTEXT("NodeTooltip", 
@@ -30,133 +31,159 @@ FText UPVFoliagePaletteSettings::GetNodeTooltipText() const
 		"\n\nPress Ctrl + L to lock/unlock node output"
 	);
 }
-
+ 
 void UPVFoliagePaletteSettings::PostEditChangeProperty(FPropertyChangedEvent& PropertyChangedEvent)
 {
 	Super::PostEditChangeProperty(PropertyChangedEvent);
-
+ 
 	if (PropertyChangedEvent.Property
 		&& PropertyChangedEvent.Property->GetFName() == GET_MEMBER_NAME_CHECKED(UPVFoliagePaletteSettings, bOverrideFoliage))
 	{
 		DirtyCache();
 	}
 }
-
+ 
 void UPVFoliagePaletteSettings::PostEditChangeChainProperty(FPropertyChangedChainEvent& PropertyChangedEvent)
 {
 	Super::PostEditChangeChainProperty(PropertyChangedEvent);
-
+ 
 	FProperty* Property = PropertyChangedEvent.Property;
 	FProperty* MemberProperty = nullptr;
-
+ 
 	if (PropertyChangedEvent.PropertyChain.GetActiveMemberNode())
 	{
 		MemberProperty = PropertyChangedEvent.PropertyChain.GetActiveMemberNode()->GetValue();
 	}
-
+ 
 	if (MemberProperty && Property)
 	{
-		if (MemberProperty->GetFName() == GET_MEMBER_NAME_CHECKED(UPVFoliagePaletteSettings, FoliageMeshes))
+		if (MemberProperty->GetFName() == GET_MEMBER_NAME_CHECKED(UPVFoliagePaletteSettings, FoliageInfos))
 		{
 			FName Name = Property->GetFName();
 			
 			int Index = PropertyChangedEvent.GetArrayIndex(Name.ToString());
-
-			if (Index >= 0 && FoliageMeshes.Num() > Index)
+ 
+			if (Index >= 0 && FoliageInfos.Num() > Index)
 			{
-				if (!FoliageMeshes[Index].IsNull())
+				if (!FoliageInfos[Index].Mesh.IsNull())
 				{
-					FString MeshPath = FoliageMeshes[Index].ToString();
-
+					FString MeshPath = FoliageInfos[Index].Mesh.ToString();
+ 
 					PV::Analytics::SendFoliageMeshChangeEvent(MeshPath);	
 				}
 			}
+ 
+			// For backwards compatibility 
+			FoliageMeshes.Empty();
+ 
+			for (int i = 0; i < FoliageInfos.Num(); i++)
+			{
+				FoliageMeshes.Add(FoliageInfos[i].Mesh);
+			}
 		}
 	}
 }
-
+ 
 #endif
-
+ 
 FPCGDataTypeIdentifier UPVFoliagePaletteSettings::GetInputPinTypeIdentifier() const
 {
 	return FPCGDataTypeIdentifier{ FPVDataTypeInfoMesh::AsId() };
 }
-
+ 
 FPCGDataTypeIdentifier UPVFoliagePaletteSettings::GetOutputPinTypeIdentifier() const
 {
 	return FPCGDataTypeIdentifier{ FPVDataTypeInfoMesh::AsId() };
 }
-
+ 
 FPCGElementPtr UPVFoliagePaletteSettings::CreateElement() const
 {
 	return MakeShared<FPVFoliageElement>();
 }
-
+ 
+void UPVFoliagePaletteSettings::PostLoad()
+{
+	Super::PostLoad();
+ 
+	if (bOverrideFoliage)
+	{
+		if (FoliageInfos.Num() == 0)
+		{
+			for (int32 i = 0; i < FoliageMeshes.Num(); i++)
+			{
+				FoliageInfos.Emplace(FoliageMeshes[i]);
+			}
+		}
+	}
+}
+ 
 bool FPVFoliageElement::ExecuteInternal(FPCGContext* InContext) const
 {
 	TRACE_CPUPROFILER_EVENT_SCOPE(FPVFoliageElement::Execute);
-
+ 
 	check(InContext);
-
+ 
 	const UPVFoliagePaletteSettings* Settings = InContext->GetInputSettings<UPVFoliagePaletteSettings>();
 	check(Settings);
-
+ 
 	UPVFoliagePaletteSettings* const MutableSettings = const_cast<UPVFoliagePaletteSettings*>(Settings);
 	check(MutableSettings);
-
+ 
 	const TArray<FPCGTaggedData>& Inputs = InContext->InputData.GetInputsByPin(PCGPinConstants::DefaultInputLabel);
 	if (!Inputs.IsEmpty())
 	{
 		if (const UPVMeshData* InputData = Cast<UPVMeshData>(Inputs[0].Data))
 		{
 			FManagedArrayCollection Collection = InputData->GetCollection();
-
+ 
 			UPVMeshData* OutManagedArrayCollectionData = FPCGContext::NewObject_AnyThread<UPVMeshData>(InContext);
-
+ 
 			PV::Facades::FFoliageFacade Facade(Collection);
-
+ 
 			if (!Settings->bOverrideFoliage)
 			{
-				MutableSettings->FoliageMeshes.Empty();
+				MutableSettings->FoliageInfos.Empty();
 				MutableSettings->bOverrideFoliage = false;
-				for (int32 FoliageIndex = 0; FoliageIndex < Facade.NumFoliageNames(); ++FoliageIndex)
+				for (int32 FoliageIndex = 0; FoliageIndex < Facade.NumFoliageInfo(); ++FoliageIndex)
 				{
-					MutableSettings->FoliageMeshes.Emplace(FSoftObjectPath(Facade.GetFoliageName(FoliageIndex)));
+					FPVFoliageInfo FoliageInfo = Facade.GetFoliageInfo(FoliageIndex);
+					MutableSettings->FoliageInfos.Emplace(FoliageInfo);
 				}
 #if WITH_EDITOR
 				MutableSettings->PostEditChange();
 #endif
 			}
-
-			const int32 IterationLimit = FMath::Min(Settings->FoliageMeshes.Num(), Facade.NumFoliageNames());
-			TArray<FString> FoliageNames;
+ 
+			const int32 IterationLimit = FMath::Min(Settings->FoliageInfos.Num(), Facade.NumFoliageInfo());
+			TArray<FPVFoliageInfo> FoliageInfos;
+			
 			for (int32 i = 0; i < IterationLimit; ++i)
 			{
-				FoliageNames.Emplace(Settings->FoliageMeshes[i].ToString());
+				FoliageInfos.Emplace(Settings->FoliageInfos[i]);
 			}
-
-			if (Settings->FoliageMeshes.Num() > Facade.NumFoliageNames())
+ 
+			if (Settings->FoliageInfos.Num() > Facade.NumFoliageInfo())
 			{
 				UE_LOG(LogProceduralVegetation, Warning,
 				       TEXT("Num of Foliage meshes specified is greater than the number of available entries, ignoring the last %d entries"),
-				       Settings->FoliageMeshes.Num() - Facade.NumFoliageNames());
+				       Settings->FoliageInfos.Num() - Facade.NumFoliageInfo());
 			}
-			else if (Settings->FoliageMeshes.Num() < Facade.NumFoliageNames())
+			else if (Settings->FoliageInfos.Num() < Facade.NumFoliageInfo())
 			{
-				if (Settings->FoliageMeshes.Num() > 0)
+				if (Settings->FoliageInfos.Num() > 0)
 				{
 					UE_LOG(LogProceduralVegetation, Warning,
 					       TEXT("Num of Foliage meshes specified is less than the number of available entries, repeating the last entry"));
-
-					for (int32 i = Settings->FoliageMeshes.Num(); i < Facade.NumFoliageNames(); ++i)
+ 
+					for (int32 i = Settings->FoliageInfos.Num(); i < Facade.NumFoliageInfo(); ++i)
 					{
-						FoliageNames.Emplace(Settings->FoliageMeshes.Last().ToString());
+						FoliageInfos.Emplace(Settings->FoliageInfos.Last());
 					}
 				}
 			}
-
-			Facade.SetFoliageNames(FoliageNames);
-
+ 
+			Facade.SetFoliageInfos(FoliageInfos);
+ 
 			OutManagedArrayCollectionData->Initialize(MoveTemp(Collection));
 			InContext->OutputData.TaggedData.Emplace(OutManagedArrayCollectionData);
 		}
@@ -166,8 +193,8 @@ bool FPVFoliageElement::ExecuteInternal(FPCGContext* InContext) const
 			return true;
 		}
 	}
-
+ 
 	return true;
 }
-
+ 
 #undef LOCTEXT_NAMESPACE
\ No newline at end of file
diff --git a/Engine/Plugins/Experimental/ProceduralVegetationEditor/Source/ProceduralVegetation/Private/Nodes/PVFoliagePaletteSettings.h b/Engine/Plugins/Experimental/ProceduralVegetationEditor/Source/ProceduralVegetation/Private/Nodes/PVFoliagePaletteSettings.h
--- a/Engine/Plugins/Experimental/ProceduralVegetationEditor/Source/ProceduralVegetation/Private/Nodes/PVFoliagePaletteSettings.h
+++ b/Engine/Plugins/Experimental/ProceduralVegetationEditor/Source/ProceduralVegetation/Private/Nodes/PVFoliagePaletteSettings.h
@@ -1,51 +1,55 @@
 ﻿// Copyright Epic Games, Inc. All Rights Reserved.
 #pragma once
-
+ 
 #include "CoreMinimal.h"
 #include "Nodes/PVBaseSettings.h"
 #include "DataTypes/PVData.h"
 #include "PVFoliagePaletteSettings.generated.h"
-
+ 
 UCLASS(BlueprintType, ClassGroup = (Procedural))
 class UPVFoliagePaletteSettings : public UPVBaseSettings
 {
 	GENERATED_BODY()
-
+ 
 public:
 #if WITH_EDITOR
 	virtual FName GetDefaultNodeName() const override { return FName(TEXT("ProceduralVegetationFoliage")); }
 	virtual FText GetDefaultNodeTitle() const override;
 	virtual FText GetNodeTooltipText() const override;
 	virtual FLinearColor GetNodeTitleColor() const override { return FLinearColor::Green; }
-
+ 
 	virtual TArray<EPVRenderType> GetDefaultRenderType() const override { return TArray{EPVRenderType::FoliageGrid}; }
 	virtual TArray<EPVRenderType> GetSupportedRenderTypes() const override { return TArray{EPVRenderType::FoliageGrid, EPVRenderType::Foliage, EPVRenderType::FoliageAttachments, EPVRenderType::Mesh}; }
 	
 	virtual void PostEditChangeProperty(FPropertyChangedEvent& PropertyChangedEvent) override;
 	virtual void PostEditChangeChainProperty(FPropertyChangedChainEvent& PropertyChangedEvent) override;
 #endif
-
+ 
 protected:
 	virtual FPCGDataTypeIdentifier GetInputPinTypeIdentifier() const override;
 	virtual FPCGDataTypeIdentifier GetOutputPinTypeIdentifier() const override;
 	virtual FPCGElementPtr CreateElement() const override;
-
+ 
 public:
+	virtual void PostLoad() override;
+ 
 	UPROPERTY(EditAnywhere, Category="Meshes", meta=(PinHiddenByDefault, InlineEditConditionToggle))
 	bool bOverrideFoliage = false;
-
+ 
+	UPROPERTY()
+	TArray<TSoftObjectPtr<UObject>> FoliageMeshes;
+	
 	UPROPERTY(EditAnywhere, Category = "Meshes", EditFixedSize,
-		Meta = (NoResetToDefault, EditCondition="bOverrideFoliage", AllowedClasses = "/Script/Engine.StaticMesh,/Script/Engine.SkeletalMesh", Tooltip=
+		Meta = (NoResetToDefault, EditCondition="bOverrideFoliage", AllowedClasses = "/Script/Engine.StaticMesh,/Script/Engine.SkeletalMesh", ShowOnlyInnerProperties, FullyExpand="true", Tooltip=
 			"List of Foliage meshes used for this plant\n\nList of all foliage meshes applied based on procedural vegetation preset."))
-	TArray<TSoftObjectPtr<UObject>> FoliageMeshes;
-
+	TArray<struct FPVFoliageInfo> FoliageInfos;
 };
-
+ 
 class FPVFoliageElement : public FPVBaseElement
 {
 protected:
 	virtual bool ExecuteInternal(FPCGContext* InContext) const override;
-
+ 
 public:
 	virtual bool CanExecuteOnlyOnMainThread(FPCGContext* Context) const override { return true; }
-};
+};
\ No newline at end of file
diff --git a/Engine/Plugins/Experimental/ProceduralVegetationEditor/Source/ProceduralVegetation/Private/Nodes/PVImporterSettings.cpp b/Engine/Plugins/Experimental/ProceduralVegetationEditor/Source/ProceduralVegetation/Private/Nodes/PVImporterSettings.cpp
--- a/Engine/Plugins/Experimental/ProceduralVegetationEditor/Source/ProceduralVegetation/Private/Nodes/PVImporterSettings.cpp
+++ b/Engine/Plugins/Experimental/ProceduralVegetationEditor/Source/ProceduralVegetation/Private/Nodes/PVImporterSettings.cpp
@@ -1,18 +1,19 @@
 ﻿// Copyright Epic Games, Inc. All Rights Reserved.
-
+ 
 #include "PVImporterSettings.h"
-
+ 
 #include "ProceduralVegetationModule.h"
-
+ 
 #include "DataTypes/PVData.h"
 #include "DataTypes/PVGrowthData.h"
-
+ 
 #include "GeometryCollection/ManagedArrayCollection.h"
-
+ 
 #include "Helpers/PCGHelpers.h"
 #include "Helpers/PCGSettingsHelpers.h"
+#include "Helpers/PVFoliageJSONHelper.h"
 #include "Helpers/PVJSONHelper.h"
-
+ 
 UPVImporterSettings::UPVImporterSettings()
 {
 #if WITH_EDITOR
@@ -22,28 +23,28 @@ UPVImporterSettings::UPVImporterSettings()
 	}
 #endif
 }
-
+ 
 TArray<FPCGPinProperties> UPVImporterSettings::InputPinProperties() const
 {
 	return TArray<FPCGPinProperties>();
 }
-
+ 
 FPCGElementPtr UPVImporterSettings::CreateElement() const
 {
 	return MakeShared<FPVImporterElement>();
 }
-
+ 
 FPCGDataTypeIdentifier UPVImporterSettings::GetOutputPinTypeIdentifier() const
 {
 	return FPCGDataTypeIdentifier{ FPVDataTypeInfoGrowth::AsId() };
 }
-
+ 
 bool FPVImporterElement::ExecuteInternal(FPCGContext* InContext) const
 {
 	TRACE_CPUPROFILER_EVENT_SCOPE(FPVImporterElement::Execute);
-
+ 
 	check(InContext);
-
+ 
 	const UPVImporterSettings* Settings = InContext->GetInputSettings<UPVImporterSettings>();
 	check(Settings);
 	
@@ -53,7 +54,7 @@ bool FPVImporterElement::ExecuteInternal(FPCGContext* InContext) const
 		SkeletonFilePath = FPaths::ProjectDir() / SkeletonFilePath;
 	}
 	FString ErrorMessage;
-
+ 
 	FManagedArrayCollection Collection;
 	bool bSuccess = PV::LoadMegaPlantsJsonToCollection(Collection, SkeletonFilePath, ErrorMessage);
 	if (!bSuccess)
@@ -61,14 +62,14 @@ bool FPVImporterElement::ExecuteInternal(FPCGContext* InContext) const
 		PCGLog::LogErrorOnGraph(FText::FromString(ErrorMessage), InContext);
 		return true;
 	}
-
+ 
 	if (!PV::Utilities::IsValidPVData(Collection))
 	{
 		PCGLog::LogErrorOnGraph(FText::FromString(FString::Format(TEXT("Invalid data in the imported json {0}"), { SkeletonFilePath})), InContext);
 		return true;
 	}
-
-	PV::SetFoliagePaths(Collection, SkeletonFilePath);
+ 
+	PV::PVFoliageJSONHelper::SetFoliagePaths(Collection, SkeletonFilePath);
 	
 	UPVGrowthData* OutManagedArrayCollectionData = FPCGContext::NewObject_AnyThread<UPVGrowthData>(InContext);
 	OutManagedArrayCollectionData->Initialize(MoveTemp(Collection));
@@ -76,4 +77,4 @@ bool FPVImporterElement::ExecuteInternal(FPCGContext* InContext) const
 	InContext->OutputData.TaggedData.Emplace(OutManagedArrayCollectionData);
 	
 	return true;
-}
+}
\ No newline at end of file
diff --git a/Engine/Plugins/Experimental/ProceduralVegetationEditor/Source/ProceduralVegetation/Private/ProceduralVegetationPreset.cpp b/Engine/Plugins/Experimental/ProceduralVegetationEditor/Source/ProceduralVegetation/Private/ProceduralVegetationPreset.cpp
--- a/Engine/Plugins/Experimental/ProceduralVegetationEditor/Source/ProceduralVegetation/Private/ProceduralVegetationPreset.cpp
+++ b/Engine/Plugins/Experimental/ProceduralVegetationEditor/Source/ProceduralVegetation/Private/ProceduralVegetationPreset.cpp
@@ -1,5 +1,5 @@
 ﻿// Copyright Epic Games, Inc. All Rights Reserved.
-
+ 
 #include "ProceduralVegetationPreset.h"
 #include "ProceduralVegetationModule.h"
 #include "Misc/PackageName.h"
@@ -7,22 +7,24 @@
 #include "AssetRegistry/IAssetRegistry.h"
 #include "Facades/PVBranchFacade.h"
 #include "Facades/PVProfileFacade.h"
+#include "Helpers/PVFoliageJSONHelper.h"
 #include "Helpers/PVJSONHelper.h"
 #include "Helpers/PVUtilities.h"
-
+#include "Implementations/PVFoliage.h"
+ 
 void FPVPresetVariationInfo::Fill(FName InName, FManagedArrayCollection Collection)
 {
 	Name = InName;
 	
 	const PV::Facades::FFoliageFacade Facade(Collection);
-	for (int32 FoliageIndex = 0; FoliageIndex < Facade.NumFoliageNames(); ++FoliageIndex)
+	for (int32 FoliageIndex = 0; FoliageIndex < Facade.NumFoliageInfo(); ++FoliageIndex)
 	{
-		FoliageMeshes.Emplace(FSoftObjectPath(Facade.GetFoliageName(FoliageIndex)));
+		FoliageMeshes.Emplace(Facade.GetFoliageInfo(FoliageIndex).Mesh);
 	}
-
+ 
 	const PV::Facades::FBranchFacade BranchFacade(Collection);
 	Materials.Emplace(FSoftObjectPath(BranchFacade.GetTrunkMaterialPath()));
-
+ 
 	if (const PV::Facades::FPlantProfileFacade PlantProfileFacade = PV::Facades::FPlantProfileFacade(Collection);
 			PlantProfileFacade.NumProfileEntries() > 0)
 	{
@@ -32,24 +34,24 @@ void FPVPresetVariationInfo::Fill(FName InName, FManagedArrayCollection Collecti
 		}
 	}
 }
-
+ 
 UProceduralVegetationPreset::UProceduralVegetationPreset()
 {
 	FString AssetPath = FPackageName::GetLongPackagePath(GetPackage()->GetName());
 	FoliageFolder.Path = FPaths::Combine(AssetPath, TEXT("Instances"));
 	MaterialsFolder.Path = FPaths::Combine(AssetPath, TEXT("Materials"));
 }
-
+ 
 void UProceduralVegetationPreset::PostLoad()
 {
 	Super::PostLoad();
-
+ 
 	if (PresetVariations.IsEmpty())
 	{
 		FillVariationInfo();
 	}
 }
-
+ 
 void UProceduralVegetationPreset::LoadFromVariantFiles(const TMap<FString, FString>& InVariantFiles)
 {
 	for (const auto& [VariantName, VariantFile] : InVariantFiles)
@@ -62,7 +64,7 @@ void UProceduralVegetationPreset::LoadFromVariantFiles(const TMap<FString, FStri
 		}
 	}
 }
-
+ 
 void UProceduralVegetationPreset::UpdateDataAsset()
 {
 	UE_LOG(LogProceduralVegetation, Log, TEXT("UpdateDataAsset Clicked"));
@@ -71,20 +73,41 @@ void UProceduralVegetationPreset::UpdateDataAsset()
 		UE_LOG(LogProceduralVegetation, Warning, TEXT("Please specify a JSON path"));
 		return;
 	}
-	
+	FString OutError;
 	IFileManager& FileManager = IFileManager::Get();
 	
 	TArray<FString> FileNames;
 	FileManager.FindFiles(FileNames, *JsonDirectoryPath.Path, TEXT("*.json"));
-
+ 
 	Variants.Empty();
+ 
+	FString JSONExtension = ".json";
+	FString FoliageDataFileName = "FoliageData" + JSONExtension;
+ 
+	if (!FileNames.Contains(FoliageDataFileName))
+	{
+		UE_LOG(LogProceduralVegetation, Warning, TEXT("Unable to find Foliage Data file"));
+		return;
+	}
+	else
+	{
+		FileNames.Remove(FoliageDataFileName);
+	}
+ 
+	FString FoliageDataPath = JsonDirectoryPath.Path / FoliageDataFileName;
+ 
+	PV::FoliageVariationsMap FoliageVariationsMap;
+ 
+	if (!PV::PVFoliageJSONHelper::LoadFoliageData(FoliageDataPath, FoliageFolder.Path, FoliageVariationsMap, OutError))
+	{
+		UE_LOG(LogProceduralVegetation, Warning, TEXT("Unable to load Foliage Data :  Error : %s"), *OutError);
+	}
 	
 	TArray<TSharedPtr<FJsonObject>> MetaFiles;
 	for (const auto& FileName : FileNames)
 	{
 		FString FullPath = JsonDirectoryPath.Path / FileName;
 		FManagedArrayCollection Collection;
-		FString OutError;
 		FString MetaJsonError;
 		if ( PV::LoadMegaPlantsJsonToCollection(Collection, FullPath, OutError))
 		{
@@ -118,6 +141,17 @@ void UProceduralVegetationPreset::UpdateDataAsset()
 		{
 			PV::LoadMetaJsonToCollection(Pair.Value, MetaFiles[0]);
 		}
+ 
+		// load Foliage Data again
+		if (FoliageVariationsMap.Contains(Pair.Key))
+		{
+			FString VariationPresetPath = JsonDirectoryPath.Path / Pair.Key + JSONExtension;
+			
+			if (!PV::PVFoliageJSONHelper::LoadFoliageDataInCollection(Pair.Value, VariationPresetPath, FoliageVariationsMap[Pair.Key], OutError))
+			{
+				UE_LOG(LogProceduralVegetation, Warning, TEXT("Unable to load foliage data in collection:  Error : %s"), *OutError);
+			}
+		}
 		
 		VariantIndex++;
 	}
@@ -131,7 +165,7 @@ void UProceduralVegetationPreset::UpdateDataAsset()
 	FillVariationInfo();
 	GetPackage()->SetDirtyFlag(true);
 }
-
+ 
 void UProceduralVegetationPreset::FillVariationInfo()
 {
 	PresetVariations.Empty();
@@ -142,7 +176,7 @@ void UProceduralVegetationPreset::FillVariationInfo()
 		PresetVariations.Emplace(Info);
 	}
 }
-
+ 
 void UProceduralVegetationPreset::ShowHideInternalProperties(bool bDebugEnable)
 {
 #if WITH_METADATA
@@ -151,7 +185,7 @@ void UProceduralVegetationPreset::ShowHideInternalProperties(bool bDebugEnable)
 	{
 		CategoryName = TEXT("Internal");
 	}
-
+ 
 	for (TFieldIterator<FProperty> It(StaticClass()); It; ++It)
 	{
 		FProperty* Property = *It;
@@ -160,7 +194,7 @@ void UProceduralVegetationPreset::ShowHideInternalProperties(bool bDebugEnable)
 			Property->SetMetaData(TEXT("Category"), *CategoryName);
 		}
 	}
-
+ 
 	for (TFieldIterator<UFunction> It(StaticClass()); It; ++It)
 	{
 		UFunction* Function = *It;
@@ -171,7 +205,7 @@ void UProceduralVegetationPreset::ShowHideInternalProperties(bool bDebugEnable)
 	}
 #endif
 }
-
+ 
 void UProceduralVegetationPreset::CreateProfileDataAsset()
 {
 	FString ProfilePackageName(PlantProfileName);
@@ -179,7 +213,7 @@ void UProceduralVegetationPreset::CreateProfileDataAsset()
 	UPackage* PlantProfilePackage = CreatePackage(*PackagePath);
 	
 	UPlantProfileAsset* NewPlantProfileAsset = NewObject<UPlantProfileAsset>(PlantProfilePackage, *ProfilePackageName, RF_Public | RF_Standalone);
-
+ 
 	for (auto& Pair : Variants)
 	{
 		PV::Facades::FPlantProfileFacade Facade = PV::Facades::FPlantProfileFacade(Pair.Value);
@@ -196,7 +230,7 @@ void UProceduralVegetationPreset::CreateProfileDataAsset()
 	IAssetRegistry::Get()->AssetCreated(NewPlantProfileAsset);
 	PlantProfilePackage->SetDirtyFlag(true);
 }
-
+ 
 void UProceduralVegetationPreset::UpdateFoliageAndMaterialPath()
 {
 	for (auto& Pair : Variants)
@@ -204,13 +238,8 @@ void UProceduralVegetationPreset::UpdateFoliageAndMaterialPath()
 		PV::Facades::FFoliageFacade FoliageFacade(Pair.Value);
 		PV::Facades::FBranchFacade BranchFacade(Pair.Value);
 			
-		for (int32 FoliageIndex = 0; FoliageIndex < FoliageFacade.NumFoliageNames(); FoliageIndex++)
+		for (int32 FoliageIndex = 0; FoliageIndex < FoliageFacade.NumFoliageInfo(); FoliageIndex++)
 		{
-			FString ShortName = FoliageFacade.GetFoliageName(FoliageIndex);
-			//FString ObjectName = FString::Format(TEXT("{0}.{1}"),{ ShortName, ShortName});
-			FString FullObjectName = FPaths::Combine(*FoliageFolder.Path, *ShortName);
-			FoliageFacade.SetFoliageName(FoliageIndex, FullObjectName);
-
 			if (!TrunkMaterialName.IsEmpty())
 			{
 				FString FullMaterialName = FPaths::Combine(*MaterialsFolder.Path, TrunkMaterialName + '.' + TrunkMaterialName);
@@ -222,4 +251,4 @@ void UProceduralVegetationPreset::UpdateFoliageAndMaterialPath()
 			}
 		}
 	}
-}
+}
\ No newline at end of file
diff --git a/Engine/Plugins/Experimental/ProceduralVegetationEditor/Source/ProceduralVegetation/Public/Facades/PVFoliageFacade.h b/Engine/Plugins/Experimental/ProceduralVegetationEditor/Source/ProceduralVegetation/Public/Facades/PVFoliageFacade.h
--- a/Engine/Plugins/Experimental/ProceduralVegetationEditor/Source/ProceduralVegetation/Public/Facades/PVFoliageFacade.h
+++ b/Engine/Plugins/Experimental/ProceduralVegetationEditor/Source/ProceduralVegetation/Public/Facades/PVFoliageFacade.h
@@ -1,11 +1,13 @@
 ﻿// Copyright Epic Games, Inc. All Rights Reserved.
-
+ 
 #pragma once
-
+ 
 #include "CoreMinimal.h"
 #include "IShrinkableFacade.h"
 #include "GeometryCollection/ManagedArrayAccessor.h"
-
+ 
+struct FPVFoliageInfo;
+ 
 namespace PV::Facades
 {
 	struct PROCEDURALVEGETATION_API FFoliageEntryData
@@ -17,9 +19,14 @@ namespace PV::Facades
 		FVector3f NormalVector;
 		float Scale;
 		float LengthFromRoot;
+		float ConditionUpAlignment;
+		float ConditionTip;
+		float ConditionLight;
+		float ConditionScale;
+		float ConditionHealth;
 		int32 ParentBoneID = INDEX_NONE;
 	};
-
+ 
 	/**
 	 * FFoliageFacade is used to access and manipulate Foliage related data for a ProceduralVegetation's FManagedArrayCollection
 	 * It presents access to foliage Mesh names, their respective Branch ids, and instancing transform data (Pivot point, Up Vector, and Scale)
@@ -31,82 +38,96 @@ namespace PV::Facades
 	public:
 		FFoliageFacade(FManagedArrayCollection& InCollection, int32 InitialSize = 0);
 		FFoliageFacade(const FManagedArrayCollection& InCollection);
-
+ 
 		bool IsConst() const { return Collection == nullptr; }
-
+ 
 		bool IsValid() const;
-
+ 
 		int32 NumFoliageEntries() const { return NameIdsAttribute.Num(); }
-
+ 
 		FFoliageEntryData GetFoliageEntry(const int32 Index) const;
-
+ 
 		float GetLengthFromRoot(int32 Index) const;
-
+ 
 		void SetLengthFromRoot(int32 Index, float Input);
-
+ 
 		int32 GetParentBoneID(int32 Index) const;
-
+ 
 		void SetParentBoneID(int32 Index, int32 Input);
-
+ 
 		const TArray<int32>& GetFoliageEntryIdsForBranch(int32 Index) const;
-
+ 
 		int32 AddFoliageEntry(const FFoliageEntryData& InputData);
-
+ 
 		void SetFoliageEntry(int32 Index, const FFoliageEntryData& InputData);
-
+ 
 		void SetFoliageIdsArray(int32 Index, TArray<int32> InputIds);
-
+ 
 		int32 GetFoliageBranchId(int32 Index) const;
-
+ 
+		int32 NumFoliageNames() const { return FoliageInfoNamesAttribute.Num(); };
 		void SetFoliageBranchId(int32 Index, int32 InputId);
-
-		int32 NumFoliageNames() const { return NamesAttribute.Num(); };
-
-		const FString& GetFoliageName(int32 Index) const;
-
+		
+		const FString GetFoliageName(int32 Index) const;
+ 
 		const TArray<FString>& GetFoliageNames() const;
-
+ 
+		int32 NumFoliageInfo() const { return FoliageInfoNamesAttribute.Num(); };
+ 
+		FPVFoliageInfo GetFoliageInfo(const int32 Index) const;
 		const FVector3f& GetPivotPoint(int32 Index) const;
-
+ 
+		TArray<FPVFoliageInfo> GetFoliageInfos() const;
+		
 		const FVector3f& GetUpVector(int32 Index) const;
-
+ 
 		const FVector3f& GetNormalVector(int32 Index) const;
-
+ 
 		void SetPivotPoint(int32 Index, const FVector3f& Input);
-
+ 
 		void SetUpVector(int32 Index, const FVector3f& Input);
-
+ 
 		void SetNormalVector(int32 Index, const FVector3f& Input);
-
+ 
 		void SetScale(int32 Index, float Input);
-
+ 
 		void SetFoliageName(int32 Index, const FString& Input);
-
+ 
 		void SetFoliageNames(const TArray<FString>& InputNames);
-
+ 
+		void SetFoliageInfo(const int32 Index, const FString& Input);
+ 
+		void SetFoliageInfos(const TArray<FPVFoliageInfo>& InputNames);
+ 
 		void SetFoliagePath(FString InPath);
-
+ 
 		FString GetFoliagePath() const;
-
+ 
 		virtual int32 GetElementCount() const override;
-
+ 
 		virtual void CopyEntry(int32 FromIndex, int32 ToIndex) override;
-
+ 
 		virtual void RemoveEntries(int32 NumEntries, int32 StartIndex) override;
 		
 		FTransform GetFoliageTransform(int32 Id) const;
-
+ 
 		const TManagedArray<FVector3f>& GetPivotPositions() const;
 		
 		TManagedArray<FVector3f>& ModifyPivotPositions();
-
+ 
 	protected:
 		void DefineSchema(int32 InitialSize = 0);
-
+ 
 		const FManagedArrayCollection& ConstCollection;
 		FManagedArrayCollection* Collection = nullptr;
-
-		TManagedArrayAccessor<FString> NamesAttribute;
+ 
+		TManagedArrayAccessor<FString> FoliageInfoNamesAttribute;
+		TManagedArrayAccessor<float> FoliageInfoLightAttribute;
+		TManagedArrayAccessor<float> FoliageInfoScaleAttribute;
+		TManagedArrayAccessor<float> FoliageInfoTipAttribute;
+		TManagedArrayAccessor<float> FoliageInfoUpAlignmentAttribute;
+		TManagedArrayAccessor<float> FoliageInfoHealthAttribute;
+		
 		TManagedArrayAccessor<int32> NameIdsAttribute;
 		TManagedArrayAccessor<int32> BranchIdsAttribute;
 		TManagedArrayAccessor<FVector3f> PivotPointsAttribute;
@@ -114,8 +135,13 @@ namespace PV::Facades
 		TManagedArrayAccessor<FVector3f> NormalVectorsAttribute;
 		TManagedArrayAccessor<float> ScalesAttribute;
 		TManagedArrayAccessor<float> LengthFromRootAttribute;
+		TManagedArrayAccessor<float> ConditionLightAttribute;
+		TManagedArrayAccessor<float> ConditionScaleAttribute;
+		TManagedArrayAccessor<float> ConditionTipAttribute;
+		TManagedArrayAccessor<float> ConditionUpAlignmentAttribute;
+		TManagedArrayAccessor<float> ConditionHealthAttribute;
 		TManagedArrayAccessor<int32> ParentBoneIdsAttribute;
 		TManagedArrayAccessor<TArray<int32>> FoliageIdsAttribute;
 		TManagedArrayAccessor<FString> FoliagePathAttribute;
 	};
-}
+}
\ No newline at end of file
diff --git a/Engine/Plugins/Experimental/ProceduralVegetationEditor/Source/ProceduralVegetationEditor/Private/Helpers/PVExportHelper.cpp b/Engine/Plugins/Experimental/ProceduralVegetationEditor/Source/ProceduralVegetationEditor/Private/Helpers/PVExportHelper.cpp
--- a/Engine/Plugins/Experimental/ProceduralVegetationEditor/Source/ProceduralVegetationEditor/Private/Helpers/PVExportHelper.cpp
+++ b/Engine/Plugins/Experimental/ProceduralVegetationEditor/Source/ProceduralVegetationEditor/Private/Helpers/PVExportHelper.cpp
@@ -1,7 +1,7 @@
 ﻿// Copyright Epic Games, Inc. All Rights Reserved.
-
+ 
 #include "Helpers/PVExportHelper.h"
-
+ 
 #include "DynamicWindImportData.h"
 #include "DynamicWindSkeletalData.h"
 #include "EditorSupportDelegates.h"
@@ -20,45 +20,46 @@
 #include "SkinnedAssetCompiler.h"
 #include "StaticMeshCompiler.h"
 #include "UDynamicMesh.h"
-
+ 
 #include "Animation/Skeleton.h"
-
+ 
 #include "AssetRegistry/AssetData.h"
 #include "AssetRegistry/IAssetRegistry.h"
-
+ 
 #include "ConversionUtils/SceneComponentToDynamicMesh.h"
-
+ 
 #include "DynamicMesh/DynamicBoneAttribute.h"
 #include "DynamicMesh/DynamicVertexSkinWeightsAttribute.h"
 #include "DynamicMesh/MeshTransforms.h"
-
+ 
 #include "Engine/SkeletalMesh.h"
 #include "Engine/StaticMesh.h"
-
+ 
 #include "Facades/PVBranchFacade.h"
 #include "Facades/PVPlantFacade.h"
 #include "Facades/PVPointFacade.h"
-
+ 
 #include "GeometryCollection/GeometryCollection.h"
 #include "GeometryCollection/Facades/CollectionConstraintOverrideFacade.h"
-
+ 
 #include "GeometryScript/GeometryScriptTypes.h"
 #include "GeometryScript/MeshAssetFunctions.h"
 #include "GeometryScript/MeshBoneWeightFunctions.h"
 #include "Helpers/PVAnalyticsHelper.h"
-
+ 
 #include "Helpers/PVUtilities.h"
-
+ 
 #include "Misc/Paths.h"
 #include "PhysicsEngine/ConstraintUtils.h"
 #include "PhysicsEngine/PhysicsAsset.h"
 #include "PhysicsEngine/PhysicsConstraintTemplate.h"
 #include "PhysicsEngine/SkeletalBodySetup.h"
-
+#include "ProceduralVegetation/Private/Implementations/PVFoliage.h"
+ 
 #include "UObject/Package.h"
-
+ 
 #define LOCTEXT_NAMESPACE "PVExportHelper"
-
+ 
 namespace PV::Export::Internal
 {
 	enum class EExportMeshStage
@@ -75,7 +76,7 @@ namespace PV::Export::Internal
 		Complete,
 		MAX
 	};
-
+ 
 	float CalculateMeshExportStageWeight(EExportMeshStage InStage, bool bIsBuildingNaniteAssemblies, bool bIsExportingSkeletalMesh)
 	{
 		switch (InStage)
@@ -103,16 +104,16 @@ namespace PV::Export::Internal
 		default:
 			break;
 		}
-
+ 
 		check(false);
 		return 0.f;
 	}
-
+ 
 	float CalculateMeshExportProgress(EExportMeshStage InStage, float StageProgress, bool bIsBuildingNaniteAssemblies, bool bIsExportingSkeletalMesh)
 	{
 		ensure(StageProgress >= 0 && StageProgress <= 1.01f);
 		StageProgress = FMath::Clamp(StageProgress, 0.f, 1.f);
-
+ 
 		float TotalWeightSum = 0;
 		float ThisStageWeightSum = 0;
 		for (int32 i = 0; i < static_cast<int32>(EExportMeshStage::MAX); ++i)
@@ -124,15 +125,15 @@ namespace PV::Export::Internal
 				ThisStageWeightSum += StatusWeight;
 			}
 		}
-
+ 
 		const float ThisStatusWeight = CalculateMeshExportStageWeight(InStage, bIsBuildingNaniteAssemblies, bIsExportingSkeletalMesh);
 		ThisStageWeightSum += ThisStatusWeight * StageProgress;
-
+ 
 		return TotalWeightSum > 0
 			? ThisStageWeightSum / TotalWeightSum
 			: 0.f;
 	}
-
+ 
 	FText GetMeshExportStageText(EExportMeshStage InStage, bool bIsBuildingNaniteAssemblies, bool bIsExportingSkeletalMesh)
 	{
 		switch (InStage)
@@ -169,19 +170,19 @@ namespace PV::Export::Internal
 		default:
 			break;
 		}
-
+ 
 		check(false);
 		return FText::FromString("Error");
 	}
-
+ 
 	void CleanupPackages(const TArray<UPackage*>& Packages)
 	{
 		TArray<UPackage*> PackagesToReload;
 		PackagesToReload.Reserve(Packages.Num());
-
+ 
 		TArray<UObject*> ObjectsToDelete;
 		ObjectsToDelete.Reserve(Packages.Num());
-
+ 
 		for (UPackage* Package : Packages)
 		{
 			if (Package->HasAnyPackageFlags(PKG_InMemoryOnly))
@@ -193,7 +194,7 @@ namespace PV::Export::Internal
 				PackagesToReload.Add(Package);
 			}
 		}
-
+ 
 		if (PackagesToReload.Num() > 0)
 		{
 			FText UnloadError;
@@ -205,7 +206,7 @@ namespace PV::Export::Internal
 			ObjectTools::ForceDeleteObjects(ObjectsToDelete, false);
 		}
 	}
-
+ 
 	enum class EGetValidAssetPathResult
 	{
 		Success,
@@ -213,14 +214,14 @@ namespace PV::Export::Internal
 		ConflictingPackageExist,
 		SkipAsset
 	};
-
+ 
 	EGetValidAssetPathResult GetValidAssetPath(EPVAssetReplacePolicy AssetReplacePolicy, UClass* AssetClass, FString& InOutPackageName, FString& InOutAssetName)
 	{
 		if (!FPackageName::IsValidTextForLongPackageName(InOutPackageName))
 		{
 			return EGetValidAssetPathResult::InvalidPath;
 		}
-
+ 
 		if (FindPackage(nullptr, *InOutPackageName) != nullptr || FPackageName::DoesPackageExist(InOutPackageName))
 		{
 			if (AssetReplacePolicy == EPVAssetReplacePolicy::Append)
@@ -232,19 +233,19 @@ namespace PV::Export::Internal
 				return EGetValidAssetPathResult::SkipAsset;
 			}
 		}
-
+ 
 		if (Utilities::DoesConflictingPackageExist(InOutPackageName, AssetClass))
 		{
 			return EGetValidAssetPathResult::ConflictingPackageExist;
 		}
-
+ 
 		return EGetValidAssetPathResult::Success;
 	}
-
+ 
 	UPackage* CreateAssetPackage(const FString& LongPackageName)
 	{
 		UPackage* Package = LoadPackage(nullptr, *LongPackageName, LOAD_None);
-
+ 
 		if (Package == nullptr)
 		{
 			Package = CreatePackage(*LongPackageName);
@@ -253,7 +254,7 @@ namespace PV::Export::Internal
 		
 		return Package;
 	}
-
+ 
 	void ReplaceAssetInPackage(const TObjectPtr<UObject>& NewObject)
 	{
 		UPackage* Package = NewObject->GetPackage();
@@ -276,7 +277,7 @@ namespace PV::Export::Internal
 		}
 		NewObject->Rename(*ObjectName, Package, REN_DontCreateRedirectors | REN_DoNotDirty | REN_NonTransactional);
 	}
-
+ 
 	TObjectPtr<UDynamicMesh> CollectionToDynamicMesh(const FManagedArrayCollection& Collection)
 	{
 		TRACE_CPUPROFILER_EVENT_SCOPE(PV::Export::CollectionToDynamicMesh);
@@ -286,40 +287,40 @@ namespace PV::Export::Internal
 		{
 			return NewMesh;
 		}
-
+ 
 		const TUniquePtr<FGeometryCollection> GeomCollection = TUniquePtr<FGeometryCollection>(Collection.NewCopy<FGeometryCollection>());
 		if (!GeomCollection.IsValid())
 		{
 			return NewMesh;
 		}
-
+ 
 		const TManagedArray<FTransform3f>& BoneTransforms = Collection.GetAttribute<FTransform3f>(
 			"Transform", FGeometryCollection::TransformGroup);
-
+ 
 		TArray<int32> TransformIndices;
 		TransformIndices.AddUninitialized(BoneTransforms.Num());
-
+ 
 		int32 Idx = 0;
 		for (int32& TransformIdx : TransformIndices)
 		{
 			TransformIdx = Idx++;
 		}
-
+ 
 		constexpr bool bCenterPivot = false;
 		FTransform TransformOut;
 		FDynamicMesh3 CombinedMesh;
 		ConvertGeometryCollectionToDynamicMesh(CombinedMesh, TransformOut, bCenterPivot, *GeomCollection, true, BoneTransforms.GetConstArray(), true, TransformIndices);
-
+ 
 		NewMesh->SetMesh(CombinedMesh);
-
+ 
 		return NewMesh;
 	}
-
+ 
 	bool IsNaniteSupported()
 	{
 		return DoesPlatformSupportNanite(GMaxRHIShaderPlatform);
 	}
-
+ 
 	FGeometryScriptCopyMeshToAssetOptions GetCopyMeshToAssetOptions(const TArray<TObjectPtr<UMaterialInterface>>& InMaterials, const FMeshNaniteSettings& InNaniteSettings, const ENaniteShapePreservation InShapePreservation)
 	{
 		FGeometryScriptCopyMeshToAssetOptions Options;
@@ -328,25 +329,28 @@ namespace PV::Export::Internal
 		Options.bEnableRecomputeTangents = false;
 		Options.bReplaceMaterials = true;
 		Options.NewMaterials = InMaterials;
-
+ 
 		Options.NewNaniteSettings = InNaniteSettings;
 		Options.NewNaniteSettings.bEnabled = IsNaniteSupported();
 		Options.NewNaniteSettings.ShapePreservation = InShapePreservation;
 		Options.bApplyNaniteSettings = true;
-
+ 
 		return Options;
 	}
-
+ 
 	FString GetFoliageName(const Facades::FFoliageFacade& FoliageFacade, int32 Id)
 	{
 		Facades::FFoliageEntryData Data = FoliageFacade.GetFoliageEntry(Id);
-		return FoliageFacade.GetFoliageName(Data.NameId);
+		
+		FPVFoliageInfo Info = FoliageFacade.GetFoliageInfo(Data.NameId);
+		
+		return Info.Mesh.ToString();
 	}
-
+ 
 	TArray<TObjectPtr<UMaterialInterface>> GetMaterialsFromCollection(const FManagedArrayCollection& Collection)
 	{
 		TRACE_CPUPROFILER_EVENT_SCOPE(PV::Export::GetMaterialsFromCollection);
-
+ 
 		TArray<TObjectPtr<UMaterialInterface>> Materials;
 		if (Collection.HasAttribute("MaterialPath", FGeometryCollection::MaterialGroup))
 		{
@@ -358,11 +362,11 @@ namespace PV::Export::Internal
 		}
 		return Materials;
 	}
-
+ 
 	static bool SkeletalMeshToDynamicMesh(const USkeletalMesh* SkeletalMesh, const int32 LodIndex, FDynamicMesh3& ToDynamicMesh)
 	{
 		TRACE_CPUPROFILER_EVENT_SCOPE(PV::Export::SkeletalMeshToDynamicMesh);
-
+ 
 		if (SkeletalMesh->HasMeshDescription(LodIndex))
 		{
 			const FMeshDescription* SourceMesh = SkeletalMesh->GetMeshDescription(LodIndex);
@@ -370,15 +374,15 @@ namespace PV::Export::Internal
 			{
 				return false;
 			}
-
+ 
 			FMeshDescriptionToDynamicMesh Converter;
 			Converter.Convert(SourceMesh, ToDynamicMesh, true);
 			return true;
 		}
-
+ 
 		return false;
 	}
-
+ 
 	bool CombineMeshInstancesToDynamicMesh(
 		const FManagedArrayCollection& Collection,
 		UDynamicMesh* DynamicMesh,
@@ -388,16 +392,16 @@ namespace PV::Export::Internal
 	)
 	{
 		TRACE_CPUPROFILER_EVENT_SCOPE(PV::Export::CombineMeshInstancesToDynamicMesh);
-
+ 
 		const Facades::FFoliageFacade FoliageFacade(Collection);
 		const int32 NumInstances = FoliageFacade.NumFoliageEntries();
-		const int32 NumFoliageNames = FoliageFacade.NumFoliageNames();
+		const int32 NumFoliageNames = FoliageFacade.NumFoliageInfo();
 		
 		if (NumInstances == 0 || NumFoliageNames == 0)
 		{
 			return true;
 		}
-
+ 
 		struct FConvertFoliageToDynamicMeshResult
 		{
 			FString FoliageName;
@@ -406,37 +410,38 @@ namespace PV::Export::Internal
 			FString OutMeshPath;
 			bool bSuccess = true;
 		};
-
+ 
 		TArray<FConvertFoliageToDynamicMeshResult> ConvertFoliageToDynamicMeshResults;
 		TMap<FString, int32> FoliageNameToDynamicMeshResult;
 		{
 			TRACE_CPUPROFILER_EVENT_SCOPE(ConvertFoliageToDynamicMeshes);
-
+ 
 			ConvertFoliageToDynamicMeshResults.Reserve(NumFoliageNames);
 			FoliageNameToDynamicMeshResult.Reserve(NumFoliageNames);
-
+ 
 			for (int32 FoliageNameIndex = 0; FoliageNameIndex < NumFoliageNames; ++FoliageNameIndex)
 			{
-				const FString FoliageName = FoliageFacade.GetFoliageName(FoliageNameIndex);
-
+				FPVFoliageInfo Info = FoliageFacade.GetFoliageInfo(FoliageNameIndex);
+				const FString FoliageName = Info.Mesh.ToString();
+				
 				if (FoliageNameToDynamicMeshResult.Contains(FoliageName))
 				{
 					UE_LOG(LogProceduralVegetationEditor, Warning, TEXT("Duplicate foliage name %s in foliage facade"), *FoliageName);
 					continue;
 				}
-
+ 
 				FConvertFoliageToDynamicMeshResult& Result = ConvertFoliageToDynamicMeshResults.AddDefaulted_GetRef();
 				Result.FoliageName = FoliageName;
-
+ 
 				FoliageNameToDynamicMeshResult.Add(FoliageName, ConvertFoliageToDynamicMeshResults.Num() - 1);
 			}
-
+ 
 			ParallelFor(ConvertFoliageToDynamicMeshResults.Num(), [&ConvertFoliageToDynamicMeshResults](int32 Index)
 			{
 				FConvertFoliageToDynamicMeshResult& Result = ConvertFoliageToDynamicMeshResults[Index];
-
+ 
 				const FString& FoliageName = Result.FoliageName;
-
+ 
 				if (const TObjectPtr<UStaticMesh> StaticMesh = LoadObject<UStaticMesh>(nullptr, *FoliageName))
 				{
 					FText OutErrorMessage;
@@ -446,12 +451,12 @@ namespace PV::Export::Internal
 						Result.bSuccess = false;
 						return;
 					}
-
+ 
 					for (auto Material : StaticMesh->GetStaticMaterials())
 					{
 						Result.OutMaterials.Add(Material.MaterialInterface);
 					}
-
+ 
 					Result.OutMeshPath = StaticMesh->GetFullName();
 				}
 				else if (const TObjectPtr<USkeletalMesh> SkeletalMesh = LoadObject<USkeletalMesh>(nullptr, *FoliageName))
@@ -461,12 +466,12 @@ namespace PV::Export::Internal
 						Result.bSuccess = false;
 						return;
 					}
-
+ 
 					for (auto Material : SkeletalMesh->GetMaterials())
 					{
 						Result.OutMaterials.Add(Material.MaterialInterface);
 					}
-
+ 
 					Result.OutDynamicMesh.Attributes()->RemoveAttribute("BoneIndex");
 					Result.OutDynamicMesh.Attributes()->RemoveAttribute("BoneParentIndex");
 					Result.OutMeshPath = SkeletalMesh->GetFullName();
@@ -476,22 +481,22 @@ namespace PV::Export::Internal
 					Result.bSuccess = false;
 				}
 			});
-
+ 
 			{
 				TRACE_CPUPROFILER_EVENT_SCOPE(RemapMaterials);
-
+ 
 				TSet<UMaterialInterface*> UniqueMaterials;
 				UniqueMaterials.Reserve(ConvertFoliageToDynamicMeshResults.Num() * 2);
 				for (const FConvertFoliageToDynamicMeshResult& Result : ConvertFoliageToDynamicMeshResults)
 				{
 					UniqueMaterials.Append(Result.OutMaterials);
 				}
-
+ 
 				for (UMaterialInterface* Material : UniqueMaterials)
 				{
 					Materials.AddUnique(Material);
 				}
-
+ 
 				for (FConvertFoliageToDynamicMeshResult& Result : ConvertFoliageToDynamicMeshResults)
 				{
 					UE::Geometry::FDynamicMeshMaterialAttribute* MaterialIDs = Result.OutDynamicMesh.Attributes()->GetMaterialID();
@@ -507,34 +512,34 @@ namespace PV::Export::Internal
 						{
 							continue;
 						}
-
+ 
 						check(Result.OutMaterials.IsValidIndex(MaterialID));
-
+ 
 						const int32 NewMaterialID = Materials.Find(Result.OutMaterials[MaterialID]);
 						check(NewMaterialID != INDEX_NONE);
-
+ 
 						MaterialIDs->SetValue(Tid, NewMaterialID);
 					}
 				}
 			}
 		}
-
+ 
 		const float ConvertToDynamicMeshProgress = 0.1f;
 		if (!OnProgressUpdated(ConvertToDynamicMeshProgress))
 		{
 			return false;
 		}
-
+ 
 		struct FGenerateDynamicMeshesResult
 		{
 			FDynamicMesh3 DynamicMesh;
 			int32 ConvertToDynamicMeshResultIndex;
 		};
-
+ 
 		TArray<FGenerateDynamicMeshesResult> GenerateDynamicMeshesResults;
 		{
 			TRACE_CPUPROFILER_EVENT_SCOPE(GenerateDynamicMeshInstances);
-
+ 
 			GenerateDynamicMeshesResults.SetNum(NumInstances);
 			ParallelFor(NumInstances, [
 				&GenerateDynamicMeshesResults,
@@ -544,23 +549,23 @@ namespace PV::Export::Internal
 			](int32 Index)
 			{
 				const FString FoliageName = GetFoliageName(FoliageFacade, Index);
-
+ 
 				const int32* PrevResultIndex = FoliageNameToDynamicMeshResult.Find(FoliageName);
 				if (PrevResultIndex == nullptr)
 				{
 					return;
 				}
-
+ 
 				FConvertFoliageToDynamicMeshResult& ConvertFoliageToDynamicMeshResult = ConvertFoliageToDynamicMeshResults[*PrevResultIndex];
 				if (!ConvertFoliageToDynamicMeshResult.bSuccess)
 				{
 					return;
 				}
-
+ 
 				FGenerateDynamicMeshesResult& Result = GenerateDynamicMeshesResults[Index];
 				Result.DynamicMesh = ConvertFoliageToDynamicMeshResult.OutDynamicMesh;
 				Result.ConvertToDynamicMeshResultIndex = *PrevResultIndex;
-
+ 
 				const int32 FoliageParentBoneID = FoliageFacade.GetParentBoneID(Index);
 				UE::Geometry::FDynamicMeshAttributeSet* MeshAttributes = Result.DynamicMesh.Attributes();
 				UE::Geometry::FDynamicMeshBoneParentIndexAttribute* BoneParentIndices = MeshAttributes->GetBoneParentIndices();
@@ -569,17 +574,17 @@ namespace PV::Export::Internal
 					// Reparent root bone according to the instanced transform
 					BoneParentIndices->SetValue(0, FoliageParentBoneID);
 				}
-
+ 
 				const FTransform Transform = FoliageFacade.GetFoliageTransform(Index);
 				MeshTransforms::ApplyTransform(Result.DynamicMesh, Transform, true);
 			});
 		}
-
+ 
 		bool bWasCanceled = false;
 		DynamicMesh->EditMesh([&](FDynamicMesh3& Mesh)
 		{
 			TRACE_CPUPROFILER_EVENT_SCOPE(AppendDynamicMeshes);
-
+ 
 			for (int32 i = 0; i < GenerateDynamicMeshesResults.Num(); ++i)
 			{
 				const FGenerateDynamicMeshesResult& GenerateDynamicMeshesResult = GenerateDynamicMeshesResults[i];
@@ -588,13 +593,13 @@ namespace PV::Export::Internal
 				{
 					continue;
 				}
-
+ 
 				const int32 VertexStart = Mesh.VertexCount();
 				Mesh.AppendWithOffsets(MeshToAppend);
 				
 				const FConvertFoliageToDynamicMeshResult& ConvertFoliageToDynamicMeshResult = ConvertFoliageToDynamicMeshResults[GenerateDynamicMeshesResult.ConvertToDynamicMeshResultIndex];
 				OnMeshInstanceCombined(ConvertFoliageToDynamicMeshResult.OutMeshPath, i, VertexStart, Mesh.VertexCount());
-
+ 
 				const float CombineMeshesProgress = ConvertToDynamicMeshProgress + i / (float)NumInstances * (1.f - ConvertToDynamicMeshProgress);
 				if (!OnProgressUpdated(CombineMeshesProgress))
 				{
@@ -603,10 +608,10 @@ namespace PV::Export::Internal
 				}
 			}
 		});
-
+ 
 		return !bWasCanceled;
 	}
-
+ 
 	bool ContainsFoliageName(const Facades::FFoliageFacade& FoliageFacade, const FString& InFoliageName)
 	{
 		int32 NumInstances = FoliageFacade.NumFoliageEntries();
@@ -619,27 +624,29 @@ namespace PV::Export::Internal
 				return true;
 			}
 		}
-
+ 
 		return false;
 	}
-
+ 
 	TArray<FString> GetUsedFoliage(const Facades::FFoliageFacade& FoliageFacade)
 	{
 		TArray<FString> OutFoliageNames;
-		OutFoliageNames.Reserve(FoliageFacade.NumFoliageNames());
-
-		for(int32 FoliageNameIndex = 0; FoliageNameIndex < FoliageFacade.NumFoliageNames(); ++FoliageNameIndex)
+		OutFoliageNames.Reserve(FoliageFacade.NumFoliageInfo());
+ 
+		for(int32 FoliageNameIndex = 0; FoliageNameIndex < FoliageFacade.NumFoliageInfo(); ++FoliageNameIndex)
 		{
-			FString FoliageName = FoliageFacade.GetFoliageName(FoliageNameIndex);
+			FPVFoliageInfo Info = FoliageFacade.GetFoliageInfo(FoliageNameIndex);
+			FString FoliageName = Info.Mesh.ToString();
+			
 			if (ContainsFoliageName(FoliageFacade, FoliageName))
 			{
 				OutFoliageNames.Add(FoliageName);
 			}
 		}
-
+ 
 		return OutFoliageNames;
 	}
-
+ 
 	void BuildNaniteAssemblyData(
 		const FManagedArrayCollection& Collection,
 		UStaticMesh* StaticMesh,
@@ -647,27 +654,27 @@ namespace PV::Export::Internal
 	)
 	{
 		TRACE_CPUPROFILER_EVENT_SCOPE(PV::Export::BuildNaniteAssemblyData_StaticMesh);
-
+ 
 		//Create NaniteAssemblyDataBuilder and the main static mesh materials to it
 		FNaniteAssemblyDataBuilder AssemblyBuilder;
 		for (const FStaticMaterial& Material : StaticMesh->GetStaticMaterials())
 		{
 			AssemblyBuilder.AddMaterialSlot(Material.MaterialInterface);
 		}
-
+ 
 		const Facades::FFoliageFacade FoliageFacade(Collection);
 		const TArray<FString> UsedFoliage = GetUsedFoliage(FoliageFacade);
-
+ 
 		TMap<FString, int32> MeshNamePartMap;
 		MeshNamePartMap.Reserve(UsedFoliage.Num());
 		TMap<FString, TArray<TObjectPtr<UMaterialInterface>>> MeshMaterialsMap;
 		MeshMaterialsMap.Reserve(UsedFoliage.Num() * 2);
-
+ 
 		const auto AddStaticMeshToAssembly = [&](const FString& FoliageInstanceObjectPath, TObjectPtr<UStaticMesh> InStaticMesh)
 		{
 			const int32 PartIndex = AssemblyBuilder.FindOrAddPart(InStaticMesh.GetPath());
 			MeshNamePartMap.FindOrAdd(FoliageInstanceObjectPath) = PartIndex;
-
+ 
 			TArray<TObjectPtr<UMaterialInterface>>& MeshMaterials = MeshMaterialsMap.FindOrAdd(FoliageInstanceObjectPath);
 			for (const FStaticMaterial& Material : InStaticMesh->GetStaticMaterials())
 			{
@@ -688,29 +695,29 @@ namespace PV::Export::Internal
 				AddStaticMeshToAssembly(FoliageInstanceObjectPath, PartStaticMesh);
 				continue;
 			}
-
+ 
 			const FString FoliageInstancePackagePath = FPackageName::ObjectPathToPackageName(FoliageInstanceObjectPath);
 			const FString LongFoliagePackagePath = FPackageName::GetLongPackagePath(FoliageInstancePackagePath);
 			const FString ShortFoliagePackageName = FPackageName::GetShortName(FoliageInstancePackagePath);
-
+ 
 			const FString StaticMeshName = "SKM_" + ShortFoliagePackageName;
 			const FString StaticMeshPackagePath = FPaths::Combine(LongFoliagePackagePath, StaticMeshName);
 			const FString StaticMeshObjectPath = FString::Printf(TEXT("%s.%s"), *StaticMeshPackagePath, *StaticMeshName);
-
+ 
 			PartStaticMesh = LoadObject<UStaticMesh>(nullptr, StaticMeshObjectPath);
 			if (PartStaticMesh)
 			{
 				AddStaticMeshToAssembly(FoliageInstanceObjectPath, PartStaticMesh);
 				continue;
 			}
-
+ 
 			const TObjectPtr<USkeletalMesh> PartSkeletalMesh = LoadObject<USkeletalMesh>(nullptr, *FoliageInstanceObjectPath);
 			if (!PartSkeletalMesh)
 			{
 				UE_LOG(LogProceduralVegetationEditor, Warning, TEXT("Could not find foliage instance %s"), *FoliageInstanceObjectPath);
 				continue;
 			}
-
+ 
 			FDynamicMesh3 PartDynamicMesh;
 			if (!SkeletalMeshToDynamicMesh(PartSkeletalMesh, 0, PartDynamicMesh))
 			{
@@ -720,30 +727,30 @@ namespace PV::Export::Internal
 					
 			UPackage* StaticMeshPackage = Internal::CreateAssetPackage(*StaticMeshPackagePath);
 			OutModifiedPackages.Add(StaticMeshPackage);
-
+ 
 			PartStaticMesh = NewObject<UStaticMesh>(StaticMeshPackage, FName(StaticMeshName), RF_Standalone | RF_Public);
-
+ 
 			TArray<TObjectPtr<UMaterialInterface>> PartMaterials;
 			for (auto Material : PartSkeletalMesh->GetMaterials())
 			{
 				PartMaterials.Add(Material.MaterialInterface);
 			}
-
+ 
 			EGeometryScriptOutcomePins Output;
 			FGeometryScriptCopyMeshToAssetOptions Options = GetCopyMeshToAssetOptions(PartMaterials, PartStaticMesh->GetNaniteSettings(), NaniteVoxelsSupported() ? ENaniteShapePreservation::Voxelize : ENaniteShapePreservation::PreserveArea);
-
+ 
 			TObjectPtr<UDynamicMesh> DynamicMesh = NewObject<UDynamicMesh>();
 			DynamicMesh->SetMesh(MoveTemp(PartDynamicMesh));
-
+ 
 			UGeometryScriptLibrary_StaticMeshFunctions::CopyMeshToStaticMesh(DynamicMesh, PartStaticMesh, Options, {}, Output);
-
+ 
 			AddStaticMeshToAssembly(FoliageInstanceObjectPath, PartStaticMesh);
 		}
-
+ 
 		AddNodeToBuilder(AssemblyBuilder, Collection, MeshNamePartMap, MeshMaterialsMap);
 		AssemblyBuilder.ApplyToStaticMesh(*StaticMesh);
 	}
-
+ 
 	void BuildNaniteAssemblyData(
 		const FManagedArrayCollection& Collection,
 		USkeletalMesh* SkeletalMesh,
@@ -751,34 +758,34 @@ namespace PV::Export::Internal
 	)
 	{
 		TRACE_CPUPROFILER_EVENT_SCOPE(PV::Export::BuildNaniteAssemblyData_SkeletalMesh);
-
+ 
 		//Create NaniteAssemblyDataBuilder and the main static mesh materials to it
 		FNaniteAssemblyDataBuilder AssemblyBuilder;
 		for (const FSkeletalMaterial& Material : SkeletalMesh->GetMaterials())
 		{
 			AssemblyBuilder.AddMaterialSlot(Material.MaterialInterface);
 		}
-
+ 
 		const Facades::FFoliageFacade FoliageFacade(Collection);
 		const TArray<FString> UsedFoliage = GetUsedFoliage(FoliageFacade);
-
+ 
 		TMap<FString, int32> MeshNamePartMap;
 		MeshNamePartMap.Reserve(UsedFoliage.Num());
 		TMap<FString, TArray<TObjectPtr<UMaterialInterface>>> MeshMaterialsMap;
 		MeshMaterialsMap.Reserve(UsedFoliage.Num() * 2);
-
+ 
 		const auto AddSkeletalMeshToAssembly = [&](const FString& FoliageInstanceObjectPath, TObjectPtr<USkeletalMesh> InSkeletalMesh)
 		{
 			const int32 PartIndex = AssemblyBuilder.FindOrAddPart(InSkeletalMesh.GetPath());
 			MeshNamePartMap.FindOrAdd(FoliageInstanceObjectPath) = PartIndex;
-
+ 
 			TArray<TObjectPtr<UMaterialInterface>>& MeshMaterials = MeshMaterialsMap.FindOrAdd(FoliageInstanceObjectPath);
 			for (const FSkeletalMaterial& Material : InSkeletalMesh->GetMaterials())
 			{
 				MeshMaterials.AddUnique(Material.MaterialInterface);
 			}
 		};
-
+ 
 		for (const FString& FoliageInstanceObjectPath : UsedFoliage)
 		{
 			if (FoliageInstanceObjectPath.IsEmpty())
@@ -792,19 +799,19 @@ namespace PV::Export::Internal
 				AddSkeletalMeshToAssembly(FoliageInstanceObjectPath, PartSkeletalMesh);
 				continue;
 			}
-
+ 
 			const FString FoliageInstancePackagePath = FPackageName::ObjectPathToPackageName(FoliageInstanceObjectPath);
 			const FString LongFoliagePackagePath = FPackageName::GetLongPackagePath(FoliageInstancePackagePath);
 			const FString ShortFoliagePackageName = FPackageName::GetShortName(FoliageInstancePackagePath);
-
+ 
 			const FString SkeletalMeshName = "SKM_" + ShortFoliagePackageName;
 			const FString SkeletalMeshPackagePath = FPaths::Combine(LongFoliagePackagePath, SkeletalMeshName);
 			const FString SkeletalMeshObjectPath = FString::Printf(TEXT("%s.%s"), *SkeletalMeshPackagePath, *SkeletalMeshName);
-
+ 
 			const FString SkeletonName = "SK_" + ShortFoliagePackageName;
 			const FString SkeletonPackagePath = FPaths::Combine(LongFoliagePackagePath, SkeletonName);
 			const FString SkeletonObjectPath = FString::Printf(TEXT("%s.%s"), *SkeletonPackagePath, *SkeletonName);
-
+ 
 			PartSkeletalMesh = LoadObject<USkeletalMesh>(nullptr, SkeletalMeshObjectPath);
 			TObjectPtr<USkeleton> Skeleton = LoadObject<USkeleton>(nullptr, SkeletonObjectPath);
 			if (PartSkeletalMesh && Skeleton)
@@ -812,21 +819,21 @@ namespace PV::Export::Internal
 				AddSkeletalMeshToAssembly(FoliageInstanceObjectPath, PartSkeletalMesh);
 				continue;
 			}
-
+ 
 			if (!Skeleton)
 			{
 				UPackage* SkeletonPackage = Internal::CreateAssetPackage(*SkeletonPackagePath);
 				OutModifiedPackages.Add(SkeletonPackage);
-
+ 
 				Skeleton = NewObject<USkeleton>(SkeletonPackage, FName(SkeletonName), RF_Standalone | RF_Public);
-
+ 
 				if (PartSkeletalMesh)
 				{
 					PartSkeletalMesh->SetSkeleton(Skeleton);
 					Skeleton->RecreateBoneTree(PartSkeletalMesh);
 				}
 			}
-
+ 
 			if (!PartSkeletalMesh)
 			{
 				const TObjectPtr<UStaticMesh> PartStaticMesh = LoadObject<UStaticMesh>(nullptr, *FoliageInstanceObjectPath);
@@ -835,7 +842,7 @@ namespace PV::Export::Internal
 					UE_LOG(LogProceduralVegetationEditor, Warning, TEXT("Could not find foliage instance %s"), *FoliageInstanceObjectPath);
 					continue;
 				}
-
+ 
 				FDynamicMesh3 PartDynamicMesh;
 				FText OutErrorMessage;
 				UE::Conversion::FStaticMeshConversionOptions Options;
@@ -844,48 +851,48 @@ namespace PV::Export::Internal
 					UE_LOG(LogProceduralVegetationEditor, Warning, TEXT("Failed to convert foliage instance \"%s\" to dynamic mesh"), *FoliageInstancePackagePath);
 					continue;
 				}
-
+ 
 				UPackage* SkeletalMeshPackage = Internal::CreateAssetPackage(*SkeletalMeshPackagePath);
 				OutModifiedPackages.Add(SkeletalMeshPackage);
-
+ 
 				if (!Skeleton)
 				{
 					UPackage* SkeletonPackage = Internal::CreateAssetPackage(*SkeletonPackagePath);
 					OutModifiedPackages.Add(SkeletonPackage);
 					Skeleton = NewObject<USkeleton>(SkeletonPackage, FName(SkeletonName), RF_Standalone | RF_Public);
 				}
-
+ 
 				PartSkeletalMesh = NewObject<USkeletalMesh>(SkeletalMeshPackage, FName(SkeletalMeshName), RF_Standalone | RF_Public);
 				PartSkeletalMesh->SetSkeleton(Skeleton);
-
+ 
 				TArray<TObjectPtr<UMaterialInterface>> PartMaterials;
 				for (auto Material : PartStaticMesh->GetStaticMaterials())
 				{
 					PartMaterials.Add(Material.MaterialInterface);
 				}
-
+ 
 				ConvertToDefaultSkeletalMesh(PartSkeletalMesh, PartDynamicMesh, PartMaterials);
 			}
 			
 			AddSkeletalMeshToAssembly(FoliageInstanceObjectPath, PartSkeletalMesh);
 		}
-
+ 
 		AddNodeToBuilder(AssemblyBuilder, Collection, MeshNamePartMap, MeshMaterialsMap);
 		AssemblyBuilder.ApplyToSkeletalMesh(*SkeletalMesh);
 	}
-
+ 
 	void AddNodeToBuilder(FNaniteAssemblyDataBuilder& AssemblyBuilder, const FManagedArrayCollection& Collection,
 		const TMap<FString, int32>& InMeshNamePartMap, const TMap<FString, TArray<TObjectPtr<UMaterialInterface>>>& InMeshMaterialsMap)
 	{
 		TRACE_CPUPROFILER_EVENT_SCOPE(PV::Export::AddNodeToBuilder);
-
+ 
 		Facades::FFoliageFacade FoliageFacade(Collection);
 		const int32 NumInstances = FoliageFacade.NumFoliageEntries();
 			
 		for (int32 Id = 0; Id < NumInstances; Id++)
 		{
 			const FString FoliageName = GetFoliageName(FoliageFacade, Id);
-
+ 
 			const int32* PartIndex = InMeshNamePartMap.Find(FoliageName);
 			if (PartIndex == nullptr)
 			{
@@ -894,7 +901,7 @@ namespace PV::Export::Internal
 				
 			const FTransform Transform = FoliageFacade.GetFoliageTransform(Id);
 			const UE::Geometry::FTransformSRT3d GeoTransform(Transform);
-
+ 
 			const int32 FoliageParentBoneID = FoliageFacade.GetParentBoneID(Id);
 				
 			FNaniteAssemblyBoneInfluence Influence;
@@ -902,14 +909,14 @@ namespace PV::Export::Internal
 			Influence.BoneWeight = 1.0;
 				
 			AssemblyBuilder.AddNode(*PartIndex, FTransform3f(GeoTransform), ENaniteAssemblyNodeTransformSpace::Local, { Influence });
-
+ 
 			//Add the materials for part to the assembly builder
 			auto* FoliageMaterials = InMeshMaterialsMap.Find(FoliageName);
 			if (FoliageMaterials == nullptr)
 			{
 				continue;
 			}
-
+ 
 			int32 LocalMaterialIndex = 0;
 			for (auto PartMaterial : *FoliageMaterials)
 			{
@@ -920,17 +927,17 @@ namespace PV::Export::Internal
 				{
 					PartMaterialIndex = AssemblyBuilder.AddMaterialSlot(PartMaterial);
 				}
-
+ 
 				AssemblyBuilder.RemapPartMaterial(*PartIndex, LocalMaterialIndex, PartMaterialIndex);
 				LocalMaterialIndex++;
 			}
 		}
 	}
-
+ 
 	void ConvertToDefaultSkeletalMesh(USkeletalMesh* SkeletalMesh, FDynamicMesh3& Mesh, const TArray<TObjectPtr<UMaterialInterface>>& Materials)
 	{
 		TRACE_CPUPROFILER_EVENT_SCOPE(PV::Export::ConvertToDefaultSkeletalMesh);
-
+ 
 		TObjectPtr<UDynamicMesh> DynamicMesh = NewObject<UDynamicMesh>();
 		DynamicMesh->SetMesh(MoveTemp(Mesh));
 		
@@ -938,11 +945,11 @@ namespace PV::Export::Internal
 		int ParentBoneIndex = INDEX_NONE;
 		UE::Geometry::FDynamicMeshAttributeSet* MeshAttributes = DynamicMesh->GetMeshPtr()->Attributes();
 		MeshAttributes->EnableBones(1);
-
+ 
 		UE::Geometry::FDynamicMeshBoneNameAttribute* BoneNames = MeshAttributes->GetBoneNames();
 		UE::Geometry::FDynamicMeshBoneParentIndexAttribute* BoneParentIndices = MeshAttributes->GetBoneParentIndices();
 		UE::Geometry::FDynamicMeshBonePoseAttribute* BonePoses = MeshAttributes->GetBonePoses();
-
+ 
 		int BoneIndex = 0;
 		BoneNames->SetValue(BoneIndex, "Root");
 		BonePoses->SetValue(BoneIndex, FTransform::Identity);
@@ -960,10 +967,10 @@ namespace PV::Export::Internal
 			{},
 			Output
 		);
-
+ 
 		TObjectPtr<USkeleton> Skeleton = SkeletalMesh->GetSkeleton();
 		Skeleton->RecreateBoneTree(SkeletalMesh);
-
+ 
 		const FGeometryScriptBoneWeightProfile SkinProfile = FGeometryScriptBoneWeightProfile();
 		FGeometryScriptSmoothBoneWeightsOptions SmoothBoneWeightsOptions;
 		SmoothBoneWeightsOptions.DistanceWeighingType = EGeometryScriptSmoothBoneWeightsType::GeodesicVoxel;
@@ -975,7 +982,7 @@ namespace PV::Export::Internal
 			SmoothBoneWeightsOptions,
 			SkinProfile
 		);
-
+ 
 		const FName ProfileName = SkinProfile.GetProfileName();
 		UGeometryScriptLibrary_StaticMeshFunctions::CopySkinWeightProfileToSkeletalMesh(
 			DynamicMesh,
@@ -987,7 +994,7 @@ namespace PV::Export::Internal
 			Output
 		);
 	}
-
+ 
 	void AddCollisionToStaticMesh(const FManagedArrayCollection& Collection, TObjectPtr<UStaticMesh> ExportMesh)
 	{
 		UBodySetup* BodySetup = ExportMesh->GetBodySetup();
@@ -1001,7 +1008,7 @@ namespace PV::Export::Internal
 		}
 			
 		FKConvexElem ConvexElem;
-
+ 
 		FGeometryCollection GeoCollection;
 		Collection.CopyTo(&GeoCollection);
 		
@@ -1014,14 +1021,14 @@ namespace PV::Export::Internal
 		ConvexElem.VertexData = Vertices;
 		ConvexElem.UpdateElemBox();
 		BodySetup->AggGeom.ConvexElems.Add(ConvexElem);
-
+ 
 		BodySetup->bHasCookedCollisionData = true;
 		BodySetup->InvalidatePhysicsData();
 		BodySetup->CreatePhysicsMeshes();
-
+ 
 		ExportMesh->SetBodySetup(BodySetup);
 	}
-
+ 
 	bool ExportCollectionToStaticMesh(
 		TObjectPtr<UStaticMesh> ExportMesh,
 		const FManagedArrayCollection& Collection,
@@ -1034,39 +1041,39 @@ namespace PV::Export::Internal
 	)
 	{
 		TRACE_CPUPROFILER_EVENT_SCOPE(PV::Export::ExportCollectionToStaticMesh);
-
+ 
 		const auto UpdateStage = [&](EExportMeshStage NewStage, float StageProgress)->bool
 		{
 			return StatusReportCallback(GetMeshExportStageText(NewStage, bBuildNaniteAssemblies, false), CalculateMeshExportProgress(NewStage, StageProgress, bBuildNaniteAssemblies, false));
 		};
-
+ 
 		FManagedArrayCollection ExportedCollection;
 		TObjectPtr<UDynamicMesh> DynamicMesh;
 		TArray<TObjectPtr<UMaterialInterface>> Materials;
 		{
 			TRACE_CPUPROFILER_EVENT_SCOPE(InitExportMesh);
-
+ 
 			if (!UpdateStage(EExportMeshStage::InitializingExportMesh, 0.f))
 			{
 				return false;
 			}
-
+ 
 			Collection.CopyTo(&ExportedCollection);
-
+ 
 			DynamicMesh = CollectionToDynamicMesh(ExportedCollection);
-
+ 
 			if (bCollision)
 			{
 				AddCollisionToStaticMesh(ExportedCollection, ExportMesh);
 			}
-
+ 
 			if (!UpdateStage(EExportMeshStage::InitializingExportMesh, 0.5f))
 			{
 				return false;
 			}
-
+ 
 			Materials = GetMaterialsFromCollection(ExportedCollection);
-
+ 
 			if (!UpdateStage(EExportMeshStage::InitializingExportMesh, 1.f))
 			{
 				return false;
@@ -1091,27 +1098,27 @@ namespace PV::Export::Internal
 				return false;
 			}
 		}
-
+ 
 		EGeometryScriptOutcomePins Output;
 		{
 			TRACE_CPUPROFILER_EVENT_SCOPE(CopyMeshToStaticMesh);
-
+ 
 			if (!UpdateStage(EExportMeshStage::CopyingMeshDataToOutputMesh, 1.f))
 			{
 				return false;
 			}
-
+ 
 			FGeometryScriptCopyMeshToAssetOptions Options = GetCopyMeshToAssetOptions(Materials, ExportMesh->GetNaniteSettings(), InShapePreservation);
 			UGeometryScriptLibrary_StaticMeshFunctions::CopyMeshToStaticMesh(DynamicMesh, ExportMesh, Options, {}, Output);
 		}
-
+ 
 		if (bBuildNaniteAssemblies && bShouldExportFoliage)
 		{
 			if (!UpdateStage(EExportMeshStage::BuildingNaniteAssemblies, 1.f))
 			{
 				return false;
 			}
-
+ 
 			if (NaniteAssembliesSupported())
 			{
 				//Build the nanite assemblies
@@ -1122,7 +1129,7 @@ namespace PV::Export::Internal
 				UE_LOG(LogProceduralVegetationEditor, Warning, TEXT("Failed to build Nanite Assemblies, because neither Nanite Foliage nor Nanite Assemblies are enabled for the project."));
 			}
 		}
-
+ 
 		{
 			TRACE_CPUPROFILER_EVENT_SCOPE(BuildStaticMesh);
 			if (!UpdateStage(EExportMeshStage::BuildingMesh, 1.f))
@@ -1133,24 +1140,24 @@ namespace PV::Export::Internal
 			ExportMesh->Build(true);
 			FStaticMeshCompilingManager::Get().FinishAllCompilation();
 		}
-
+ 
 		UpdateStage(EExportMeshStage::Complete, 1.f);
 		return true;
 	}
-
+ 
 	void BuildWindImportData(FDynamicWindSkeletalImportData& OutImportData, const FManagedArrayCollection& Collection, const TArray<PV::Facades::FBoneNode>& Bones)
 	{
 		TRACE_CPUPROFILER_EVENT_SCOPE(PV::Export::BuildWindImportData);
-
+ 
 		Facades::FBranchFacade BranchFacade(Collection);
-
+ 
 		for (const auto& Bone : Bones)
 		{
 			int32 SimulationGroupIndex = BranchFacade.GetBranchSimulationGroupIndex(Bone.BranchIndex);
 			OutImportData.Joints.Add({Bone.BoneName, SimulationGroupIndex});
 		}
 	}
-
+ 
 	bool ExportCollectionToSkeletalMesh(
 		TObjectPtr<USkeletalMesh> ExportMesh,
 		const FManagedArrayCollection& Collection,
@@ -1165,7 +1172,7 @@ namespace PV::Export::Internal
 	)
 	{
 		TRACE_CPUPROFILER_EVENT_SCOPE(PV::Export::ExportCollectionToSkeletalMesh);
-
+ 
 		const auto UpdateStage = [&](EExportMeshStage NewStage, float StageProgress)->bool
 		{
 			return StatusReportCallback(
@@ -1173,7 +1180,7 @@ namespace PV::Export::Internal
 				CalculateMeshExportProgress(NewStage, StageProgress, bBuildNaniteAssemblies, true)
 			);
 		};
-
+ 
 		FManagedArrayCollection ExportedCollection;
 		TObjectPtr<UDynamicMesh> DynamicMesh;
 		TArray<TObjectPtr<UMaterialInterface>> Materials;
@@ -1182,42 +1189,42 @@ namespace PV::Export::Internal
 		int32 FoliageVertexStart = 0;
 		{
 			TRACE_CPUPROFILER_EVENT_SCOPE(InitExportMesh);
-
+ 
 			if (!UpdateStage(EExportMeshStage::InitializingExportMesh, 0.f))
 			{
 				return false;
 			}
-
+ 
 			Collection.CopyTo(&ExportedCollection);
-
+ 
 			DynamicMesh = CollectionToDynamicMesh(ExportedCollection);
-
+ 
 			if (!UpdateStage(EExportMeshStage::InitializingExportMesh, 0.5f))
 			{
 				return false;
 			}
-
+ 
 			Materials = GetMaterialsFromCollection(ExportedCollection);
-
+ 
 			UE::Geometry::FDynamicMeshAttributeSet* MeshAttributes = DynamicMesh->GetMeshPtr()->Attributes();
-
+ 
 			Facades::FBoneFacade BoneFacade = Facades::FBoneFacade(ExportedCollection);
 			BoneNodes = BoneFacade.GetBoneDataFromCollection();
-
+ 
 			// if bones already not created with bone reduction node create bones with full density
 			if (BoneNodes.IsEmpty())
 			{
 				BoneFacade.CreateBoneData(BoneNodes, 0);
 			}
-
+ 
 			if (bShouldExportFoliage)
 			{
 				AssignBoneIDsToFoliage(BoneNodes, ExportedCollection);
 			}
-
+ 
 			const int32 BoneCount = BoneNodes.Num();
 			MeshAttributes->EnableBones(BoneCount);
-
+ 
 			TArray<int32> VertexPointIDs;
 			BoneFacade.GetPointIds(VertexPointIDs);
 			
@@ -1234,35 +1241,35 @@ namespace PV::Export::Internal
 					UE_LOG(LogProceduralVegetationEditor, Log, TEXT("Invalid Bone Assigned to PointID %i"), PointID);
 				}
 			}
-
+ 
 			FoliageVertexStart = VertexBoneIDs.Num();
-
+ 
 			UE::Geometry::FDynamicMeshBoneNameAttribute* BoneNames = MeshAttributes->GetBoneNames();
 			UE::Geometry::FDynamicMeshBoneParentIndexAttribute* BoneParentIndices = MeshAttributes->GetBoneParentIndices();
 			UE::Geometry::FDynamicMeshBonePoseAttribute* BonePoses = MeshAttributes->GetBonePoses();
-
+ 
 			for (int i = 0; i < BoneCount; i++)
 			{
 				BoneNames->SetValue(i, BoneNodes[i].BoneName);
 				BonePoses->SetValue(i, BoneNodes[i].BoneTransform);
 				BoneParentIndices->SetValue(i, BoneNodes[i].ParentBoneIndex);
 			}
-
+ 
 			if (!UpdateStage(EExportMeshStage::InitializingExportMesh, 1.f))
 			{
 				return false;
 			}
 		}
-
+ 
 		if (!bBuildNaniteAssemblies && bShouldExportFoliage)
 		{
 			if (!UpdateStage(EExportMeshStage::CombiningMeshes, 0.f))
 			{
 				return false;
 			}
-
+ 
 			Facades::FFoliageFacade FoliageFacade(ExportedCollection);
-
+ 
 			// Combine all the foliage instances into one dynamic mesh if not building nanite assemblies
 			if (!CombineMeshInstancesToDynamicMesh(
 				ExportedCollection,
@@ -1272,7 +1279,7 @@ namespace PV::Export::Internal
 				{
 					// Assign parent bone id to foliage vertices
 					int32 FoliageParentBoneID = FoliageFacade.GetParentBoneID(InstanceID);
-
+ 
 					for (int VertexID = VertexStart; VertexID < VertexCount; ++VertexID)
 					{
 						VertexBoneIDs.Add(FoliageParentBoneID);
@@ -1284,17 +1291,17 @@ namespace PV::Export::Internal
 				return false;
 			}
 		}
-
+ 
 		// Convert the dynamic mesh to static mesh
 		EGeometryScriptOutcomePins Output;
 		{
 			TRACE_CPUPROFILER_EVENT_SCOPE(CopyMeshToSkeletalMesh);
-
+ 
 			if (!UpdateStage(EExportMeshStage::CopyingMeshDataToOutputMesh, 1.f))
 			{
 				return false;
 			}
-
+ 
 			FGeometryScriptCopyMeshToAssetOptions Options = GetCopyMeshToAssetOptions(Materials, ExportMesh->GetNaniteSettings(), InShapePreservation);
 			Options.BoneHierarchyMismatchHandling = EGeometryScriptBoneHierarchyMismatchHandling::CreateNewReferenceSkeleton;
 			UGeometryScriptLibrary_StaticMeshFunctions::CopyMeshToSkeletalMesh(
@@ -1305,7 +1312,7 @@ namespace PV::Export::Internal
 				Output
 			);
 		}
-
+ 
 		TObjectPtr<USkeleton> Skeleton = ExportMesh->GetSkeleton();
 		{
 			TRACE_CPUPROFILER_EVENT_SCOPE(RecreateBoneTree);
@@ -1315,38 +1322,38 @@ namespace PV::Export::Internal
 			}
 			Skeleton->RecreateBoneTree(ExportMesh);
 		}
-
+ 
 		const FGeometryScriptBoneWeightProfile SkinProfile = FGeometryScriptBoneWeightProfile();
 		const FName ProfileName = SkinProfile.GetProfileName();
-
+ 
 		{
 			TRACE_CPUPROFILER_EVENT_SCOPE(ComputeSmoothBoneWeights);
-
+ 
 			if (!UpdateStage(EExportMeshStage::ComputingBoneWeights, 1.f))
 			{
 				return false;
 			}
-
+ 
 			FGeometryScriptSmoothBoneWeightsOptions SmoothBoneWeightsOptions;
 			SmoothBoneWeightsOptions.DistanceWeighingType = EGeometryScriptSmoothBoneWeightsType::DirectDistance;
 			SmoothBoneWeightsOptions.MaxInfluences = 3;
 			SmoothBoneWeightsOptions.Stiffness = 0.5;
-
+ 
 			UGeometryScriptLibrary_MeshBoneWeightFunctions::ComputeSmoothBoneWeights(
 				DynamicMesh,
 				Skeleton,
 				SmoothBoneWeightsOptions,
 				SkinProfile
 			);
-
+ 
 			if (!UpdateStage(EExportMeshStage::ComputingBoneWeights, 0.8f))
 			{
 				return false;
 			}
-
+ 
 			// Here we make the correction to the weights
 			RemoveUnwantedSkinWeights(DynamicMesh, ProfileName, VertexBoneIDs, FoliageVertexStart, BoneNodes);
-
+ 
 			if (!UpdateStage(EExportMeshStage::ComputingBoneWeights, 1.f))
 			{
 				return false;
@@ -1355,12 +1362,12 @@ namespace PV::Export::Internal
 		
 		{
 			TRACE_CPUPROFILER_EVENT_SCOPE(CopySkinWeightProfileToSkeletalMesh);
-
+ 
 			if (!UpdateStage(EExportMeshStage::CopyingSkinWeightProfiles, 1.f))
 			{
 				return false;
 			}
-
+ 
 			FGeometryScriptCopySkinWeightProfileToAssetOptions Options;
 			Options.bEmitTransaction = false;
 			Options.bDeferMeshPostEditChange = true;
@@ -1374,14 +1381,14 @@ namespace PV::Export::Internal
 				Output
 			);
 		}
-
+ 
 		if (bBuildNaniteAssemblies && bShouldExportFoliage)
 		{
 			if (!UpdateStage(EExportMeshStage::BuildingNaniteAssemblies, 1.f))
 			{
 				return false;
 			}
-
+ 
 			if (NaniteAssembliesSupported())
 			{
 				//Build the nanite assemblies
@@ -1392,26 +1399,26 @@ namespace PV::Export::Internal
 				UE_LOG(LogProceduralVegetationEditor, Warning, TEXT("Failed to build Nanite Assemblies, because neither Nanite Foliage nor Nanite Assemblies are enabled for the project."));
 			}
 		}
-
+ 
 		if (InPhysicsAsset)
 		{
 			SetupPhysicAsset(Collection, CollisionGeneration, InPhysicsAsset, BoneNodes);
 		}
-
+ 
 		{
 			TRACE_CPUPROFILER_EVENT_SCOPE(DynamicWind);
-
+ 
 			if (!UpdateStage(EExportMeshStage::ImportingDynamicWind, 1.f))
 			{
 				return false;
 			}
-
+ 
 			const bool bWindDataExist = ExportMesh->HasAssetUserDataOfClass(UDynamicWindSkeletalData::StaticClass());
 			const bool bOverwriteWindSettings = InWindSettings ? InWindSettings->bOverwriteExisting : false;
 			
 			FDynamicWindSkeletalImportData ImportData;
 			BuildWindImportData(ImportData, ExportedCollection, BoneNodes);
-
+ 
 			if (InWindSettings)
 			{
 				ImportData.SimulationGroups = InWindSettings->SimulationGroupData;
@@ -1426,28 +1433,28 @@ namespace PV::Export::Internal
 					ImportData.SimulationGroups = WindSkeletalData->SimulationGroups;
 				}
 			}
-
+ 
 			if (!DynamicWind::ImportSkeletalData(*ExportMesh, ImportData))
 			{
 				UE_LOG(LogProceduralVegetationEditor, Warning, TEXT("Failed to build wind data"));
 			}
 		}
-
+ 
 		{
 			TRACE_CPUPROFILER_EVENT_SCOPE(BuildSkinnedMesh);
 			if (!UpdateStage(EExportMeshStage::BuildingMesh, 1.f))
 			{
 				return false;
 			}
-
+ 
 			ExportMesh->Build();
 			FSkinnedAssetCompilingManager::Get().FinishAllCompilation();
 		}
-
+ 
 		UpdateStage(EExportMeshStage::Complete, 1.f);
 		return true;
 	}
-
+ 
 	void AttachProceduralVegetationLink(UObject* InExportedMesh, const TObjectPtr<UProceduralVegetation>& InProceduralVegetation)
 	{
 		if (IInterface_AssetUserData* IAssetUserData = Cast<IInterface_AssetUserData>(InExportedMesh))
@@ -1458,16 +1465,16 @@ namespace PV::Export::Internal
 			IAssetUserData->AddAssetUserData(Data);	
 		}
 	}
-
+ 
 	void AssignBoneIDsToFoliage(const TArray<Facades::FBoneNode>& Bones, FManagedArrayCollection& Collection)
 	{
 		TRACE_CPUPROFILER_EVENT_SCOPE(PV::Export::AssignBoneIDsToFoliage);
-
+ 
 		Facades::FPointFacade PointFacade(Collection);
 		Facades::FFoliageFacade FoliageFacade(Collection);
 		
 		int32 NumInstances = FoliageFacade.NumFoliageEntries();
-
+ 
 		auto FindBoneById = ( [&](const int32 Id)
 		{
 			const Facades::FBoneNode* BoneNode = Bones.FindByPredicate([Id](const Facades::FBoneNode& Node)
@@ -1481,7 +1488,7 @@ namespace PV::Export::Internal
 		{
 			Facades::FFoliageEntryData Data = FoliageFacade.GetFoliageEntry(Id);
 			float FoliageLFR = Data.LengthFromRoot;
-
+ 
 			TArray<int32> BoneIDs;
 			for (Facades::FBoneNode BoneNode : Bones)
 			{
@@ -1490,7 +1497,7 @@ namespace PV::Export::Internal
 					BoneIDs.Add(BoneNode.BoneIndex);
 				}
 			}
-
+ 
 			bool bBoneIDAssigned = false;
 			for (const int32& BoneId : BoneIDs)
 			{
@@ -1503,7 +1510,7 @@ namespace PV::Export::Internal
 				int ParentPointIndex = BoneParentNode ? BoneParentNode->PointIndex : INDEX_NONE;
 				float BoneLFR = PointFacade.GetLengthFromRoot(PointIndex);
 				float BoneParentLFR = PointFacade.GetLengthFromRoot(ParentPointIndex);
-
+ 
 				if (FoliageLFR <= BoneLFR && FoliageLFR >= BoneParentLFR)
 				{
 					FoliageFacade.SetParentBoneID(Id, BoneId);
@@ -1511,18 +1518,18 @@ namespace PV::Export::Internal
 					break;
 				}
 			}
-
+ 
 			if (!bBoneIDAssigned)
 			{
 				UE_LOG(LogProceduralVegetationEditor, Warning, TEXT("No bone assigned for foliage Id {%i} , Branch Bone count {%i}"), Id , BoneIDs.Num());
 			}
 		}
 	}
-
+ 
 	void RemoveUnwantedSkinWeights(TObjectPtr<UDynamicMesh> DynamicMesh, const FName ProfileName, const TArray<int32>& VertexBoneIDs, int32 FoliageVertexStart, const TArray<Facades::FBoneNode>& BoneNodes)
 	{
 		TRACE_CPUPROFILER_EVENT_SCOPE(PV::Export::RemoveUnwantedSkinWeights);
-
+ 
 		auto InMesh = DynamicMesh->GetMeshPtr();
 		check(InMesh)
 		const int32 NumVertices = VertexBoneIDs.Num();
@@ -1530,27 +1537,27 @@ namespace PV::Export::Internal
 		DynamicMesh->EditMesh([&](FDynamicMesh3& EditMesh)
 		{
 			UE::Geometry::FDynamicMeshVertexSkinWeightsAttribute *SkinWeights = EditMesh.Attributes()->GetSkinWeightsAttribute(ProfileName);
-
+ 
 			ParallelFor(NumVertices, [&](const int32 VertexIdx)
 			{
 				int32 BoneID = VertexBoneIDs[VertexIdx];
-
+ 
 				if (BoneID != INDEX_NONE)
 				{
 					int32 ParentBoneID = BoneNodes[BoneID].ParentBoneIndex;
 				
 					UE::AnimationCore::FBoneWeights Weights;
 					SkinWeights->GetValue(VertexIdx, Weights);
-
+ 
 					auto NextBone = BoneNodes.FindByPredicate([&](const Facades::FBoneNode& Node)
 					{
 						return Node.BoneIndex == BoneID + 1 && Node.BranchIndex == BoneNodes[BoneID].BranchIndex;
 					});
-
+ 
 					for (int32 WeightID = Weights.Num() - 1 ; WeightID >= 0; WeightID--)
 					{
 						auto Weight = Weights[WeightID];
-
+ 
 						bool bValidBone = (Weight.GetBoneIndex() == ParentBoneID);
 						bValidBone |= (Weight.GetBoneIndex() == BoneID);
 						bValidBone |= NextBone && (Weight.GetBoneIndex() == NextBone->BoneIndex);
@@ -1562,7 +1569,7 @@ namespace PV::Export::Internal
 							Weights.RemoveBoneWeight(Weight.GetBoneIndex(),Settings);
 						}
 					}
-
+ 
 					if (VertexIdx > FoliageVertexStart || BoneID == 0 || Weights.Num() == 0)
 					{
 						UE::AnimationCore::FBoneWeightsSettings Settings;
@@ -1580,7 +1587,7 @@ namespace PV::Export::Internal
 				
 		}, EDynamicMeshChangeType::GeneralEdit, EDynamicMeshAttributeChangeFlags::Unknown, false);
 	}
-
+ 
 	void SetupPhysicAsset(const FManagedArrayCollection& Collection, const EPVCollisionGeneration& CollisionGeneration,UPhysicsAsset* PhysicsAsset, TArray<Facades::FBoneNode> InBones)
 	{
 		if (!PhysicsAsset)
@@ -1588,7 +1595,7 @@ namespace PV::Export::Internal
 			UE_LOG(LogProceduralVegetationEditor, Warning, TEXT("No physics asset specified to create physics bodies"));
 			return;
 		}
-
+ 
 		TArray<TObjectPtr<USkeletalBodySetup>> SkeletalBodySetups;
 		
 		TArray<Facades::FBoneNode> ReducedBones;
@@ -1597,27 +1604,27 @@ namespace PV::Export::Internal
 		Facades::FPlantFacade PlantFacade = Facades::FPlantFacade(Collection);
 		
 		BoneFacade.CreateBoneData(ReducedBones, 1.0, false);
-
+ 
 		for (const Facades::FBoneNode&  Node : ReducedBones)
 		{
 			USkeletalBodySetup* BodySetup = NewObject<USkeletalBodySetup>(PhysicsAsset);
 				
 			auto MappedBone = InBones.FindByPredicate([Node](const Facades::FBoneNode& IterNode){ return Node.PointIndex == IterNode.PointIndex; });
-
+ 
 			if (!MappedBone)
 			{
 				UE_LOG(LogProceduralVegetationEditor, Log, TEXT("MappedBone not found for BodySetup"));
 				continue;
 			}
-
+ 
 			if (!InBones.IsValidIndex(MappedBone->ParentBoneIndex))
 			{
 				UE_LOG(LogProceduralVegetationEditor, Log, TEXT("MapedBone not found for BodySetup"));
 				continue;
 			}
-
+ 
 			bool bTrunk = PlantFacade.IsTrunkIndex(Node.BranchIndex);
-
+ 
 			if (bTrunk || CollisionGeneration == EPVCollisionGeneration::AllGenerations)
 			{
 				BodySetup->BoneName = MappedBone->BoneName;
@@ -1638,13 +1645,13 @@ namespace PV::Export::Internal
 				
 				FVector Direction = (BasePosition - Node.AbsolutePosition);
 				Direction.Normalize();
-
+ 
 				FVector Axis = FVector::UpVector;
 				FQuat Quat = FQuat::FindBetween(Axis, Direction);
-
+ 
 				Element.Center = Direction * Length * 0.5;
 				Element.Rotation = Quat.Rotator();
-
+ 
 				if (Element.Length > 0)
 				{
 					BodySetup->AggGeom.SphylElems.Add(Element);
@@ -1654,14 +1661,14 @@ namespace PV::Export::Internal
 				}
 			}
 		}
-
+ 
 		PhysicsAsset->SkeletalBodySetups = SkeletalBodySetups;
 		PhysicsAsset->UpdateBodySetupIndexMap();
 		PhysicsAsset->UpdateBoundsBodiesArray();
 		FEditorSupportDelegates::RedrawAllViewports.Broadcast();
 	}
 }
-
+ 
 namespace PV::Export
 {
 	EExportResult ExportCollectionAsMesh(
@@ -1673,13 +1680,13 @@ namespace PV::Export
 	)
 	{
 		TRACE_CPUPROFILER_EVENT_SCOPE(PV::Export::ExportCollectionAsMesh);
-
+ 
 		TArray<UPackage*> ModifiedPackages;
-
+ 
 		const FDateTime ExportStartTime = FDateTime::Now();
 		TArray<FAnalyticsEventAttribute> AnalyticsEventAttribute = Analytics::GatherExportCommonAttributes(Collection, ExportParams);
 		Analytics::SendExportStartEvent(AnalyticsEventAttribute);
-
+ 
 		const auto ExitWithResult = [&](EExportResult InResult)->EExportResult
 		{
 			const static auto ExportResultToAnalyticsResult = [](EExportResult InExportResult)->PV::Analytics::EExportResult
@@ -1698,27 +1705,27 @@ namespace PV::Export
 				ensure(false);
 				return PV::Analytics::EExportResult::Failed;
 			};
-
+ 
 			if (InResult != EExportResult::Success)
 			{
 				Internal::CleanupPackages(ModifiedPackages);
 			}
-
+ 
 			const double ExportTime = (FDateTime::Now() - ExportStartTime).GetTotalSeconds();
 			Analytics::SendExportFinishedEvent(MoveTemp(AnalyticsEventAttribute), ExportTime, ExportResultToAnalyticsResult(InResult));
 			return InResult;
 		};
-
+ 
 		const auto ExitWithAssetPathResult = [&](Internal::EGetValidAssetPathResult InAssetPathResult)->EExportResult
 		{
 			if (InAssetPathResult == Internal::EGetValidAssetPathResult::SkipAsset)
 			{
 				return ExitWithResult(EExportResult::Skipped);
 			}
-
+ 
 			return ExitWithResult(EExportResult::Fail);
 		};
-
+ 
 		Internal::EGetValidAssetPathResult AssetPathResult;
 		
 		FString MeshName = ExportParams.MeshName.ToString();
@@ -1731,11 +1738,11 @@ namespace PV::Export
 		
 		const bool bCreateNaniteFoliage = ExportParams.bCreateNaniteFoliage && NaniteAssembliesSupported();
 		const bool bShouldExportFoliage = ExportParams.bShouldExportFoliage;
-
+ 
 		if (ExportParams.ExportMeshType == EPVExportMeshType::StaticMesh)
 		{
 			const TObjectPtr<UStaticMesh> ExportMesh = NewObject<UStaticMesh>(Internal::CreateAssetPackage(MeshPackageName), NAME_None, RF_Public | RF_Standalone);
-
+ 
 			if (!PV::Export::Internal::ExportCollectionToStaticMesh(
 				ExportMesh,
 				Collection,
@@ -1750,12 +1757,12 @@ namespace PV::Export
 				ExportMesh->MarkAsGarbage();
 				return ExitWithResult(EExportResult::Cancelled);
 			}
-
+ 
 			Internal::ReplaceAssetInPackage(ExportMesh);
 			ModifiedPackages.Add(ExportMesh->GetOutermost());
-
+ 
 			OutCreatedAssets.Add(MeshPackageName);
-
+ 
 			Internal::AttachProceduralVegetationLink(ExportMesh, InProceduralVegetation);
 		}
 		else if (ExportParams.ExportMeshType == EPVExportMeshType::SkeletalMesh)
@@ -1767,12 +1774,12 @@ namespace PV::Export
 			{
 				return ExitWithAssetPathResult(AssetPathResult);
 			}
-
+ 
 			const TObjectPtr<USkeletalMesh> ExportMesh = NewObject<USkeletalMesh>(Internal::CreateAssetPackage(MeshPackageName), NAME_None, RF_Standalone | RF_Public);
 			const TObjectPtr<USkeleton> ExportSkeleton = NewObject<USkeleton>(Internal::CreateAssetPackage(SkeletonPackageName), NAME_None, RF_Standalone | RF_Public);
-
+ 
 			ExportMesh->SetSkeleton(ExportSkeleton);
-
+ 
 			FString PhysicsAssetName = ExportParams.GetOutputPhysicsAssetName();
 			FString PhysicsPackageName = ExportParams.GetOutputPhysicsAssetPackagePath();
 			AssetPathResult = Internal::GetValidAssetPath(ExportParams.ReplacePolicy, UPhysicsAsset::StaticClass(), PhysicsPackageName, PhysicsAssetName);
@@ -1780,7 +1787,7 @@ namespace PV::Export
 			{
 				return ExitWithAssetPathResult(AssetPathResult);
 			}
-
+ 
 			TObjectPtr<UPhysicsAsset> PhysicsAsset = nullptr;
 			if (ExportParams.IsCollisionEnable())
 			{
@@ -1788,7 +1795,7 @@ namespace PV::Export
 				ExportMesh->SetPhysicsAsset(PhysicsAsset);
 				PhysicsAsset->SetPreviewMesh(ExportMesh, false);
 			}
-
+ 
 			if (!PV::Export::Internal::ExportCollectionToSkeletalMesh(
 				ExportMesh,
 				Collection,
@@ -1819,7 +1826,7 @@ namespace PV::Export
 			{
 				Internal::ReplaceAssetInPackage(PhysicsAsset);
 				ModifiedPackages.Add(PhysicsAsset->GetOutermost());
-
+ 
 				OutCreatedAssets.Add(PhysicsPackageName);
 				Internal::AttachProceduralVegetationLink(PhysicsAsset, InProceduralVegetation);
 			}
@@ -1831,15 +1838,15 @@ namespace PV::Export
 			
 			OutCreatedAssets.Add(MeshPackageName);
 			OutCreatedAssets.Add(SkeletonPackageName);
-
+ 
 			Internal::AttachProceduralVegetationLink(ExportMesh, InProceduralVegetation);
 			Internal::AttachProceduralVegetationLink(ExportSkeleton, InProceduralVegetation);
 		}
-
+ 
 		for (UPackage* Package : ModifiedPackages)
 		{
 			Package->MarkPackageDirty();
-
+ 
 			if (Package->HasAnyPackageFlags(PKG_NewlyCreated))
 			{
 				UObject* Asset = Package->FindAssetInPackage();
@@ -1849,9 +1856,9 @@ namespace PV::Export
 				}
 			}
 		}
-
+ 
 		return ExitWithResult(EExportResult::Success);
 	}
 }
-
-#undef LOCTEXT_NAMESPACE 
+ 
+#undef LOCTEXT_NAMESPACE 
\ No newline at end of file
diff --git a/Engine/Plugins/Experimental/ProceduralVegetationEditor/Source/ProceduralVegetationEditor/Private/Visualizations/PVDataVisualization.cpp b/Engine/Plugins/Experimental/ProceduralVegetationEditor/Source/ProceduralVegetationEditor/Private/Visualizations/PVDataVisualization.cpp
--- a/Engine/Plugins/Experimental/ProceduralVegetationEditor/Source/ProceduralVegetationEditor/Private/Visualizations/PVDataVisualization.cpp
+++ b/Engine/Plugins/Experimental/ProceduralVegetationEditor/Source/ProceduralVegetationEditor/Private/Visualizations/PVDataVisualization.cpp
@@ -1,7 +1,7 @@
 // Copyright Epic Games, Inc. All Rights Reserved.
-
+ 
 #include "PVDataVisualization.h"
-
+ 
 #include "AdvancedPreviewScene.h"
 #include "EditorViewportClient.h"
 #include "GeometryCollectionToDynamicMesh.h"
@@ -10,47 +10,48 @@
 #include "PlanarCut.h"
 #include "PVBoneComponent.h"
 #include "SPVEditorViewport.h"
-
+ 
 #include "Components/DynamicMeshComponent.h"
 #include "Components/InstancedStaticMeshComponent.h"
 #include "Components/LineBatchComponent.h"
 #include "Components/SkeletalMeshComponent.h"
 #include "Components/StaticMeshComponent.h"
 #include "Components/TextRenderComponent.h"
-
+ 
 #include "DataTypes/PVData.h"
-
+ 
 #include "DebugVisualization/PVDebugVisualizer.h"
-
+ 
 #include "Engine/SkeletalMesh.h"
 #include "Engine/StaticMesh.h"
-
+ 
 #include "Facades/PVFoliageFacade.h"
-
+ 
 #include "GeometryCollection/GeometryCollectionObject.h"
-
+ 
 #include "Helpers/PVUtilities.h"
-
+ 
 #include "Materials/Material.h"
-
+ 
 #include "Nodes/PVBaseSettings.h"
 #include "Nodes/PVObjectAvoidance.h"
-
+ 
 #include "Rendering/SkeletalMeshRenderData.h"
-
+ 
 #include "Visualizations/PVSkeletonVisualizerComponent.h"
-
+#include "ProceduralVegetation/Private/Implementations/PVFoliage.h"
+ 
 #define LOCTEXT_NAMESPACE "PCGManagedArrayCollectionDataVisualization"
-
+ 
 class UPVBoneComponent;
-
+ 
 DEFINE_LOG_CATEGORY(LogProceduralVegetationDataVisualization);
-
+ 
 FPVDataVisualization::FPVDataVisualization()
 {
 	RegisterVisualizations();
 }
-
+ 
 void FPVDataVisualization::RegisterVisualizations()
 {
 	RenderMap.Add(EPVRenderType::PointData, FPVRenderCallback::CreateRaw(this, &FPVDataVisualization::SkeletonRenderer));
@@ -60,13 +61,13 @@ void FPVDataVisualization::RegisterVisualizations()
 	RenderMap.Add(EPVRenderType::FoliageGrid, FPVRenderCallback::CreateRaw(this, &FPVDataVisualization::FoliageGridRenderer));
 	RenderMap.Add(EPVRenderType::FoliageAttachments, FPVRenderCallback::CreateRaw(this, &FPVDataVisualization::FoliageAttachmentPointsRenderer));
 }
-
+ 
 void FPVDataVisualization::ExecuteDebugDisplay(FPCGContext* Context, const UPCGSettingsInterface* SettingsInterface, const UPCGData* Data,
                                                          AActor* TargetActor) const
 {
 	return;
 }
-
+ 
 FPCGSetupSceneFunc FPVDataVisualization::GetViewportSetupFunc(const UPCGSettingsInterface* SettingsInterface, const UPCGData* Data) const
 {
 	return [this, WeakData=TWeakObjectPtr<const UPVData>(Cast<UPVData>(Data)), WeakSettings = TWeakObjectPtr<const UPVBaseSettings>(Cast<UPVBaseSettings>(SettingsInterface))](
@@ -78,12 +79,12 @@ FPCGSetupSceneFunc FPVDataVisualization::GetViewportSetupFunc(const UPCGSettings
 			UE_LOG(LogProceduralVegetationDataVisualization, Error, TEXT("Failed to setup data viewport, the data was lost or invalid."));
 			return;
 		}
-
+ 
 		InOutParams.Scene->SetFloorVisibility(true);
 		InOutParams.Scene->SetFloorOffset(1.0f);
-
+ 
 		InOutParams.EditorViewportClient->OverrideNearClipPlane(0.001f);
-
+ 
 		// Reapply view mode settings as for some reason SetFloorVisibility changes EngineFlags
 		InOutParams.EditorViewportClient->SetViewMode(InOutParams.EditorViewportClient->GetViewMode());
 		InOutParams.EditorViewportClient->SetRealtime(true);
@@ -99,11 +100,11 @@ FPCGSetupSceneFunc FPVDataVisualization::GetViewportSetupFunc(const UPCGSettings
 		{
 			InOutParams.Scene->HandleTogglePostProcessing();
 		}
-
+ 
 		const FManagedArrayCollection& Collection = WeakData->GetCollection();
-
+ 
 		FBoxSphereBounds Bounds(ForceInit);
-
+ 
 		for (const EPVRenderType& RenderType : WeakSettings->GetCurrentRenderTypes())
 		{
 			if (RenderMap.Contains(RenderType))
@@ -112,7 +113,7 @@ FPCGSetupSceneFunc FPVDataVisualization::GetViewportSetupFunc(const UPCGSettings
 				RenderMap[RenderType].Execute(Collection, InOutParams, Bounds);
 			}
 		}
-
+ 
 		if (const UPVObjectAvoidanceSettings* ObjectAvoidanceSettings = Cast<UPVObjectAvoidanceSettings>(WeakSettings))
 		{
 			for (const auto& [Mesh, Transform] : ObjectAvoidanceSettings->ObjectAvoidanceSettings.Colliders)
@@ -120,12 +121,12 @@ FPCGSetupSceneFunc FPVDataVisualization::GetViewportSetupFunc(const UPCGSettings
 				TObjectPtr<UStaticMeshComponent> StaticMeshComponent = NewObject<UStaticMeshComponent>(GetTransientPackage(), NAME_None, RF_Transient);
 				StaticMeshComponent->SetStaticMesh(Mesh.LoadSynchronous());
 				StaticMeshComponent->UpdateBounds();
-
+ 
 				InOutParams.ManagedResources.Add(StaticMeshComponent);
 				InOutParams.Scene->AddComponent(StaticMeshComponent, Transform);
 			}
 		}
-
+ 
 		if (PV::Utilities::DebugModeEnabled())
 		{
 			if (WeakSettings.IsValid())
@@ -137,16 +138,16 @@ FPCGSetupSceneFunc FPVDataVisualization::GetViewportSetupFunc(const UPCGSettings
 				}
 			}
 		}
-
+ 
 		InOutParams.FocusBounds = FBoxSphereBounds(Bounds.Origin, Bounds.BoxExtent, Bounds.SphereRadius);
 	};
 }
-
+ 
 void FPVDataVisualization::MeshRenderer(const FManagedArrayCollection& InCollection, FPCGSceneSetupParams& InOutParams,
                                                   FBoxSphereBounds& OutBounds)
 {
 	TRACE_CPUPROFILER_EVENT_SCOPE(FPVDataVisualization::MeshRenderer);
-
+ 
 	TArray<UMaterialInterface*> Materials;
 	
 	if (InCollection.HasAttribute("MaterialPath", FGeometryCollection::MaterialGroup))
@@ -157,7 +158,7 @@ void FPVDataVisualization::MeshRenderer(const FManagedArrayCollection& InCollect
 			Materials.Add(LoadObject<UMaterialInterface>(nullptr, MaterialPath));
 		}
 	}
-
+ 
 	UE::Geometry::FGeometryCollectionToDynamicMeshes Converter;
 	Converter.Init(InCollection, {});
 	for (const auto& MeshInfo : Converter.Meshes)
@@ -170,51 +171,51 @@ void FPVDataVisualization::MeshRenderer(const FManagedArrayCollection& InCollect
 			DynamicMeshComponent->SetTangentsType(EDynamicMeshComponentTangentsMode::AutoCalculated);
 			DynamicMeshComponent->ConfigureMaterialSet(Materials);
 			DynamicMeshComponent->UpdateBounds();
-
+ 
 			InOutParams.ManagedResources.Add(DynamicMeshComponent);
 			InOutParams.Scene->AddComponent(DynamicMeshComponent, FTransform::Identity);
-
+ 
 			OutBounds = OutBounds + DynamicMeshComponent->CalcLocalBounds();
 		}
 	}
 }
-
+ 
 void FPVDataVisualization::SkeletonRenderer(const FManagedArrayCollection& InCollection, FPCGSceneSetupParams& InOutParams,
                                                       FBoxSphereBounds& OutBounds)
 {
 	TRACE_CPUPROFILER_EVENT_SCOPE(FPVDataVisualization::SkeletonRenderer);
-
+ 
 	const TObjectPtr<UPVSkeletonVisualizerComponent> VisualizerComponent = NewObject<UPVSkeletonVisualizerComponent>(
 		GetTransientPackage(), NAME_None, RF_Transient);
 	VisualizerComponent->SetCollection(&InCollection);
 	VisualizerComponent->UpdateBounds();
-
+ 
 	InOutParams.ManagedResources.Add(VisualizerComponent);
 	InOutParams.Scene->AddComponent(VisualizerComponent, FTransform::Identity);
-
+ 
 	OutBounds = OutBounds + VisualizerComponent->CalcLocalBounds();
 }
-
+ 
 void FPVDataVisualization::FoliageRenderer(const FManagedArrayCollection& InCollection, FPCGSceneSetupParams& InOutParams,
                                                      FBoxSphereBounds& OutBounds)
 {
 	TRACE_CPUPROFILER_EVENT_SCOPE(FPVDataVisualization::FoliageRenderer);
-
+ 
 	PV::Facades::FFoliageFacade FoliageFacade(InCollection);
-
+ 
 	auto OnComponentCreated = [&InOutParams](UMeshComponent* InInstancedComponent)
 	{
 		InOutParams.ManagedResources.Add(InInstancedComponent);
 		InInstancedComponent->RegisterComponentWithWorld(InOutParams.Scene->GetWorld());
 		InOutParams.Scene->AddComponent(InInstancedComponent, FTransform::Identity);
 	};
-
+ 
 	TMap<FString, TObjectPtr<UMeshComponent>> OutInstancedComponentMap;
-
+ 
 	PV::Utilities::AddFoliageInstances(FoliageFacade, nullptr,
 	                                           PV::Utilities::FFoliageComponentCreatedCallback::CreateLambda(OnComponentCreated),
 	                                           OutInstancedComponentMap);
-
+ 
 	const TSharedPtr<SPVEditorViewport> ViewportWidget = StaticCastSharedPtr<SPVEditorViewport>(
 			InOutParams.EditorViewportClient->GetEditorViewportWidget());
 	for (const auto& [FoliageName , Component] : OutInstancedComponentMap)
@@ -223,23 +224,23 @@ void FPVDataVisualization::FoliageRenderer(const FManagedArrayCollection& InColl
 		OutBounds = OutBounds + Component->CalcLocalBounds();
 	}
 }
-
+ 
 void FPVDataVisualization::BonesRenderer(const FManagedArrayCollection& InCollection, FPCGSceneSetupParams& InOutParams,
 	FBoxSphereBounds& OutBounds)
 {
 	TRACE_CPUPROFILER_EVENT_SCOPE(FPVDataVisualization::BonesRenderer);
-
+ 
 	const TObjectPtr<UPVBoneComponent> BoneComponent = NewObject<UPVBoneComponent>(
 		GetTransientPackage(), NAME_None, RF_Transient);
 	BoneComponent->SetCollection(&InCollection);
 	BoneComponent->UpdateBounds();
-
+ 
 	InOutParams.ManagedResources.Add(BoneComponent);
 	InOutParams.Scene->AddComponent(BoneComponent, FTransform::Identity);
-
+ 
 	FString InTextToDraw = FString::Format(TEXT("{0} Bones"),{BoneComponent->GetBoneCount()});
 	const FVector3f& Pos = FVector3f::ForwardVector * 100;
-
+ 
 	DrawDebugString(BoneComponent->GetWorld(), FVector::Zero(), InTextToDraw);
 	
 	TObjectPtr<UTextRenderComponent> Component = NewObject<UTextRenderComponent>(GetTransientPackage(), NAME_None, RF_Transient);
@@ -253,10 +254,10 @@ void FPVDataVisualization::BonesRenderer(const FManagedArrayCollection& InCollec
 	FRotator TextRotator = FRotator(0,90,0);
 	TextTransform.SetRotation(TextRotator.Quaternion());
 	InOutParams.Scene->AddComponent(Component, TextTransform);
-
+ 
 	OutBounds = OutBounds + BoneComponent->CalcLocalBounds();
 }
-
+ 
 TPair<int32, int32> GetMeshTrisAndVerts(const UObject* FoliageObject)
 {
 	int32 NumTriangles = 0;
@@ -277,17 +278,17 @@ TPair<int32, int32> GetMeshTrisAndVerts(const UObject* FoliageObject)
 	}
 	return {NumTriangles, NumVertices};
 }
-
+ 
 void FPVDataVisualization::FoliageGridRenderer(const FManagedArrayCollection& InCollection, FPCGSceneSetupParams& InOutParams,
 	FBoxSphereBounds& OutBounds)
 {
 	const PV::Facades::FFoliageFacade FoliageFacade(InCollection);
-
-	const TArray<FString> FoliageNames = FoliageFacade.GetFoliageNames();
-
+ 
+	const TArray<FPVFoliageInfo> FoliageInfos = FoliageFacade.GetFoliageInfos();
+ 
 	FBox MaxFoliageBounds(ForceInit);
 	int ValidFoliageEntries = 0;
-
+ 
 	static const auto GetMeshBounds = [](auto* Mesh)-> FBox
 		{
 			const FBox BoundingBox = Mesh
@@ -295,21 +296,23 @@ void FPVDataVisualization::FoliageGridRenderer(const FManagedArrayCollection& In
 				: FBox(ForceInit);
 			return BoundingBox;
 		};
-
+ 
 	const auto AddComponentToScene = [&](USceneComponent* ComponentRef, const FTransform& Transform, const bool UpdateBounds = true)
 		{
 			InOutParams.ManagedResources.Add(ComponentRef);
 			InOutParams.Scene->AddComponent(ComponentRef, Transform);
-
+ 
 			if (UpdateBounds)
 			{
 				ComponentRef->UpdateBounds();
 				OutBounds = OutBounds + ComponentRef->CalcBounds(Transform);
 			}
 		};
-
-	for (const FString& FoliageName : FoliageNames)
+ 
+	for (const FPVFoliageInfo& Info : FoliageInfos)
 	{
+		FString FoliageName = Info.Mesh.ToString();
+		
 		if (UObject* FoliageMesh = LoadObject<UObject>(nullptr, *FoliageName))
 		{
 			FBox Bounds = GetMeshBounds(Cast<UStaticMesh>(FoliageMesh));
@@ -328,11 +331,13 @@ void FPVDataVisualization::FoliageGridRenderer(const FManagedArrayCollection& In
 	
 	const float FoliageHorizontalPadding = MaxFoliageBounds.GetSize().GetMax() + 100.0f;
 	const float FoliageVerticalPadding = MaxFoliageBounds.GetSize().GetMax() + 100.0f;
-
-	FVector FoliagePos = FVector(FoliageHorizontalPadding * FoliageNames.Num() / -2.0f, 0.0f, 0.0f);
-
-	for (const FString& FoliagePath : FoliageNames)
+ 
+	FVector FoliagePos = FVector(FoliageHorizontalPadding * FoliageInfos.Num() / -2.0f, 0.0f, 0.0f);
+ 
+	for (const FPVFoliageInfo& Info : FoliageInfos)
 	{
+		const FString& FoliagePath = Info.Mesh.ToString();
+		
 		UObject* FoliageObject = LoadObject<UObject>(nullptr, FoliagePath);
 		if (UStaticMesh* StaticFoliageMesh = Cast<UStaticMesh>(FoliageObject))
 		{
@@ -346,7 +351,7 @@ void FPVDataVisualization::FoliageGridRenderer(const FManagedArrayCollection& In
 			MeshComponent->SetSkeletalMesh(SkeletalFoliageMesh);
 			AddComponentToScene(MeshComponent, FTransform(FoliagePos));
 		}
-
+ 
 		if (FoliageObject)
 		{
 			const auto [Tris, Verts] = GetMeshTrisAndVerts(FoliageObject);
@@ -367,16 +372,16 @@ void FPVDataVisualization::FoliageGridRenderer(const FManagedArrayCollection& In
 				FTransform(FRotator(0, 90, 0), TextPos),
 				false
 			);
-
+ 
 			TextBounds = TextBounds.TransformBy(TextRenderComponent->GetComponentTransform());
-
+ 
 			if (TextBackgroundMesh && BackgroundMat)
 			{
 				UStaticMeshComponent* TextBackgroundMeshComponent = NewObject<UStaticMeshComponent>(GetTransientPackage(), NAME_None, RF_Transient);
 				TextBackgroundMeshComponent->SetStaticMesh(TextBackgroundMesh);
 				TextBackgroundMeshComponent->SetMaterial(0, BackgroundMat);
-
-
+ 
+ 
 				AddComponentToScene(
 					TextBackgroundMeshComponent,
 					FTransform(
@@ -392,19 +397,19 @@ void FPVDataVisualization::FoliageGridRenderer(const FManagedArrayCollection& In
 		}
 	}
 }
-
+ 
 void FPVDataVisualization::FoliageAttachmentPointsRenderer(const FManagedArrayCollection& InCollection, FPCGSceneSetupParams& InOutParams,
 	FBoxSphereBounds& OutBounds)
 {
 	const PV::Facades::FFoliageFacade FoliageFacade(InCollection);
-
+ 
 	const TObjectPtr<UPVLineBatchComponent> LineBatchComponent = NewObject<UPVLineBatchComponent>(GetTransientPackage(), NAME_None, RF_Transient);
-
+ 
 	const TObjectPtr<UInstancedStaticMeshComponent> SphereInstancer = NewObject<UInstancedStaticMeshComponent>(GetTransientPackage(), NAME_None, RF_Transient);
 	SphereInstancer->SetStaticMesh(LoadObject<UStaticMesh>(nullptr, TEXT("/Engine/BasicShapes/Sphere.Sphere")));
 	SphereInstancer->SetMaterial(0, LoadObject<UMaterialInterface>(nullptr, TEXT("/ProceduralVegetationEditor/Materials/FoliageAttachmentPointMaterial.FoliageAttachmentPointMaterial")));
 	SphereInstancer->SetNumCustomDataFloats(3);
-
+ 
 	TArray<TPair<FColor, FText>> LegendItems;
 	const TArray<FString> FoliageNames = FoliageFacade.GetFoliageNames();
 	for (int32 FoliageIndex = 0; FoliageIndex < FoliageNames.Num(); FoliageIndex++)
@@ -412,12 +417,12 @@ void FPVDataVisualization::FoliageAttachmentPointsRenderer(const FManagedArrayCo
 		const FLinearColor Color = PV::Utilities::GetRandomHueColor(static_cast<float>(FoliageIndex) / static_cast<float>(FoliageNames.Num()));
 		LegendItems.Add({Color.ToFColorSRGB(), FText::FromString(FSoftObjectPath(FoliageNames[FoliageIndex]).GetAssetName())});
 	}
-
+ 
 	FBox Bounds;
 	for (int32 PointIndex = 0; PointIndex < FoliageFacade.GetElementCount(); PointIndex++)
 	{
 		const PV::Facades::FFoliageEntryData FoliageData = FoliageFacade.GetFoliageEntry(PointIndex);
-
+ 
 		Bounds += static_cast<FVector>(FoliageData.PivotPoint);
 		const FLinearColor Color = PV::Utilities::GetRandomHueColor(static_cast<float>(FoliageData.NameId) / static_cast<float>(FoliageNames.Num()));
 		LineBatchComponent->AddLine(
@@ -427,7 +432,7 @@ void FPVDataVisualization::FoliageAttachmentPointsRenderer(const FManagedArrayCo
 			SDPG_World,
 			EPointDrawSettings::None
 		);
-
+ 
 		const int32 InstanceIndex = SphereInstancer->AddInstance(
 			FTransform(
 				FRotator(),
@@ -439,19 +444,19 @@ void FPVDataVisualization::FoliageAttachmentPointsRenderer(const FManagedArrayCo
 		SphereInstancer->SetCustomDataValue(InstanceIndex, 1, Color.G, false);
 		SphereInstancer->SetCustomDataValue(InstanceIndex, 2, Color.B, false);
 	}
-
+ 
 	SphereInstancer->MarkRenderStateDirty();
-
+ 
 	InOutParams.ManagedResources.Add(SphereInstancer);
 	InOutParams.Scene->AddComponent(SphereInstancer, FTransform::Identity);
 	
 	InOutParams.ManagedResources.Add(LineBatchComponent);
 	InOutParams.Scene->AddComponent(LineBatchComponent, FTransform::Identity);
-
+ 
 	TSharedPtr<SPVEditorViewport> ViewportWidget = StaticCastSharedPtr<SPVEditorViewport>(InOutParams.EditorViewportClient->GetEditorViewportWidget());
 	ViewportWidget->PopulateLegendOverlayText(LegendItems);
 	
 	OutBounds = OutBounds + Bounds;
 }
-
-#undef LOCTEXT_NAMESPACE
+ 
+#undef LOCTEXT_NAMESPACE
\ No newline at end of file
