[PVE] Fix horde errors

diff --git a/Engine/Plugins/Experimental/ProceduralVegetationEditor/Source/ProceduralVegetation/Private/Facades/PVFoliageFacade.cpp b/Engine/Plugins/Experimental/ProceduralVegetationEditor/Source/ProceduralVegetation/Private/Facades/PVFoliageFacade.cpp
--- a/Engine/Plugins/Experimental/ProceduralVegetationEditor/Source/ProceduralVegetation/Private/Facades/PVFoliageFacade.cpp
+++ b/Engine/Plugins/Experimental/ProceduralVegetationEditor/Source/ProceduralVegetation/Private/Facades/PVFoliageFacade.cpp
@@ -452,7 +452,7 @@ namespace PV::Facades
 		return FString();
 	}
  
-	const TArray<FString>& FFoliageFacade::GetFoliageNames() const
+	TArray<FString> FFoliageFacade::GetFoliageNames() const
 	{
 		TArray<FString> FoliageNames;
 		
diff --git a/Engine/Plugins/Experimental/ProceduralVegetationEditor/Source/ProceduralVegetation/Public/Facades/PVFoliageFacade.h b/Engine/Plugins/Experimental/ProceduralVegetationEditor/Source/ProceduralVegetation/Public/Facades/PVFoliageFacade.h
--- a/Engine/Plugins/Experimental/ProceduralVegetationEditor/Source/ProceduralVegetation/Public/Facades/PVFoliageFacade.h
+++ b/Engine/Plugins/Experimental/ProceduralVegetationEditor/Source/ProceduralVegetation/Public/Facades/PVFoliageFacade.h
@@ -70,7 +70,7 @@ namespace PV::Facades
 		
 		const FString GetFoliageName(int32 Index) const;
  
-		const TArray<FString>& GetFoliageNames() const;
+		TArray<FString> GetFoliageNames() const;
  
 		int32 NumFoliageInfo() const { return FoliageInfoNamesAttribute.Num(); };
  
