[SPVEditorViewport]
- Fix crash when running ASAN on QAGCookedEditor by static initializing on-demand

#Synced-CL 49408961

diff --git a/Engine/Plugins/Experimental/ProceduralVegetationEditor/Source/ProceduralVegetationEditor/Private/SPVEditorViewport.cpp b/Engine/Plugins/Experimental/ProceduralVegetationEditor/Source/ProceduralVegetationEditor/Private/SPVEditorViewport.cpp
--- a/Engine/Plugins/Experimental/ProceduralVegetationEditor/Source/ProceduralVegetationEditor/Private/SPVEditorViewport.cpp
+++ b/Engine/Plugins/Experimental/ProceduralVegetationEditor/Source/ProceduralVegetationEditor/Private/SPVEditorViewport.cpp
@@ -40,7 +40,9 @@
 
 #define LOCTEXT_NAMESPACE "PVEditorViewportToolbarSections"
 
-static FTableViewStyle GTransparentViewStyle = []
+static FTableViewStyle& GetTransparentViewStyle()
+{
+	static FTableViewStyle TransparentViewStyle = []
 	{
 		const FSlateColorBrush Brush(FLinearColor::Transparent);
 
@@ -49,28 +51,34 @@ static FTableViewStyle GTransparentViewStyle = []
 
 		return ViewStyle;
 	}();
+	return TransparentViewStyle;
+}
 
-static FTableRowStyle GTransparentRowStyle = []
-	{
-		FTableRowStyle Style = FAppStyle::Get().GetWidgetStyle<FTableRowStyle>("TableView.Row");
+static FTableRowStyle& GetTransparentRowStyle()
+{
+	static FTableRowStyle TransparentRowStyle = []
+		{
+			FTableRowStyle Style = FAppStyle::Get().GetWidgetStyle<FTableRowStyle>("TableView.Row");
 
-		const FSlateColorBrush Brush(FLinearColor::Transparent);
+			const FSlateColorBrush Brush(FLinearColor::Transparent);
 
-		Style.SetEvenRowBackgroundBrush(Brush)
-		     .SetEvenRowBackgroundHoveredBrush(Brush)
-		     .SetOddRowBackgroundBrush(Brush)
-		     .SetOddRowBackgroundHoveredBrush(Brush)
-		     .SetSelectorFocusedBrush(Brush)
-		     .SetActiveBrush(Brush)
-		     .SetActiveHoveredBrush(Brush)
-		     .SetInactiveBrush(Brush)
-		     .SetInactiveHoveredBrush(Brush)
-		     .SetDropIndicator_Above(Brush)
-		     .SetDropIndicator_Onto(Brush)
-		     .SetDropIndicator_Below(Brush);
-
-		return Style;
-	}();
+			Style.SetEvenRowBackgroundBrush(Brush)
+				 .SetEvenRowBackgroundHoveredBrush(Brush)
+				 .SetOddRowBackgroundBrush(Brush)
+				 .SetOddRowBackgroundHoveredBrush(Brush)
+				 .SetSelectorFocusedBrush(Brush)
+				 .SetActiveBrush(Brush)
+				 .SetActiveHoveredBrush(Brush)
+				 .SetInactiveBrush(Brush)
+				 .SetInactiveHoveredBrush(Brush)
+				 .SetDropIndicator_Above(Brush)
+				 .SetDropIndicator_Onto(Brush)
+				 .SetDropIndicator_Below(Brush);
+
+			return Style;
+		}();
+	return TransparentRowStyle;
+}
 
 SPVEditorViewport::SPVEditorViewport()
 {
@@ -271,12 +279,12 @@ void SPVEditorViewport::PopulateViewportOverlays(TSharedRef<SOverlay> Overlay)
 				[
 					SNew(SListView<TSharedPtr<FText>>)
 					.ListItemsSource(&LegendData)
-					.ListViewStyle(&GTransparentViewStyle)
+					.ListViewStyle(&GetTransparentViewStyle())
 					.OnGenerateRow_Lambda(
 						[](const TSharedPtr<FText>& Item, const TSharedRef<STableViewBase>& OwningTable) -> TSharedRef<ITableRow>
 							{
 								return SNew(STableRow<TSharedPtr<FText>>, OwningTable)
-									.Style(&GTransparentRowStyle)
+									.Style(&GetTransparentRowStyle())
 									.ShowSelection(false)
 									.Padding(FMargin(0))
 									[
