Implement remaining features of image tracing algorithm from APIX

- Replace custom marching squares algorithm with implementation from PCGSpatialAlgo
- Add WhiteLevel and bInvertImage to image sampling to align with Houdini algorithm
- Filter out small curves
- Add curve smoothing
- Add curve simplification

#Synced-CL 49615590

diff --git a/Engine/Plugins/Experimental/ProceduralVegetationEditor/Source/ProceduralVegetation/Public/PVImportTexture2DParams.h b/Engine/Plugins/Experimental/ProceduralVegetationEditor/Source/ProceduralVegetation/Public/PVImportTexture2DParams.h
--- a/Engine/Plugins/Experimental/ProceduralVegetationEditor/Source/ProceduralVegetation/Public/PVImportTexture2DParams.h
+++ b/Engine/Plugins/Experimental/ProceduralVegetationEditor/Source/ProceduralVegetation/Public/PVImportTexture2DParams.h
@@ -13,7 +13,10 @@ enum class EPVImportTexture2DDebugState : uint8
 	None,
 	Texture,
 	SampleImage,
-	ExtractContourLines,
+	TraceImage,
+	RemoveSmallCurves,
+	SmoothCurves,
+	SimplifyCurves,
 	MAX UMETA(Hidden)
 };
 
@@ -22,15 +25,32 @@ struct FPVImportTexture2DParams
 {
 	GENERATED_BODY()
 
-	UPROPERTY(EditAnywhere, Category = "Asset", meta=(ClampMin=0, UIMin=0))
+	UPROPERTY(EditAnywhere, Category = "Image Trace", meta=(ClampMin=0, UIMin=0))
 	int32 SampleWidth = 128;
 	
-	UPROPERTY(EditAnywhere, Category = "Asset", meta=(ClampMin=0, UIMin=0))
+	UPROPERTY(EditAnywhere, Category = "Image Trace", meta=(ClampMin=0, UIMin=0))
     int32 SampleHeight = 128;
+    
+    // Any branches with a bounds area lower than this value will be ignored.
+    // Value is normalized meaning a value of 0.1 will filter out any branch whose bounds cover less than 10% of the image. 
+	UPROPERTY(EditAnywhere, Category = "Image Trace", meta=(ClampMin=0, UIMin=0))
+	float MinBoundsArea = 0.05f;
+	
+	UPROPERTY(EditAnywhere, Category = "Image Trace", meta=(ClampMin=0, UIMin=0, UIMax=10, ClampMax=100))
+	int32 SmoothingIterations = 1;
+	
+	UPROPERTY(EditAnywhere, Category = "Image Trace", meta=(ClampMin=0, UIMin=0, UIMax=0.01, ClampMax=1))
+	float SimplificationEpsilon = 0.001;
 
-	UPROPERTY(EditAnywhere, Category = "Asset")
+	UPROPERTY(EditAnywhere, Category = "Texture")
 	TObjectPtr<UTexture2D> Texture2DAsset = nullptr;
+	
+	UPROPERTY(EditAnywhere, Category = "Texture")
+	bool bInvertImage = false;
+	
+	UPROPERTY(EditAnywhere, Category = "Texture", meta=(ClampMin=0, UIMin=0, UIMax=1, ClampMax=1))
+	float WhiteLevel = 0.01;
 
-	UPROPERTY(EditAnywhere, Category = "Asset")
+	UPROPERTY(EditAnywhere, Category = "Inspect")
 	EPVImportTexture2DDebugState DebugState = EPVImportTexture2DDebugState::None;
 };
diff --git a/Engine/Plugins/Experimental/ProceduralVegetationEditor/Source/ProceduralVegetationEditor/Private/Importer/PVImporter_Texture2D.cpp b/Engine/Plugins/Experimental/ProceduralVegetationEditor/Source/ProceduralVegetationEditor/Private/Importer/PVImporter_Texture2D.cpp
--- a/Engine/Plugins/Experimental/ProceduralVegetationEditor/Source/ProceduralVegetationEditor/Private/Importer/PVImporter_Texture2D.cpp
+++ b/Engine/Plugins/Experimental/ProceduralVegetationEditor/Source/ProceduralVegetationEditor/Private/Importer/PVImporter_Texture2D.cpp
@@ -11,6 +11,7 @@
 
 #include "GeometryCollection/GeometryCollection.h"
 #include "GeometryCollection/GeometryCollectionUtility.h"
+#include "SpatialAlgo/PCGMarchingSquares.h"
 
 namespace Internal
 {
@@ -20,12 +21,12 @@ namespace Internal
 	const static FName PixelsDebugGroupName = TEXT("PixelDebug");
 	const static FName PixelPositionsAttributeName = TEXT("Position");
 	const static FName PixelValueAttributeName = TEXT("Value");
-	const static FName LineSegmentsGroupName = TEXT("LineSegments");
-	const static FName LineSegmentsAttributeName = TEXT("LineSegment");
+	const static FName CurveGroupName = TEXT("Curves");
+	const static FName CurveAttributeName = TEXT("Curve");
 	
-	static int32 GetPixelIndex(const FVector2i& PixelCoors, int32 ImageWidth)
+	static int32 GetPixelIndex(const FVector2i& PixelCoords, int32 ImageWidth)
 	{
-		return PixelCoors.X + PixelCoors.Y * ImageWidth;
+		return PixelCoords.X + PixelCoords.Y * ImageWidth;
 	}
 	
 	template<typename T>
@@ -66,18 +67,19 @@ namespace Internal
 		}
 	}
 
-	static void AddDebugLineSegmentsToCollection(FManagedArrayCollection& InOutCollection, const TArray<TPair<FVector2f, FVector2f>>& InLineSegments)
-	{
-		InOutCollection.AddGroup(LineSegmentsGroupName);
-		InOutCollection.AddElements(InLineSegments.Num(), LineSegmentsGroupName);
-		
-		auto& LineSegmentAttribute = InOutCollection.AddAttribute<TArray<FVector3f>>(LineSegmentsAttributeName, LineSegmentsGroupName);
-		for (int32 i = 0; i < InLineSegments.Num(); i++)
+	static void AddDebugCurvesToCollection(FManagedArrayCollection& InOutCollection, const TArray<TArray<FVector2f>>& InCurves)
+	{	
+		InOutCollection.AddGroup(CurveGroupName);
+		InOutCollection.AddElements(InCurves.Num(), CurveGroupName);
+    	
+		auto& CurveAttribute = InOutCollection.AddAttribute<TArray<FVector3f>>(CurveAttributeName, CurveGroupName);
+		for (int32 i = 0; i < InCurves.Num(); i++)
 		{
-			const auto& It = InLineSegments[i];
-			const FVector2f LineStart = It.Key * DebugQuadSize;
-			const FVector2f LineEnd = It.Value * DebugQuadSize;
-			LineSegmentAttribute[i] = { FVector3f(0, LineStart.X, DebugQuadSize - LineStart.Y), FVector3f(0, LineEnd.X, DebugQuadSize - LineEnd.Y) };
+			const TArray<FVector2f>& Curve = InCurves[i];
+			for (const FVector2f& CurvePos : Curve)
+			{
+				CurveAttribute[i].Emplace(0, CurvePos.X * DebugQuadSize, DebugQuadSize - (CurvePos.Y * DebugQuadSize));
+			}
 		}
 	}
 
@@ -134,7 +136,7 @@ namespace Internal
 		return MeshGeomCollection.ToSharedRef();
 	}
 
-	static void SampleImage(FTextureSource& InTextureSource, FVector2i& InOutSampleSize, TArray<bool>& OutBlackAndWhiteImageData)
+	static void SampleImage(FTextureSource& InTextureSource, FVector2i& InOutSampleSize, TArray<bool>& OutBlackAndWhiteImageData, bool bInvertImage, float WhiteLevel)
 	{
 		const FTextureSource::FMipLock LockedMip0(FTextureSource::ELockState::ReadOnly, &InTextureSource, 0);
 		check(LockedMip0.IsValid());
@@ -158,69 +160,179 @@ namespace Internal
 				check(SamplePosY >= 0 && SamplePosY < ImageHeight);
 				
 				const FLinearColor PixelColor = LockedMip0.Image.GetOnePixelLinear(SamplePosX, SamplePosY);
-				const float ColorIntensity = PixelColor.R + PixelColor.G + PixelColor.B + PixelColor.A;
-
 				const int32 PixelIndex = GetPixelIndex(PixelCoords, InOutSampleSize.X);
-				OutBlackAndWhiteImageData[PixelIndex] = ColorIntensity > 0;
+				const float Alpha = bInvertImage ? 1.0 - PixelColor.R : PixelColor.R;
+				OutBlackAndWhiteImageData[PixelIndex] = Alpha > WhiteLevel;
 			}
 		}
 	}
-
-	static void ExtractContourLines(const TArray<bool>& ImageData, const FVector2i& ImageSize, TArray<TPair<FVector2f, FVector2f>>& OutLineSegments)
+	
+	static void TraceImage(const TArray<bool>& ImageData, const FVector2i& ImageSize, TArray<TArray<FVector2f>>& OutCurves)
 	{
-		// Extracts the contour lines of the image data using marching squares (https://en.wikipedia.org/wiki/Marching_squares)
+		TRACE_CPUPROFILER_EVENT_SCOPE(TraceImage);
 		
-		const FVector2f Positions[] = {
-			{ 0.0, 0.5 }, { 0.5, 0.0 }, 
-			{ 0.5, 0.0 }, { 1.0, 0.5 }, 
-			{ 0.5, 1.0 }, { 1.0, 0.5 }, 
-			{ 0.0, 0.5 }, { 0.5, 1.0 }, 
-			{ 0.5, 1.0 }, { 0.5, 0.0 }, 
-			{ 0.0, 0.5 }, { 1.0, 0.5 }, 
-			{ 0.5, 0.0 }, { 1.0, 0.5 }, { 0.0, 0.5 }, { 0.5, 1.0 }, 
-			{ 0.0, 0.5 }, { 0.5, 0.0 }, { 0.5, 1.0 }, { 1.0, 0.5 }
+		const auto ValueQuery = [&](int32 X, int32 Y)->double 
+		{ 
+			return ImageData[GetPixelIndex(FVector2i(X, Y), ImageSize.X)]; 
 		};
+		const TArray<PCGSpatialAlgo::FPCGMarchingSquareResult> Results = PCGSpatialAlgo::MarchingSquares(
+			ImageSize.X, 
+			ImageSize.Y,
+			1, 
+			ValueQuery, 
+			false
+		);
+		
+		OutCurves.Reserve(Results.Num());
+		for (const PCGSpatialAlgo::FPCGMarchingSquareResult& Result : Results)
+		{
+			if (!Result.bClosed)
+			{
+				continue;
+			}
+			
+			TArray<FVector2f>& OutCurve = OutCurves.AddDefaulted_GetRef();
+			OutCurve.Reserve(Result.LinkedGridCoordinates.Num());
+			for (const FVector2D& Coordinate : Result.LinkedGridCoordinates)
+			{
+				OutCurve.Add(FVector2f(Coordinate.X / (ImageSize.X - 1), Coordinate.Y / (ImageSize.Y - 1)));
+			}
+		}
+	}
 	
-		constexpr int32 Starts[] = { 0, 0, 2, 10, 4, 16, 8, 6, 6, 8, 12, 4, 10, 2, 0, 0 };
-		constexpr int32 Sizes[] =  { 0, 2, 2, 2,  2,  4, 2, 2, 2, 2,  4, 2,  2, 2, 2, 0 };
-
-		const int32 Width = ImageSize.X;
-		const int32 Height = ImageSize.Y;
-		check(Width > 1 && Height > 1);
-		const FVector2f Scale = FVector2f(1.0 / (Width - 1), 1.0 / (Height - 1));
-
-		OutLineSegments.Reset(Width * Height);
-		for (int32 x = 0; x < Width - 1; x++)
+	static void RemoveSmallCurves(TArray<TArray<FVector2f>>& InOutCurves, float Threshold)
+	{
+		const float ThresholdSqr = Threshold * Threshold;
+		for (int32 i = InOutCurves.Num() - 1; i >= 0; i--)
 		{
-			for (int32 y = 0; y < Height - 1; y++)
+			const TArray<FVector2f>& Curve = InOutCurves[i];
+			FVector2f Min = FVector2f(MAX_flt);
+			FVector2f Max = FVector2f(-MAX_flt);
+			
+			for (const FVector2f& Point : Curve)
+			{
+				Min.X = FMath::Min(Min.X, Point.X);
+				Min.Y = FMath::Min(Min.Y, Point.Y);
+				Max.X = FMath::Max(Max.X, Point.X);
+				Max.Y = FMath::Max(Max.Y, Point.Y);
+			}
+			
+			const float Area = (Max.X - Min.X) * (Max.Y - Min.Y); 
+			const float AreaSqr = Area * Area;
+			if (AreaSqr < ThresholdSqr)
 			{
-				const FVector2i PixelsToSample[] = {
-					{ x, y },
-					{ x + 1, y },
-					{ x + 1, y + 1 },
-					{ x, y + 1 }
-				};
+				InOutCurves.RemoveAtSwap(i);
+			}
+		}
+	}
 	
-				int32 Index = 0;
-				for (int32 i = 0; i < 4; i++)
+	static void SmoothCurves(TArray<TArray<FVector2f>>& InOutCurves, int32 SmoothingIterations)
+	{
+		while (SmoothingIterations != 0)
+		{
+			for (TArray<FVector2f>& Curve : InOutCurves)
+			{
+				TArray<FVector2f> SmoothedCurve;
+				SmoothedCurve.SetNum(Curve.Num());
+			
+				for (int32 i = 0; i < Curve.Num(); i++)
 				{
-					if (!SampleImageBuffer(ImageData, Width, PixelsToSample[i]))
-					{
-						Index |= (1 << i); 
-					}
+					const FVector2f& Pos1 = i == 0 ? Curve.Last() : Curve[i - 1];
+					const FVector2f& Pos2 = Curve[i];
+					const FVector2f& Pos3 = i == Curve.Num() - 1 ? Curve[0] : Curve[i + 1];
+				
+					const FVector2f AvgPos = (Pos1 + Pos2 + Pos3) / 3; 
+					SmoothedCurve[i] = AvgPos;
 				}
-		        
-				const int32 Start = Starts[Index];
-				const int32 End = Starts[Index] + Sizes[Index];
-				const FVector2f PixelPosition = FVector2f(x, y);
-		        
-				for (int i = Start; i < End; i += 2)
+			
+				Curve = MoveTemp(SmoothedCurve);
+			}
+			
+			SmoothingIterations -= 1;
+		}
+	}
+	
+	static void SimplifyCurves(TArray<TArray<FVector2f>>& InOutCurves, float SimplificationEpsilon)
+	{
+		// TODO: We should probably be using TPolygon2::Simplify for this algorithm, however, it seems to produce slightly different results so more 
+		// thorough testing is required.
+		
+		/*
+		 * Non-recursive version of Ramer-Douglas-Peucker algorithm (https://en.wikipedia.org/wiki/Ramer%E2%80%93Douglas%E2%80%93Peucker_algorithm)
+		 *
+		 * Given the following curve, the algorithm will take the start and end point (d,e) and find the point furthest away (a) from the line segment
+		 * made up from d->e. It then takes the points (d,a) and (e,a), and finds the point between d->a and e->a that are furthest away from the line
+		 * segments made up from d->a and e->a respectively. In this case that would be the points b and c. This is repeated until the point found is
+		 * closer to the line than the SimplificationEpsilon, at which point the subdivision stops for that segment.
+		 *
+		 *    a
+		 *   /\
+		 * b/  \c
+		 *  \  / 
+		 *  d  e
+		 */
+	
+		const static auto FindFarPoint = [](const TArray<FVector2f>& CurvePoints, int32 StartPointIndex, int32 EndPointIndex, float Epsilon)->int32
+		{
+			const FVector2f& StartPos = CurvePoints[StartPointIndex];
+			const FVector2f& EndPos = CurvePoints[EndPointIndex];
+			float MaxDistSqr = Epsilon * Epsilon;
+			int32 FarPointIndex = INDEX_NONE;
+			for (int32 i = StartPointIndex + 1; i < EndPointIndex - 1; i++)
+			{
+				const FVector2f& Pos = CurvePoints[i];
+				const FVector2f ClosestPoint = FMath::ClosestPointOnSegment2D(Pos, StartPos, EndPos);
+				const float DistSqr = (Pos - ClosestPoint).SizeSquared();
+				if (DistSqr >= MaxDistSqr)
+				{ 
+					MaxDistSqr = DistSqr;
+					FarPointIndex = i;
+				}
+			}
+			return FarPointIndex;
+		};
+		
+		if (SimplificationEpsilon < 0.000001)
+		{
+			return; 
+		}
+		
+		TArray<int32> DouglasPeuckerQueue;
+		TArray<int32> NewDouglasPeuckerQueue;
+		TArray<int32> KeepPoints;
+		for (TArray<FVector2f>& Curve : InOutCurves)
+		{
+			DouglasPeuckerQueue.Reset(Curve.Num() * 2);
+			NewDouglasPeuckerQueue.Reset(Curve.Num() * 2);
+			KeepPoints.Reset(Curve.Num());
+			
+			DouglasPeuckerQueue.Append({ 0, Curve.Num() - 1 });
+			KeepPoints.Append(DouglasPeuckerQueue);
+			
+			while (DouglasPeuckerQueue.Num() > 0)
+			{
+				NewDouglasPeuckerQueue.Reset();
+				for (int32 i = 0; i < DouglasPeuckerQueue.Num(); i += 2)
 				{
-					const FVector2f LineStart = (Positions[i] + PixelPosition) * Scale;
-					const FVector2f LineEnd = (Positions[i + 1] + PixelPosition) * Scale;
-					OutLineSegments.Emplace(LineStart, LineEnd);
+					const int32 StartPointIndex = DouglasPeuckerQueue[i];
+					const int32 EndPointIndex = DouglasPeuckerQueue[i + 1];
+					const int32 FarPointIndex = FindFarPoint(Curve, StartPointIndex, EndPointIndex, SimplificationEpsilon);
+					if (FarPointIndex != INDEX_NONE)
+					{ 
+						KeepPoints.Add(FarPointIndex);
+						NewDouglasPeuckerQueue.Append({ StartPointIndex, FarPointIndex, FarPointIndex, EndPointIndex });
+					}
 				}
+				Swap(DouglasPeuckerQueue, NewDouglasPeuckerQueue);
 			}
+			
+			KeepPoints.Sort();
+			for (int32 i = 0; i < KeepPoints.Num(); i++)
+			{
+				const int32 PointIndex = KeepPoints[i];
+				Curve[i] = Curve[PointIndex]; // It's safe to overwrite the Curve array with itself since KeepPoints is sorted, so i can never be larger than PointIndex
+			}
+			Curve.SetNum(KeepPoints.Num());
 		}
 	}
 };
@@ -258,7 +370,7 @@ PV::Texture2DImport::EImportResult PV::Texture2DImport::ImportGrowthDataFromText
 	
 	TArray<bool> BlackAndWhiteImageData;
 
-	SampleImage(InSourceTexture2D->Source, SampleSize, BlackAndWhiteImageData);
+	SampleImage(InSourceTexture2D->Source, SampleSize, BlackAndWhiteImageData, InOptions.bInvertImage, InOptions.WhiteLevel);
 	if (SampleSize.X <= 1 || SampleSize.Y <= 1)
 	{
 		return EImportResult::InvalidTextureSize;
@@ -270,12 +382,32 @@ PV::Texture2DImport::EImportResult PV::Texture2DImport::ImportGrowthDataFromText
 		return ReturnWithDebugMesh();
 	}
 
-	TArray<TPair<FVector2f, FVector2f>> LineSegments;
-	ExtractContourLines(BlackAndWhiteImageData, SampleSize, LineSegments);
-	if (InOptions.DebugState == FImportOptions::EDebugState::ExtractContourLines)
+	TArray<TArray<FVector2f>> Curves;
+	TraceImage(BlackAndWhiteImageData, SampleSize, Curves);
+	if (InOptions.DebugState == FImportOptions::EDebugState::TraceImage)
 	{
-		AddDebugPixelValuesToCollection(OutCollection, BlackAndWhiteImageData, SampleSize);
-		AddDebugLineSegmentsToCollection(OutCollection, LineSegments);
+		AddDebugCurvesToCollection(OutCollection, Curves);
+		return ReturnWithDebugMesh();
+	}
+	
+	RemoveSmallCurves(Curves, InOptions.MinBoundsArea);
+	if (InOptions.DebugState == FImportOptions::EDebugState::RemoveSmallCurves)
+	{
+		AddDebugCurvesToCollection(OutCollection, Curves);
+		return ReturnWithDebugMesh();
+	}
+	
+	SmoothCurves(Curves, InOptions.SmoothingIterations);
+	if (InOptions.DebugState == FImportOptions::EDebugState::SmoothCurves)
+	{
+		AddDebugCurvesToCollection(OutCollection, Curves);
+		return ReturnWithDebugMesh();
+	}
+	
+	SimplifyCurves(Curves, InOptions.SimplificationEpsilon);
+	if (InOptions.DebugState == FImportOptions::EDebugState::SimplifyCurves)
+	{
+		AddDebugCurvesToCollection(OutCollection, Curves);
 		return ReturnWithDebugMesh();
 	}
 	
diff --git a/Engine/Plugins/Experimental/ProceduralVegetationEditor/Source/ProceduralVegetationEditor/Private/Importer/PVImporter_Texture2D.h b/Engine/Plugins/Experimental/ProceduralVegetationEditor/Source/ProceduralVegetationEditor/Private/Importer/PVImporter_Texture2D.h
--- a/Engine/Plugins/Experimental/ProceduralVegetationEditor/Source/ProceduralVegetationEditor/Private/Importer/PVImporter_Texture2D.h
+++ b/Engine/Plugins/Experimental/ProceduralVegetationEditor/Source/ProceduralVegetationEditor/Private/Importer/PVImporter_Texture2D.h
@@ -21,13 +21,21 @@ namespace PV::Texture2DImport
 	{
 		int32 SampleWidth = 128;
 		int32 SampleHeight = 128;
+		float MinBoundsArea = 0;
+		int32 SmoothingIterations = 1;
+		float SimplificationEpsilon = 0.001;
+		bool bInvertImage = false;
+		float WhiteLevel = 0.01;
 		
 		enum class EDebugState : uint8
 		{
 			None,
 			Texture,
 			SampleImage,
-			ExtractContourLines,
+			TraceImage,
+			RemoveSmallCurves,
+			SmoothCurves,
+			SimplifyCurves,
 			MAX
 		};
 		EDebugState DebugState = EDebugState::None;
diff --git a/Engine/Plugins/Experimental/ProceduralVegetationEditor/Source/ProceduralVegetationEditor/Private/PVEditorInterfaceImpl.cpp b/Engine/Plugins/Experimental/ProceduralVegetationEditor/Source/ProceduralVegetationEditor/Private/PVEditorInterfaceImpl.cpp
--- a/Engine/Plugins/Experimental/ProceduralVegetationEditor/Source/ProceduralVegetationEditor/Private/PVEditorInterfaceImpl.cpp
+++ b/Engine/Plugins/Experimental/ProceduralVegetationEditor/Source/ProceduralVegetationEditor/Private/PVEditorInterfaceImpl.cpp
@@ -24,13 +24,18 @@ bool PV::FPVEditorInterfaceImpl::ImportGrowthDataFromStaticMesh(const FPVImportS
 
 bool PV::FPVEditorInterfaceImpl::ImportGrowthDataFromTexture2D(const FPVImportTexture2DParams& InParams, FManagedArrayCollection& OutCollection)
 {
-	PV::Texture2DImport::FImportOptions ImportMeshOptions;
-	ImportMeshOptions.SampleWidth = InParams.SampleWidth;
-	ImportMeshOptions.SampleHeight = InParams.SampleHeight;
+	PV::Texture2DImport::FImportOptions ImportOptions;
+	ImportOptions.SampleWidth = InParams.SampleWidth;
+	ImportOptions.SampleHeight = InParams.SampleHeight;
+	ImportOptions.MinBoundsArea = InParams.MinBoundsArea;
+	ImportOptions.SmoothingIterations = InParams.SmoothingIterations;
+	ImportOptions.SimplificationEpsilon = InParams.SimplificationEpsilon;
+	ImportOptions.bInvertImage = InParams.bInvertImage;
+	ImportOptions.WhiteLevel = InParams.WhiteLevel;
 	
 	static_assert(static_cast<int32>(EPVImportTexture2DDebugState::MAX) == static_cast<int32>(PV::Texture2DImport::FImportOptions::EDebugState::MAX));
-	ImportMeshOptions.DebugState = static_cast<PV::Texture2DImport::FImportOptions::EDebugState>(static_cast<int32>(InParams.DebugState));
+	ImportOptions.DebugState = static_cast<PV::Texture2DImport::FImportOptions::EDebugState>(static_cast<int32>(InParams.DebugState));
 
-	const PV::Texture2DImport::EImportResult ImportResult = PV::Texture2DImport::ImportGrowthDataFromTexture2D(InParams.Texture2DAsset, OutCollection, ImportMeshOptions);
+	const PV::Texture2DImport::EImportResult ImportResult = PV::Texture2DImport::ImportGrowthDataFromTexture2D(InParams.Texture2DAsset, OutCollection, ImportOptions);
 	return ImportResult == PV::Texture2DImport::EImportResult::Success;
 }
\ No newline at end of file
