[PVE] Mesh twisting fix
Issue: Calculation of Up Vector for each point was being calculated wrongly, which caused mesh twisting, especially on sharp corners
Fix: Up Vector is now calculated w.r.t the previous Point Up Vector, which works in all cases. Also there was no point of over-calculating AimVector as that was already fine.

#Synced-CL 49590164

diff --git a/Engine/Plugins/Experimental/ProceduralVegetationEditor/Source/ProceduralVegetation/Private/Implementations/PVMeshBuilder.cpp b/Engine/Plugins/Experimental/ProceduralVegetationEditor/Source/ProceduralVegetation/Private/Implementations/PVMeshBuilder.cpp
--- a/Engine/Plugins/Experimental/ProceduralVegetationEditor/Source/ProceduralVegetation/Private/Implementations/PVMeshBuilder.cpp
+++ b/Engine/Plugins/Experimental/ProceduralVegetationEditor/Source/ProceduralVegetation/Private/Implementations/PVMeshBuilder.cpp
@@ -14,6 +14,8 @@
 #include "Facades/PVPointFacade.h"
 #include "Facades/PVTreeFacade.h"
 
+#include "Helpers/PVUtilities.h"
+
 void MergeMaterials(const TArray<FTrunkGenerationMaterialSetup>& InMaterialSetups, TArray<FLocalDynamicMeshData>& MeshDatas)
 {
 	TMap<FString, int32> MaterialPathIDMap;
@@ -761,8 +763,7 @@ void FPVMeshBuilder::GenerateBranchMeshData(const bool bPrimitiveIsTrunk, const
 		&& DisplacementData.TextureHeight > 0
 		&& GenerationNumber <= MeshBuilderParams.DisplacementGenerationUpperLimit;
 
-	FVector3f PreviousUpVector = FVector3f::ForwardVector;
-	FVector3f PreviousAimVector = FVector3f::UpVector;
+	FVector3f PreviousUpVector = PointBudDirections[PrimitivePoints[0]][5].GetUnsafeNormal();
 
 	float Displacement_UV_V = 0.0f;
 	FVector3f PreviousPosition = PointPositions[PrimitivePoints[0]];
@@ -790,53 +791,50 @@ void FPVMeshBuilder::GenerateBranchMeshData(const bool bPrimitiveIsTrunk, const
 			: 1.0f / Interval;
 		TextureCoordV /= URatio;
 
-		FVector3f AimVector = PointBudDirections[PointIndex][0].GetUnsafeNormal();
-		// The AimVector needs to be adjusted in case curvature has changed during mesh simplification
 		const FVector3f& PointPosition = PointPositions[PointIndex];
-		FVector3f UpdatedAimVector;
-		
 		const float PointScale = PointScales[PointIndex];
-		Displacement_UV_V += FVector3f::Dist(PointPosition, PreviousPosition) * MaxScaleRatio * (MaxPointScale / PointScale);
 		
-		if (Index == PrimitivePoints.Num() - 1)
+		Displacement_UV_V += FVector3f::Dist(PointPosition, PreviousPosition) * MaxScaleRatio * (MaxPointScale / PointScale);
+
+		FVector3f TangentVector = FVector3f::UpVector;
+		if (PrimitivePoints.IsValidIndex(Index - 1) && PrimitivePoints.IsValidIndex(Index + 1))
+		[[likely]]
+		{
+			const FVector3f PreviousTangent = (PointPosition - PointPositions[PrimitivePoints[Index - 1]]).GetSafeNormal();
+			const FVector3f NextTangent = (PointPositions[PrimitivePoints[Index + 1]] - PointPosition).GetSafeNormal();
+			TangentVector = FVector3f::SlerpNormals(PreviousTangent, NextTangent, 0.5f);
+		}
+		else if (Index == PrimitivePoints.Num() - 1)
 		[[unlikely]]
 		{
 			const int32 PreviousPointIndex = PrimitivePoints[Index - 1];
 			const FVector3f& PreviousPointPosition = PointPositions[PreviousPointIndex];
-			UpdatedAimVector = (PointPosition - PreviousPointPosition).GetSafeNormal();
+			TangentVector = (PointPosition - PreviousPointPosition).GetSafeNormal();
 		}
-		else
-		[[likely]]
+		else if (Index == 0)
+		[[unlikely]]
 		{
 			const int32 NextPointIndex = PrimitivePoints[Index + 1];
 			const FVector3f& NextPointPosition = PointPositions[NextPointIndex];
-			UpdatedAimVector = (NextPointPosition - PointPosition).GetSafeNormal();
-		}
-		const FQuat4f RotationQuat = FQuat4f::FindBetweenNormals(AimVector, UpdatedAimVector);
-
-		const FVector3f PointUpOriginal = PointBudDirections[PointIndex][5].GetUnsafeNormal();
-		const FVector3f PointUpUpdated = RotationQuat.RotateVector(PointUpOriginal);
-
-		FVector3f PointUp = PointUpUpdated;
-
-		FVector3f CrossVector = PointUp.Cross(UpdatedAimVector);
-		PointUp = CrossVector.Cross(UpdatedAimVector);
-
-		if (Index > 0)
-		{
-			if (Index > 0 && FVector3f::DotProduct(PointUp, PreviousUpVector) < 0.0f)
-			{
-				PointUp = -PointUp;
-			}
-
-			const FVector3f InterpolatedAimVector = FMath::Lerp(PreviousAimVector, UpdatedAimVector, 0.5f).GetSafeNormal();
-			const FQuat4f AimVectorRotationQuat = FQuat4f::FindBetweenNormals(UpdatedAimVector, InterpolatedAimVector);
-			PointUp = AimVectorRotationQuat.RotateVector(PointUp);
-			UpdatedAimVector = InterpolatedAimVector;
+			TangentVector = (NextPointPosition - PointPosition).GetSafeNormal();
 		}
 
-		PreviousAimVector = UpdatedAimVector;
-		PreviousUpVector = PointUp;
+		/*
+		 * This creates a perpendicular Up Vector to the Current Tangent
+		 * The current Tangent and Previous Up Vector is crossed to get a perpendicular vector to both
+		 * Then that vector is crossed again with the Tangent Vector to get a third perpendicular vector which becomes the Up Vector
+		 *       Tangent
+		 *          |
+		 *          | 90°
+		 *      90° O------- Cross
+		 *         / 90°
+		 *        /
+		 *       Up 
+		 */
+		FVector3f CrossVector = FVector3f::CrossProduct(TangentVector, PreviousUpVector);
+		FVector3f UpVector = FVector3f::CrossProduct(CrossVector, TangentVector);
+
+		PreviousUpVector = UpVector;
 
 		FVector2f UV1 = FVector2f(0, 1.0 - Displacement_UV_V);
 		const FVector2f UV2 = FVector2f(LengthFromRoot * 0.01, SafeDivide(PointScale, MaxPointScale));
@@ -847,7 +845,7 @@ void FPVMeshBuilder::GenerateBranchMeshData(const bool bPrimitiveIsTrunk, const
 		if (PointScale == 0.0f)
 		[[unlikely]]
 		{
-			CurrentIndices.Add(OutLocalMeshData.Vertices.Emplace(PointPosition, UpdatedAimVector, FVector2f(UVOGenerationOffset, TextureCoordV), UV1, UV2, NjordPixelId));
+			CurrentIndices.Add(OutLocalMeshData.Vertices.Emplace(PointPosition, TangentVector, FVector2f(UVOGenerationOffset, TextureCoordV), UV1, UV2, NjordPixelId));
 		}
 		else
 		[[likely]]
@@ -870,7 +868,7 @@ void FPVMeshBuilder::GenerateBranchMeshData(const bool bPrimitiveIsTrunk, const
 					ProfileMultiplier = FMath::Lerp(1.0f, ProfileMultiplier * MeshBuilderParams.PlantProfileScale, ProfileFallOff);
 				}
 
-				const FVector3f Direction = FQuat4f(UpdatedAimVector, RotationDelta * Division).RotateVector(PointUp).GetSafeNormal();
+				const FVector3f Direction = FQuat4f(TangentVector, RotationDelta * Division).RotateVector(UpVector).GetSafeNormal();
 				FVector3f RadialPoint = PointPosition + (Direction * PointScale * ProfileMultiplier);
 
 				if (bShouldApplyDisplacement)
