First implementation of 2D growth data extraction algrithm

This CL implements the remaining parts if the APIX 2D growth data import algorithm, such that it is possible to generate valid growth data from a 2D image inside of PVE. This CL also moves code which is shared between the 2D and 3D importers from PVImporter_StaticMesh to PVImportHelpers.

Overview:
- Implement remaining 2D growth data import algorithm
- TriangulateCurvesAndSplitEdges
- ComputeVertexScale2D
- FindEndPoints2D
- FindRootPoints2D
- Add support for importing multi-plant assets
- PV::Import::Helper::FillGrowthDataFromBranchHierarchyDescription
- PV::Import::Helper::FillGrowthDataFromFoliageDescription
- PV::Import::Helper::GenerateGrowthDataFromPlantDescription
- Refactor PVImportHelper to allow for sharing of functionality between 2D and 3D importer
- PV::Import::Helper::ComputeBranchHierarchies
- PV::Import::Helper::ComputeLengthFromRoot: Fix issue that caused the wrong branch to be selected as parent branch during import (PVE-167)
- PV::Import::Helper::TraceBranches
- PV::Import::Helper::ExtractBranchHierarchies
- PV::Import::Helper::ConvertMeshBranchHierarchiesToBranchDescription

Main files of interest for this review:
- PVImporter_Texture2D: This file contains the main new algorithms added to complete the 2D growth data importer (TriangulateCurvesAndSplitEdges, ComputeVertexScale2D, FindEndPoints2D, FindRootPoints2D)
- PVImporter_StaticMesh: Shared functionality have been moved out of this file (TDynamicMeshVertexAttributeDefinition, ComputeLengthFromRoot, DrawBranchCurves, RemoveBranchCurveOverlaps, ExtractBranchesFromBranchCurves)
- PVImportHelpers: This file now contains shared import functionality between the 2D and 3D importer. ComputeLengthFromRoot and TraceBranches (new name for DrawBranchCurves) have been re-worked to solve issue PVE-167.

JIRAs:
https://jira.it.epicgames.com/browse/PVE-88
https://jira.it.epicgames.com/browse/PVE-91
https://jira.it.epicgames.com/browse/PVE-167

#Synced-CL 49615590

diff --git a/Engine/Plugins/Experimental/ProceduralVegetationEditor/Source/ProceduralVegetation/Private/Nodes/PVStaticMeshImporterSettings.h b/Engine/Plugins/Experimental/ProceduralVegetationEditor/Source/ProceduralVegetation/Private/Nodes/PVStaticMeshImporterSettings.h
--- a/Engine/Plugins/Experimental/ProceduralVegetationEditor/Source/ProceduralVegetation/Private/Nodes/PVStaticMeshImporterSettings.h
+++ b/Engine/Plugins/Experimental/ProceduralVegetationEditor/Source/ProceduralVegetation/Private/Nodes/PVStaticMeshImporterSettings.h
@@ -19,7 +19,7 @@ public:
 	virtual FText GetDefaultNodeTitle() const override;
 	virtual FText GetNodeTooltipText() const override;
 	virtual EPCGSettingsType GetType() const override { return EPCGSettingsType::InputOutput; }
-	virtual TArray<EPVRenderType> GetDefaultRenderType() const override { return TArray{ EPVRenderType::PointData }; }
+	virtual TArray<EPVRenderType> GetDefaultRenderType() const override { return TArray{ EPVRenderType::PointData, EPVRenderType::Mesh }; }
 	virtual TArray<EPVRenderType> GetSupportedRenderTypes() const override { return TArray{ EPVRenderType::PointData, EPVRenderType::Mesh }; }
 #endif
 
diff --git a/Engine/Plugins/Experimental/ProceduralVegetationEditor/Source/ProceduralVegetation/Private/Nodes/PVTexture2DImporterSettings.h b/Engine/Plugins/Experimental/ProceduralVegetationEditor/Source/ProceduralVegetation/Private/Nodes/PVTexture2DImporterSettings.h
--- a/Engine/Plugins/Experimental/ProceduralVegetationEditor/Source/ProceduralVegetation/Private/Nodes/PVTexture2DImporterSettings.h
+++ b/Engine/Plugins/Experimental/ProceduralVegetationEditor/Source/ProceduralVegetation/Private/Nodes/PVTexture2DImporterSettings.h
@@ -19,7 +19,7 @@ public:
 	virtual FText GetDefaultNodeTitle() const override;
 	virtual FText GetNodeTooltipText() const override;
 	virtual EPCGSettingsType GetType() const override { return EPCGSettingsType::InputOutput; }
-	virtual TArray<EPVRenderType> GetDefaultRenderType() const override { return TArray{ EPVRenderType::PointData }; }
+	virtual TArray<EPVRenderType> GetDefaultRenderType() const override { return TArray{ EPVRenderType::PointData, EPVRenderType::Mesh }; }
 	virtual TArray<EPVRenderType> GetSupportedRenderTypes() const override { return TArray{ EPVRenderType::PointData, EPVRenderType::Mesh }; }
 #endif
 
diff --git a/Engine/Plugins/Experimental/ProceduralVegetationEditor/Source/ProceduralVegetation/Public/PVImportStaticMeshParams.h b/Engine/Plugins/Experimental/ProceduralVegetationEditor/Source/ProceduralVegetation/Public/PVImportStaticMeshParams.h
--- a/Engine/Plugins/Experimental/ProceduralVegetationEditor/Source/ProceduralVegetation/Public/PVImportStaticMeshParams.h
+++ b/Engine/Plugins/Experimental/ProceduralVegetationEditor/Source/ProceduralVegetation/Public/PVImportStaticMeshParams.h
@@ -19,10 +19,7 @@ enum class EPVImportStaticMeshDebugState : uint8
 	FindEndPoints,
 	ConnectBranchesToParent,
 	FindRootPoints,
-	ComputeLengthFromRoot,
-	DrawBranchCurves,
-	RemoveBranchCurveOverlaps,
-	ExtractBranches,
+	ComputeBranchHierarchies,
 	MAX UMETA(Hidden)
 };
 
diff --git a/Engine/Plugins/Experimental/ProceduralVegetationEditor/Source/ProceduralVegetation/Public/PVImportTexture2DParams.h b/Engine/Plugins/Experimental/ProceduralVegetationEditor/Source/ProceduralVegetation/Public/PVImportTexture2DParams.h
--- a/Engine/Plugins/Experimental/ProceduralVegetationEditor/Source/ProceduralVegetation/Public/PVImportTexture2DParams.h
+++ b/Engine/Plugins/Experimental/ProceduralVegetationEditor/Source/ProceduralVegetation/Public/PVImportTexture2DParams.h
@@ -17,29 +17,46 @@ enum class EPVImportTexture2DDebugState : uint8
 	RemoveSmallCurves,
 	SmoothCurves,
 	SimplifyCurves,
+	TriangulateCurves,
+	ComputeVertexScale,
+	FindTips,
+	ApplyRootPoints,
+	ComputeBranchHierarchies,
 	MAX UMETA(Hidden)
 };
 
+USTRUCT(BlueprintType)
+struct FPVImportTexture2DRootPoint
+{
+	GENERATED_BODY()
+
+	UPROPERTY(EditAnywhere, Category = "Root Point")
+	FVector2f Position = FVector2f::ZeroVector;
+
+	UPROPERTY(EditAnywhere, Category = "Root Point")
+	float Rotation = 0;
+};
+
 USTRUCT(BlueprintType)
 struct FPVImportTexture2DParams
 {
 	GENERATED_BODY()
 
 	UPROPERTY(EditAnywhere, Category = "Image Trace", meta=(ClampMin=0, UIMin=0))
-	int32 SampleWidth = 128;
+	int32 SampleWidth = 512;
 	
 	UPROPERTY(EditAnywhere, Category = "Image Trace", meta=(ClampMin=0, UIMin=0))
-    int32 SampleHeight = 128;
+    int32 SampleHeight = 512;
     
     // Any branches with a bounds area lower than this value will be ignored.
     // Value is normalized meaning a value of 0.1 will filter out any branch whose bounds cover less than 10% of the image. 
-	UPROPERTY(EditAnywhere, Category = "Image Trace", meta=(ClampMin=0, UIMin=0))
+	UPROPERTY(EditAnywhere, Category = "Image Trace", meta=(ClampMin=0, UIMin=0, UIMax=0.2))
 	float MinBoundsArea = 0.05f;
 	
-	UPROPERTY(EditAnywhere, Category = "Image Trace", meta=(ClampMin=0, UIMin=0, UIMax=10, ClampMax=100))
+	UPROPERTY(EditAnywhere, Category = "Image Trace", meta=(ClampMin=0, ClampMax=100, UIMin=0, UIMax=10))
 	int32 SmoothingIterations = 1;
 	
-	UPROPERTY(EditAnywhere, Category = "Image Trace", meta=(ClampMin=0, UIMin=0, UIMax=0.01, ClampMax=1))
+	UPROPERTY(EditAnywhere, Category = "Image Trace", meta=(ClampMin=0, ClampMax=0.1, UIMin=0, UIMax=0.01))
 	float SimplificationEpsilon = 0.001;
 
 	UPROPERTY(EditAnywhere, Category = "Texture")
@@ -48,9 +65,21 @@ struct FPVImportTexture2DParams
 	UPROPERTY(EditAnywhere, Category = "Texture")
 	bool bInvertImage = false;
 	
-	UPROPERTY(EditAnywhere, Category = "Texture", meta=(ClampMin=0, UIMin=0, UIMax=1, ClampMax=1))
+	UPROPERTY(EditAnywhere, Category = "Texture", meta=(ClampMin=0, ClampMax=1, UIMin=0, UIMax=1))
 	float WhiteLevel = 0.01;
 
+	UPROPERTY(EditAnywhere, Category = "Branch Geneneration")
+	float Scale = 100.f;
+
+	UPROPERTY(EditAnywhere, Category = "Branch Geneneration")
+	TArray<FPVImportTexture2DRootPoint> RootPoints;
+
+	UPROPERTY(EditAnywhere, Category = "Branch Geneneration|Tip detection", meta = (ClampMin=0.001, ClampMax=180, UIMin=0.01, UIMax=180))
+	float TipAngleThresholdInDegrees = 45.0;
+
+	UPROPERTY(EditAnywhere, Category = "Branch Geneneration|Tip detection", meta = (ClampMin=0, ClampMax=0.01, UIMin=0, UIMax=0.01))
+	float MaxTipAngleSearchDist = 0.005;
+
 	UPROPERTY(EditAnywhere, Category = "Inspect")
 	EPVImportTexture2DDebugState DebugState = EPVImportTexture2DDebugState::None;
 };
diff --git a/Engine/Plugins/Experimental/ProceduralVegetationEditor/Source/ProceduralVegetationEditor/Private/Importer/PVDynamicMeshVertexAttribute.h b/Engine/Plugins/Experimental/ProceduralVegetationEditor/Source/ProceduralVegetationEditor/Private/Importer/PVDynamicMeshVertexAttribute.h
new file mode 100644
--- /dev/null
+++ b/Engine/Plugins/Experimental/ProceduralVegetationEditor/Source/ProceduralVegetationEditor/Private/Importer/PVDynamicMeshVertexAttribute.h
@@ -0,0 +1,203 @@
+// Copyright Epic Games, Inc. All Rights Reserved.
+#pragma once
+
+#include "DynamicMesh/DynamicMesh3.h"
+#include "DynamicMesh/DynamicVertexAttribute.h"
+#include "DynamicMesh/DynamicMeshAttributeSet.h"
+
+#include "GeometryCollection/GeometryCollection.h" // Needed for CopyAttributeToCollection
+
+namespace PV
+{
+template<typename AttribValueType>
+class TDynamicMeshVertexAttributeExt 
+	: public UE::Geometry::TDynamicMeshVertexAttribute<AttribValueType, 1>
+{
+public:
+	using Super = UE::Geometry::TDynamicMeshVertexAttribute<AttribValueType, 1>;
+
+	TDynamicMeshVertexAttributeExt() : Super() {}
+	TDynamicMeshVertexAttributeExt(UE::Geometry::FDynamicMesh3* ParentIn, bool bAutoInit = true) : Super(ParentIn, bAutoInit) {}
+
+	inline AttribValueType GetValue(int32 VertexID) const
+	{
+		AttribValueType Value = (AttribValueType)0;
+		Super::GetValue(VertexID, &Value);
+		return Value;
+	}
+
+	inline void SetValue(int32 VertexID, const AttribValueType& NewValue)
+	{
+		Super::SetValue(VertexID, &NewValue);
+	}
+
+	template <typename Predicate, typename Iteratable>
+	TArray<int32> FindAllByPredicate(Predicate Pred, const Iteratable& InIteratable) const
+	{
+		TArray<int32> Result;
+
+		if (Super::Parent)
+		{
+			Result.Reserve(Super::Parent->VertexCount());
+
+			for (int32 VertexID : InIteratable)
+			{
+				check(Super::Parent->IsVertex(VertexID));
+
+				const AttribValueType Value = GetValue(VertexID);
+				if (Pred(Value))
+				{
+					Result.Add(VertexID);
+				}
+			}
+		}
+
+		return Result;
+	}
+
+	TArray<int32> FindAllNonZero() const
+	{
+		if (!Super::Parent)
+		{
+			return TArray<int32>();
+		}
+
+		return FindAllByPredicate(
+			[](const auto& Value) { return Value != (AttribValueType)0; },
+			Super::Parent->VertexIndicesItr()
+		);
+	}
+
+	TArray<int32> FindAllZero() const
+	{
+		if (!Super::Parent)
+		{
+			return TArray<int32>();
+		}
+
+		return FindAllByPredicate(
+			[](const auto& Value) { return Value == (AttribValueType)0; },
+			Super::Parent->VertexIndicesItr()
+		);
+	}
+
+	TArray<int32> FindAllNonZero(const TArray<int32>& InVertexIDsToSearch) const
+	{
+		return FindAllByPredicate(
+			[](const auto& Value) { return Value != (AttribValueType)0; },
+			InVertexIDsToSearch
+		);
+	}
+
+	TArray<int32> FindAllZero(const TArray<int32>& InVertexIDsToSearch) const
+	{
+		return FindAllByPredicate(
+			[](const auto& Value) { return Value == (AttribValueType)0; },
+			InVertexIDsToSearch
+		);
+	}
+};
+
+template<typename T>
+class TDynamicMeshVertexAttributeDefinition
+{
+public:
+	using FAttributeType = TDynamicMeshVertexAttributeExt<T>;
+
+	const FName AttributeName;
+
+	TDynamicMeshVertexAttributeDefinition(FName InAttributeName)
+		: AttributeName(InAttributeName)
+	{
+	}
+
+	const FAttributeType* GetAttribute(const UE::Geometry::FDynamicMesh3& DynamicMesh) const
+	{
+		if (DynamicMesh.HasAttributes() && DynamicMesh.Attributes()->HasAttachedAttribute(AttributeName))
+		{
+			return static_cast<const FAttributeType*>(DynamicMesh.Attributes()->GetAttachedAttribute(AttributeName));
+		}
+		return nullptr;
+	}
+
+	FAttributeType* GetAttribute(UE::Geometry::FDynamicMesh3& DynamicMesh) const
+	{
+		if (DynamicMesh.HasAttributes() && DynamicMesh.Attributes()->HasAttachedAttribute(AttributeName))
+		{
+			return static_cast<FAttributeType*>(DynamicMesh.Attributes()->GetAttachedAttribute(AttributeName));
+		}
+		return nullptr;
+	}
+
+	const FAttributeType& GetAttributeChecked(const UE::Geometry::FDynamicMesh3& DynamicMesh) const
+	{
+		const FAttributeType* Attribute = GetAttribute(DynamicMesh);
+		check(Attribute != nullptr);
+		return *Attribute;
+	}
+
+	FAttributeType& GetAttributeChecked(UE::Geometry::FDynamicMesh3& DynamicMesh) const
+	{
+		FAttributeType* Attribute = GetAttribute(DynamicMesh);
+		check(Attribute != nullptr);
+		return *Attribute;
+	}
+
+	FAttributeType& GetOrAttachAttribute(UE::Geometry::FDynamicMesh3& DynamicMesh) const
+	{
+		FAttributeType* Attribute = GetAttribute(DynamicMesh);
+		if (Attribute == nullptr)
+		{
+			AttachAttribute(DynamicMesh);
+		}
+
+		return GetAttributeChecked(DynamicMesh);
+	}
+
+	void AttachAttribute(UE::Geometry::FDynamicMesh3& DynamicMesh) const
+	{
+		if (!DynamicMesh.HasAttributes())
+		{
+			DynamicMesh.EnableAttributes();
+		}
+
+		if (!DynamicMesh.Attributes()->HasAttachedAttribute(AttributeName))
+		{
+			auto* Attribute = new FAttributeType(&DynamicMesh);
+			Attribute->Initialize();
+			DynamicMesh.Attributes()->AttachAttribute(AttributeName, Attribute);
+		}
+	}
+
+	void CopyAttributeToCollection(const UE::Geometry::FDynamicMesh3& DynamicMesh, FManagedArrayCollection& OutCollection) const
+	{
+		const FAttributeType* AttachedAttribute = GetAttribute(DynamicMesh);
+		if (!AttachedAttribute)
+		{
+			return;
+		}
+
+		const static FName VertexAttributesGroupName = TEXT("VertexAttributes");
+
+		if (!OutCollection.HasGroup(VertexAttributesGroupName))
+		{
+			OutCollection.AddGroup(VertexAttributesGroupName);
+			OutCollection.AddElements(DynamicMesh.MaxVertexID(), VertexAttributesGroupName);
+
+			auto& DebugVertexPositionAttribute = OutCollection.AddAttribute<FVector3f>(FGeometryCollection::VertexPositionAttribute, VertexAttributesGroupName);
+			for (int32 VertexIndex : DynamicMesh.VertexIndicesItr())
+			{
+				DebugVertexPositionAttribute[VertexIndex] = FVector3f(DynamicMesh.GetVertex(VertexIndex));
+			}
+		}
+
+		TManagedArray<T>& CollectionAttribute = OutCollection.AddAttribute<T>(AttributeName, VertexAttributesGroupName);
+		check(CollectionAttribute.Num() >= DynamicMesh.MaxVertexID());
+
+		for (int32 VertexIndex : DynamicMesh.VertexIndicesItr())
+		{
+			CollectionAttribute[VertexIndex] = AttachedAttribute->GetValue(VertexIndex);
+		}
+	}
+};
+}
\ No newline at end of file
diff --git a/Engine/Plugins/Experimental/ProceduralVegetationEditor/Source/ProceduralVegetationEditor/Private/Importer/PVImportHelpers.cpp b/Engine/Plugins/Experimental/ProceduralVegetationEditor/Source/ProceduralVegetationEditor/Private/Importer/PVImportHelpers.cpp
--- a/Engine/Plugins/Experimental/ProceduralVegetationEditor/Source/ProceduralVegetationEditor/Private/Importer/PVImportHelpers.cpp
+++ b/Engine/Plugins/Experimental/ProceduralVegetationEditor/Source/ProceduralVegetationEditor/Private/Importer/PVImportHelpers.cpp
@@ -2,6 +2,11 @@
 #include "PVImportHelpers.h"
 
 #include "ProceduralVegetationPreset.h"
+#include "PVDynamicMeshVertexAttribute.h"
+#include "DynamicMesh/DynamicMesh3.h"
+#include "DynamicMesh/DynamicMeshAABBTree3.h"
+#include "GeometryCollectionToDynamicMesh.h"
+#include "GeometryCollection/GeometryCollection.h"
 
 #include "Helpers/PVFacadeHelper.h"
 #include "Facades/PVAttributesNames.h"
@@ -12,6 +17,618 @@
 
 #include "Engine/StaticMesh.h"
 
+namespace PV::DynamicMeshHelper
+{
+	static void ComputeLengthFromRoot(
+		const UE::Geometry::FDynamicMesh3& DynamicMesh,
+		const FBooleanVertexAttribute& RootPointsAttribute,
+		const FBooleanVertexAttribute& EndPointsAttribute,
+		FFloatVertexAttribute& LengthFromRootAttribute,
+		const TFunction<bool(int32)>& VertexFilterPredicate
+	)
+	{
+		TRACE_CPUPROFILER_EVENT_SCOPE(ComputeLengthFromRoot);
+
+		// This algorithm starts at the root point and walks along every vertex, keeping track of its distance to the root point and stores that distance
+		// in the LengthFromRootAttribute.
+
+		for (int32 VertexID : DynamicMesh.VertexIndicesItr())
+		{
+			LengthFromRootAttribute.SetValue(VertexID, -1);
+		}
+
+		TArray<int32> SearchQueue;
+		SearchQueue.Reserve(DynamicMesh.VertexCount());
+		for (int32 VertexID : RootPointsAttribute.FindAllNonZero())
+		{
+			SearchQueue.Add(VertexID);
+			LengthFromRootAttribute.SetValue(VertexID, 0);
+		}
+
+		while (SearchQueue.Num() > 0)
+		{
+			const int32& VertexID = SearchQueue.Pop(EAllowShrinking::No);
+
+			const float PrevLengthFromRoot = LengthFromRootAttribute.GetValue(VertexID);
+			const FVector PrevVertexPosition = DynamicMesh.GetVertex(VertexID);
+
+			const TArray<int32> Neighbours = GetVertexNeighbours(DynamicMesh, VertexID);
+			for (int32 NeighbourVertexID : Neighbours)
+			{
+				if (!EndPointsAttribute.GetValue(NeighbourVertexID) // Ingore vertex filter for end points as that means we've reached the end of the branch
+					&& !VertexFilterPredicate(NeighbourVertexID))
+				{
+					continue;
+				}
+
+				const FVector VertexPosition = DynamicMesh.GetVertex(NeighbourVertexID);
+				const float NewLengthFromRoot = PrevLengthFromRoot + (VertexPosition - PrevVertexPosition).Size();
+				const float ExistingLengthFromRoot = LengthFromRootAttribute.GetValue(NeighbourVertexID);
+
+				if (ExistingLengthFromRoot >= 0 
+					&& ExistingLengthFromRoot <= NewLengthFromRoot)
+				{
+					// If we find a shorter path then we should ideally remove vertices from the queue
+					// which are computing the distance using old values. 
+					// The current algorithm will still produce correct results, but may compute a large number
+					// of unecessary distances which will later be over-written when a closer path is found.
+					continue;
+				}
+
+				LengthFromRootAttribute.SetValue(NeighbourVertexID, NewLengthFromRoot);
+
+				SearchQueue.Insert(NeighbourVertexID, 0);
+			}
+		}
+	}
+
+	static void TraceBranches(
+		const UE::Geometry::FDynamicMesh3& DynamicMesh,
+		const FFloatVertexAttribute& LengthFromRootAttribute,
+		const FBooleanVertexAttribute& RootPointAttribute,
+		const FBooleanVertexAttribute& EndPointAttribute,
+		FIntVertexAttribute& NextBranchPointAttribute,
+		const TFunction<bool(int32)>& VertexFilterPredicate
+	)
+	{
+		using namespace PV::DynamicMeshHelper;
+
+		TRACE_CPUPROFILER_EVENT_SCOPE(TraceBranches);
+
+		for (int32 VertexID : DynamicMesh.VertexIndicesItr())
+		{
+			NextBranchPointAttribute.SetValue(VertexID, INDEX_NONE);
+		}
+
+		const auto FindNextPoint = [&](int32 PointIndex, float MinLengthFromRoot)->int32
+		{
+			int32 OutNextPointIndex = INDEX_NONE;
+			const TArray<int32> Neighbours = GetVertexNeighbours(DynamicMesh, PointIndex);
+			for (int32 Neighbour : Neighbours)
+			{
+				if (RootPointAttribute.GetValue(Neighbour))
+				{
+					OutNextPointIndex = Neighbour;
+					break;
+				}
+
+				if (!VertexFilterPredicate(Neighbour))
+				{
+					continue;
+				}
+
+				const float LengthFromRoot = LengthFromRootAttribute.GetValue(Neighbour);
+				check(LengthFromRoot >= 0);
+
+				if (LengthFromRoot >= MinLengthFromRoot)
+				{
+					continue;
+				}
+				MinLengthFromRoot = LengthFromRoot;
+				OutNextPointIndex = Neighbour;
+			}
+			return OutNextPointIndex;
+		};
+
+		const TArray<int32> EndPoints = EndPointAttribute.FindAllNonZero();
+
+		TArray<bool> bHasProcessedVertex;
+
+		for (int32 i = 0; i < EndPoints.Num(); i++)
+		{
+			int32 CurrentPointIndex = EndPoints[i];
+
+			float MinLengthFromRoot = LengthFromRootAttribute.GetValue(CurrentPointIndex);
+			if (MinLengthFromRoot < 0)
+			{
+				continue; // This end point was unable to trace down to a root point, perhaps log a warning?
+			}
+
+			bHasProcessedVertex.Reset();
+			bHasProcessedVertex.SetNumZeroed(DynamicMesh.MaxVertexID());
+
+			while (true)
+			{
+				const int32 NextPointIndex = FindNextPoint(CurrentPointIndex, MinLengthFromRoot);
+				if (NextPointIndex == INDEX_NONE)
+				{
+					break;
+				}
+				
+				NextBranchPointAttribute.SetValue(CurrentPointIndex, NextPointIndex);
+
+				if (NextBranchPointAttribute.GetValue(NextPointIndex) != INDEX_NONE
+					|| RootPointAttribute.GetValue(NextPointIndex))
+				{
+					break;
+				}
+
+				if (!ensureMsgf(!bHasProcessedVertex[NextPointIndex], TEXT("Infinite loop detected")))
+				{
+					break;
+				}
+
+				bHasProcessedVertex[NextPointIndex] = true;
+
+				CurrentPointIndex = NextPointIndex;
+				MinLengthFromRoot = LengthFromRootAttribute.GetValue(CurrentPointIndex);
+			}
+		}
+	}
+
+	struct FMeshBranchHierarchy
+	{
+		struct FBranch
+		{
+			TArray<int32> BranchPoints;
+			int32 ParentBranch = INDEX_NONE;
+		};
+
+		TArray<FBranch> Branches;
+	};
+
+	static void ExtractBranchHierarchies(
+		const UE::Geometry::FDynamicMesh3& DynamicMesh,
+		const FBooleanVertexAttribute& RootPointAttribute,
+		const FBooleanVertexAttribute& EndPointAttribute,
+		const FIntVertexAttribute& NextBranchPointAttribute,
+		TArray<FMeshBranchHierarchy>& OutMeshBranchHierarchies
+	)
+	{
+		using namespace PV::DynamicMeshHelper;
+
+		TRACE_CPUPROFILER_EVENT_SCOPE(ExtractBranchHierarchies);
+
+		TArray<bool> bProcessedVertices;
+		bProcessedVertices.SetNumZeroed(DynamicMesh.MaxVertexID());
+
+		const auto FindBestParentPoint = [&](int32 PrevVertexID, int32 VertexID, const TArray<int32>& NextPoints)
+		{
+			if (NextPoints.Num() == 1)
+			{
+				return NextPoints[0];
+			}
+
+			check(NextPoints.Num() > 0);
+
+			const FVector3d& VertexPos = DynamicMesh.GetVertex(VertexID);
+			const FVector3d& PrevVertexPos = DynamicMesh.GetVertex(PrevVertexID);
+			const FVector3d ApicalDirection = (VertexPos - PrevVertexPos).GetSafeNormal();
+
+			float BestDot = -1;
+			int32 BestDotIndex = INDEX_NONE;
+			for (int32 NextVertexID : NextPoints)
+			{
+				const FVector3d& NextVertexPos = DynamicMesh.GetVertex(NextVertexID);
+				const FVector3d NextApicalDirection = (NextVertexPos - VertexPos).GetSafeNormal();
+				const float Dot = FVector3d::DotProduct(ApicalDirection, NextApicalDirection);
+				if (Dot > BestDot)
+				{
+					BestDot = Dot;
+					BestDotIndex = NextVertexID;
+				}
+			}
+
+			check(BestDotIndex != INDEX_NONE);
+			return BestDotIndex;
+		};
+
+		const auto FindNextBranchPoints = [&](int32 VertexID)
+		{
+			TArray<int32> NextBranchPoints;
+			const TArray<int32> Neighbours = GetVertexNeighbours(DynamicMesh, VertexID);
+
+			for (int32 Neighbour : Neighbours)
+			{
+				if (NextBranchPointAttribute.GetValue(Neighbour) == VertexID
+					&& !bProcessedVertices[Neighbour])
+				{
+					NextBranchPoints.Add(Neighbour);
+				}
+			}
+			return NextBranchPoints;
+		};
+
+		const TFunction<void(int32, TArray<FMeshBranchHierarchy::FBranch>&)> TraceBranchesRecursive = [&](
+			int32 VertexID, 
+			TArray<FMeshBranchHierarchy::FBranch>& OutBranches
+		)
+		{
+			const int32 BranchIndex = OutBranches.Num() - 1;
+			OutBranches[BranchIndex].BranchPoints.Add(VertexID);
+			
+			bProcessedVertices[VertexID] = true;
+
+			if (EndPointAttribute.GetValue(VertexID))
+			{
+				return;
+			}
+
+			const TArray<int32> NextBranchPoints = FindNextBranchPoints(VertexID);
+			check(NextBranchPoints.Num() > 0);
+
+			for (int32 NextVertexID : NextBranchPoints)
+			{
+				bProcessedVertices[NextVertexID] = true;
+			}
+
+			const int32 PrevVertexID = OutBranches[BranchIndex].BranchPoints.Last(1);
+			const int32 NextBranchPointsIndex = FindBestParentPoint(PrevVertexID, VertexID, NextBranchPoints);
+			
+			TraceBranchesRecursive(NextBranchPointsIndex, OutBranches);
+
+			for (int32 NextVertexID : NextBranchPoints)
+			{
+				if (NextVertexID != NextBranchPointsIndex)
+				{
+					FMeshBranchHierarchy::FBranch& NextBranch = OutBranches.AddDefaulted_GetRef();
+					NextBranch.BranchPoints.Add(VertexID);
+					NextBranch.ParentBranch = BranchIndex;
+
+					TraceBranchesRecursive(NextVertexID, OutBranches);
+				}
+			}
+		};
+
+		const TArray<int32> RootPoints = RootPointAttribute.FindAllNonZero();
+		OutMeshBranchHierarchies.Reserve(RootPoints.Num());
+
+		for (const int32 RootPointIndex : RootPoints)
+		{
+			TArray<TArray<FMeshBranchHierarchy::FBranch>> BranchHierarchies;
+
+			const TArray<int32> NextBranchPoints = FindNextBranchPoints(RootPointIndex);
+			for (int32 NextVertexID : NextBranchPoints)
+			{
+				bProcessedVertices[NextVertexID] = true;
+			}
+
+			for (int32 NextVertexID : NextBranchPoints)
+			{
+				TArray<FMeshBranchHierarchy::FBranch>& Branches = BranchHierarchies.AddDefaulted_GetRef();
+				FMeshBranchHierarchy::FBranch& RootBranch = Branches.AddDefaulted_GetRef();
+				RootBranch.BranchPoints.Add(RootPointIndex);
+				TraceBranchesRecursive(NextVertexID, Branches);
+			}
+
+			if (BranchHierarchies.Num() > 0)
+			{
+				// TODO: Handle multiple branches starting from root
+				TArray<FMeshBranchHierarchy::FBranch>* BestBranches = Algo::MaxElement(
+					BranchHierarchies, 
+					[](const auto& A, const auto& B) { return A.Num() < B.Num(); }
+				);
+				OutMeshBranchHierarchies.Emplace(MoveTemp(*BestBranches));
+			}
+		}
+	}
+
+	static void ConvertMeshBranchHierarchiesToBranchDescription(
+		const UE::Geometry::FDynamicMesh3& DynamicMesh,
+		const FFloatVertexAttribute& VertexScaleAttribute,
+		const FBooleanVertexAttribute& EndPointAttribute,
+		const TArray<FMeshBranchHierarchy>& InMeshBranchHierarchies,
+		TArray<FPVBranchHierarchyDescription>& OutBranchHierarchyDescriptions
+	)
+	{
+		TRACE_CPUPROFILER_EVENT_SCOPE(ConvertMeshBranchHierarchiesToBranchDescription);
+
+		for (const FMeshBranchHierarchy& BranchHierarchy : InMeshBranchHierarchies)
+		{
+			int32 MaxNrOfPoints = 0;
+			for (const auto& Branch : BranchHierarchy.Branches)
+			{
+				MaxNrOfPoints += Branch.BranchPoints.Num();
+			}
+
+			FPVBranchHierarchyDescription& OutBranchHierarchy = OutBranchHierarchyDescriptions.AddDefaulted_GetRef();
+			OutBranchHierarchy.Branches.Reserve(BranchHierarchy.Branches.Num());
+			OutBranchHierarchy.Points.Reserve(MaxNrOfPoints);
+			OutBranchHierarchy.PointsRadii.Reserve(MaxNrOfPoints);
+
+			TMap<int32, int32> VertexIndexToPointIndex;
+			VertexIndexToPointIndex.Reserve(MaxNrOfPoints);
+
+			for (const FMeshBranchHierarchy::FBranch& Branch : BranchHierarchy.Branches)
+			{
+				FPVBranchDescription& OutBranchDescription = OutBranchHierarchy.Branches.AddDefaulted_GetRef();
+				OutBranchDescription.ParentBranchIndex = Branch.ParentBranch;
+				OutBranchDescription.BranchIndex = OutBranchHierarchy.Branches.Num() - 1;
+				OutBranchDescription.PointIndices.Reserve(Branch.BranchPoints.Num());
+
+				for (int32 i = 0; i < Branch.BranchPoints.Num(); ++i)
+				{
+					const int32 VertexIndex = Branch.BranchPoints[i];
+					int32 PointIndex = INDEX_NONE;
+
+					int32* ExistingPointIndex = VertexIndexToPointIndex.Find(VertexIndex);
+					if (!ExistingPointIndex)
+					{
+						PointIndex = OutBranchHierarchy.Points.Add(FVector3f(DynamicMesh.GetVertex(VertexIndex)));
+
+						const bool bIsEndPoint = EndPointAttribute.GetValue(VertexIndex);
+						if (bIsEndPoint)
+						{
+							OutBranchHierarchy.PointsRadii.Add(0.01f);
+						}
+						else
+						{
+							const float VertexRadius = VertexScaleAttribute.GetValue(VertexIndex);
+							OutBranchHierarchy.PointsRadii.Add(VertexRadius);
+						}
+
+						VertexIndexToPointIndex.Add(VertexIndex, PointIndex);
+					}
+					else
+					{
+						PointIndex = *ExistingPointIndex;
+					}
+
+					OutBranchDescription.PointIndices.Add(PointIndex);
+				}
+			}
+
+			for (int32 i = 0; i < OutBranchHierarchy.Branches.Num(); ++i)
+			{
+				const FPVBranchDescription& BranchDescription = OutBranchHierarchy.Branches[i];
+				if (BranchDescription.ParentBranchIndex != INDEX_NONE)
+				{
+					OutBranchHierarchy.Branches[BranchDescription.ParentBranchIndex].ChildBranchIndices.Add(i);
+				}
+				else
+				{
+					check(OutBranchHierarchy.RootBranchIndex == INDEX_NONE);
+					OutBranchHierarchy.RootBranchIndex = i;
+				}
+			}
+		}
+	}
+
+	const static TSet<FName> GeometryGroups = {
+		FGeometryCollection::VerticesGroup,
+		FGeometryCollection::FacesGroup,
+		FGeometryCollection::GeometryGroup,
+		FGeometryCollection::BreakingGroup,
+		FGeometryCollection::MaterialGroup,
+		FTransformCollection::TransformGroup,
+		FTransformCollection::ConvexGroup
+	};
+
+	const static FName MaterialPathAttributeName("MaterialPath");
+
+	static void AppendGeometryFromManagedArrayCollection(FGeometryCollection& InOutGeometryCollection, const FManagedArrayCollection& InManagedArrayCollection)
+	{
+		TArray<FName> GroupsToSkip;
+		GroupsToSkip.Reserve(InManagedArrayCollection.GroupNames().Num());
+
+		for (FName GroupName : InManagedArrayCollection.GroupNames())
+		{
+			if (!GeometryGroups.Contains(GroupName))
+			{
+				GroupsToSkip.Add(GroupName);
+			}
+		}
+
+		FGeometryCollection GeometryCollection;
+		InManagedArrayCollection.CopyTo(&GeometryCollection, GroupsToSkip);
+
+		const auto* ExistingMaterialPathsAttribute = InOutGeometryCollection.FindAttribute<FString>(MaterialPathAttributeName, FGeometryCollection::MaterialGroup);
+		const int32 NumExistingMaterials = ExistingMaterialPathsAttribute ? ExistingMaterialPathsAttribute->Num() : 1;
+		const bool bHasCopiedGeometry = InOutGeometryCollection.AppendGeometry(GeometryCollection, NumExistingMaterials, false) != INDEX_NONE;
+
+		// Copy materials. Ideally we'd want to re-use materials where possible, but given this is for debugging we just assign 
+		// one material per mesh.
+		if (bHasCopiedGeometry)
+		{
+			if (!InOutGeometryCollection.HasGroup(FGeometryCollection::MaterialGroup))
+			{
+				InOutGeometryCollection.AddGroup(FGeometryCollection::MaterialGroup);
+			}
+
+			auto* MaterialPathsAttributeToCopy = GeometryCollection.FindAttribute<FString>(MaterialPathAttributeName, FGeometryCollection::MaterialGroup);
+			const int32 NumMaterialsToCopy = MaterialPathsAttributeToCopy ? MaterialPathsAttributeToCopy->Num() : 1;
+			InOutGeometryCollection.Resize(NumExistingMaterials + NumMaterialsToCopy, FGeometryCollection::MaterialGroup);
+
+			if (MaterialPathsAttributeToCopy)
+			{
+				auto* MaterialPathsAttribute = InOutGeometryCollection.FindOrAddAttributeTyped<FString>(MaterialPathAttributeName, FGeometryCollection::MaterialGroup);
+				for (int32 i = NumExistingMaterials; i < NumExistingMaterials + NumMaterialsToCopy; ++i)
+				{
+					const int32 MaterialIndexToCopy = i - NumExistingMaterials;
+					if (MaterialPathsAttributeToCopy->IsValidIndex(MaterialIndexToCopy))
+					{
+						(*MaterialPathsAttribute)[i] = (*MaterialPathsAttributeToCopy)[MaterialIndexToCopy];
+					}
+				}
+			}
+		}
+	}
+
+	static void SetManagedArrayCollectionGeometry(const FGeometryCollection& InGeometryCollection, FManagedArrayCollection& InOutManagedArrayCollection)
+	{
+		for (const FName GroupName : GeometryGroups)
+		{
+			InOutManagedArrayCollection.RemoveGroup(GroupName);
+		}
+
+		InGeometryCollection.CopyTo(&InOutManagedArrayCollection);
+	}
+}
+
+const auto PV::DynamicMeshHelper::VertexScaleAttributeDefinition     = MakeUnique<PV::TDynamicMeshVertexAttributeDefinition<float>>(TEXT("VertexScale"));
+const auto PV::DynamicMeshHelper::EndPointAttributeDefinition        = MakeUnique<PV::TDynamicMeshVertexAttributeDefinition<bool>>(TEXT("EndPoint"));
+const auto PV::DynamicMeshHelper::RootPointAttributeDefinition       = MakeUnique<PV::TDynamicMeshVertexAttributeDefinition<bool>>(TEXT("RootPoint"));
+const auto PV::DynamicMeshHelper::LengthFromRootAttributeDefinition  = MakeUnique<PV::TDynamicMeshVertexAttributeDefinition<float>>(TEXT("LengthFromRoot"));
+const auto PV::DynamicMeshHelper::NextBranchPointAttributeDefinition = MakeUnique<PV::TDynamicMeshVertexAttributeDefinition<int32>>(TEXT("NextBranchPoint"));
+
+void PV::DynamicMeshHelper::AddDynamicMeshToCollection(const UE::Geometry::FDynamicMesh3& DynamicMesh, FManagedArrayCollection& InOutCollection)
+{
+	FGeometryCollection MeshGeomCollection;
+
+	UE::Geometry::FGeometryCollectionToDynamicMeshes::FToCollectionOptions Options2;
+	UE::Geometry::FGeometryCollectionToDynamicMeshes::AppendMeshToCollection(MeshGeomCollection, DynamicMesh, FTransform::Identity, Options2);
+
+	AppendGeometryFromManagedArrayCollection(MeshGeomCollection, InOutCollection);
+	SetManagedArrayCollectionGeometry(MeshGeomCollection, InOutCollection);
+}
+
+void PV::DynamicMeshHelper::AddBranchHierarchiesToCollection(const TArray<FPVBranchHierarchyDescription>& BranchHierarchyDescriptions, FManagedArrayCollection& OutCollection)
+{
+	const static FName CurveAttributeName = TEXT("Curve");
+	const static FName BranchCurvesGroupName = TEXT("BranchCurves");
+	OutCollection.AddGroup(BranchCurvesGroupName);
+	int32 AttributeIndex = 0;
+	for (const FPVBranchHierarchyDescription& BranchHierarchy : BranchHierarchyDescriptions)
+	{
+		OutCollection.AddElements(BranchHierarchy.Branches.Num(), BranchCurvesGroupName);
+
+		TManagedArray<TArray<FVector3f>>& CurveAttribute = *OutCollection.FindOrAddAttributeTyped<TArray<FVector3f>>(CurveAttributeName, BranchCurvesGroupName);
+		for (int32 i = 0; i < BranchHierarchy.Branches.Num(); ++i)
+		{
+			const FPVBranchDescription& Branch = BranchHierarchy.Branches[i];
+			CurveAttribute[AttributeIndex].Reserve(Branch.PointIndices.Num());
+			for (int32 VertexIndex : Branch.PointIndices)
+			{
+				CurveAttribute[AttributeIndex].Add(BranchHierarchy.Points[VertexIndex]);
+			}
+			AttributeIndex++;
+		}
+	}
+}
+
+TArray<int32> PV::DynamicMeshHelper::GetUniqueVertexIndices(const UE::Geometry::FDynamicMesh3& DynamicMesh, const TArray<int32>& TriangleIds)
+{
+	TSet<int32> UniqueVertexIndices;
+	UniqueVertexIndices.Reserve(TriangleIds.Num() * 3);
+	for (int32 tid : TriangleIds)
+	{
+		const UE::Geometry::FIndex3i VertexIndices = DynamicMesh.GetTriangle(tid);
+		UniqueVertexIndices.Add(VertexIndices[0]);
+		UniqueVertexIndices.Add(VertexIndices[1]);
+		UniqueVertexIndices.Add(VertexIndices[2]);
+	}
+
+	return UniqueVertexIndices.Array();
+}
+
+TArray<int32> PV::DynamicMeshHelper::GetVertexNeighbours(const UE::Geometry::FDynamicMesh3& DynamicMesh, int32 VertexID)
+{
+	TArray<int32> OutNeighbours;
+	OutNeighbours.Reserve(DynamicMesh.GetVtxEdgeCount(VertexID));
+	for (int32 EdgeID : DynamicMesh.VtxEdgesItr(VertexID))
+	{
+		const UE::Geometry::FIndex2i VertexIDs = DynamicMesh.GetEdgeV(EdgeID);
+		const int32 OtherVertexID = VertexIDs.A == VertexID ? VertexIDs.B : VertexIDs.A;
+		OutNeighbours.Add(OtherVertexID);
+	}
+	return OutNeighbours;
+}
+
+TArray<int32> PV::DynamicMeshHelper::FindVerticesInRadius(const UE::Geometry::FDynamicMeshAABBTree3& DynamicMeshAABBTree, const FVector& Pos, float Radius)
+{
+	TArray<int32> OutVertexIndices;
+
+	const UE::Geometry::FDynamicMesh3* DynamicMesh = DynamicMeshAABBTree.GetMesh();
+	if (DynamicMesh == nullptr)
+	{
+		return OutVertexIndices;
+	}
+
+	OutVertexIndices.Reserve(DynamicMesh->VertexCount());
+
+	const double RadiusSqr = Radius * Radius;
+
+	UE::Geometry::FDynamicMeshAABBTree3::FTreeTraversal TreeTraversal;
+
+	TreeTraversal.NextBoxF = [&](const UE::Geometry::FAxisAlignedBox3d& Box, int32 Depth)->bool
+	{
+		return Box.DistanceSquared(Pos) <= RadiusSqr;
+	};
+
+	TreeTraversal.NextTriangleF = [&](int32 TID)
+	{
+		const UE::Geometry::FIndex3i Triangle = DynamicMesh->GetTriangle(TID);
+		for (int32 j = 0; j < 3; ++j)
+		{
+			const double VertexDistSqr = (DynamicMesh->GetVertex(Triangle[j]) - Pos).SizeSquared();
+			if (VertexDistSqr <= RadiusSqr)
+			{
+				OutVertexIndices.AddUnique(Triangle[j]);
+			}
+		}
+	};
+
+	DynamicMeshAABBTree.DoTraversal(TreeTraversal);
+
+	return OutVertexIndices;
+}
+
+void PV::DynamicMeshHelper::ComputeBranchHierarchies(
+	const UE::Geometry::FDynamicMesh3& DynamicMesh,
+	const FBooleanVertexAttribute& RootPointsAttribute,
+	const FBooleanVertexAttribute& EndPointsAttribute,
+	const FFloatVertexAttribute& VertexScaleAttribute,
+	FFloatVertexAttribute& LengthFromRootAttribute,
+	FIntVertexAttribute& NextBranchPointAttribute,
+	const TFunction<bool(int32)>& VertexFilterPredicate,
+	TArray<FPVBranchHierarchyDescription>& OutBranchHierarchyDescriptions
+)
+{
+	ComputeLengthFromRoot(
+		DynamicMesh,
+		RootPointsAttribute,
+		EndPointsAttribute,
+		LengthFromRootAttribute,
+		VertexFilterPredicate
+	);
+
+	TraceBranches(
+		DynamicMesh,
+		LengthFromRootAttribute,
+		RootPointsAttribute,
+		EndPointsAttribute,
+		NextBranchPointAttribute,
+		VertexFilterPredicate
+	);
+
+	TArray<FMeshBranchHierarchy> BranchHierarchies;
+	ExtractBranchHierarchies(
+		DynamicMesh,
+		RootPointsAttribute,
+		EndPointsAttribute,
+		NextBranchPointAttribute,
+		BranchHierarchies
+	);
+
+	ConvertMeshBranchHierarchiesToBranchDescription(
+		DynamicMesh,
+		VertexScaleAttribute,
+		EndPointsAttribute,
+		BranchHierarchies,
+		OutBranchHierarchyDescriptions
+	);
+}
+
 void PV::Import::Helper::CreateEmptyGrowthData(FManagedArrayCollection& OutCollection, int32 NumBranchPoints, int32 NumFoliagePoints, int32 NumBranches, int32 NumFoliageMeshes)
 {
 	// PointGroup
@@ -131,7 +748,10 @@ void PV::Import::Helper::CreateEmptyGrowthData(FManagedArrayCollection& OutColle
 	LeafPhyllotaxyAttribute[0] = { 0, 200, 50, 1, 1, 0, 0, 0 };
 }
 
-void PV::Import::Helper::FillGrowthDataFromBranchHierarchyDescription(FManagedArrayCollection& InOutCollection, const FPVBranchHierarchyDescription& InBranchHierarchy)
+void PV::Import::Helper::FillGrowthDataFromBranchHierarchyDescription(
+	FManagedArrayCollection& InOutCollection, 
+	const TArray<FPVBranchHierarchyDescription>& InBranchHierarchies
+)
 {
 	PV::Facades::FPointFacade PointFacade(InOutCollection);
 	PV::Facades::FBranchFacade BranchFacade(InOutCollection);
@@ -142,80 +762,122 @@ void PV::Import::Helper::FillGrowthDataFromBranchHierarchyDescription(FManagedAr
 		return BranchIndex + 1; 
 	}; 
 
-	TSet<int32> AddedNodes;
-	int32 FirstPointIndex = 0;
-	int32 BudNumber = 1;
-	const TFunction<void(int32, int32)> FillBranchesRecursive = [&](int32 BranchIndex, int32 Depth)
+	using FFillBranchesRecursive = TFunction<void(const FPVBranchHierarchyDescription&, int32, int32, int32, int32, int32, int32&)>;
+
+	const FFillBranchesRecursive FillBranchesRecursive = [&PointFacade, &BranchFacade, &FillBranchesRecursive](
+		const FPVBranchHierarchyDescription& BranchHierarchy, 
+		int32 PlantNumber,
+		int32 BranchIndex,
+		int32 FacadeBranchIndexOffset,
+		int32 FacadePointIndexOffset,
+		int32 Depth,
+		int32& BudNumber
+	)
 	{
-		AddedNodes.Add(BranchIndex);
+		const FPVBranchDescription& Branch = BranchHierarchy.Branches[BranchIndex];
+
+		const int32 FacadeRootBranchIndex = BranchHierarchy.RootBranchIndex + FacadeBranchIndexOffset;
+		const int32 FacadeBranchIndex = BranchIndex + FacadeBranchIndexOffset;
+		const int32 FacadeParentBranchIndex = Branch.ParentBranchIndex != INDEX_NONE 
+			? Branch.ParentBranchIndex + FacadeBranchIndexOffset 
+			: INDEX_NONE;
 
-		const FPVBranchDescription& Branch = InBranchHierarchy.Branches[BranchIndex];
-		const int32 LastPointIndex = FirstPointIndex + Branch.PointIndices.Num() - 1;
+		TArray<int32> FacadeBranchPoints;
+		FacadeBranchPoints.Reserve(Branch.PointIndices.Num());
 
-		TArray<int32> BranchPoints;
-		BranchPoints.Reserve(Branch.PointIndices.Num());
-		for (int32 i = FirstPointIndex; i <= LastPointIndex; ++i)
+		for (int32 i = 0; i < Branch.PointIndices.Num(); ++i)
 		{
-			const int32 PointIndex = Branch.PointIndices[i - FirstPointIndex];
-			BranchPoints.Add(PointIndex);
+			const int32 FacadePointIndex = Branch.PointIndices[i] + FacadePointIndexOffset;
+			FacadeBranchPoints.Add(FacadePointIndex);
 
-			PointFacade.SetBudNumber(PointIndex, BudNumber);
+			PointFacade.SetBudNumber(FacadePointIndex, BudNumber);
 
 			// Setting the njord pixel index to the same as the bud number only works since we're traversing the branches recursivly
 			// This shold be set by the mesher but we set it here since it's easy to do so.
-			PointFacade.SetNjordPixelIndex(PointIndex, BudNumber);
+			PointFacade.SetNjordPixelIndex(FacadePointIndex, BudNumber);
 			
 			BudNumber++;
 			
 			// TODO: Set these based on this doc https://docs.google.com/document/d/1fWNcORvDreHBOakcFfiKPhSXhE5yVTCabn4eWZU1jWg/edit?tab=t.0
 			//BudHormoneLevels[i] = { 0.f, 0.f, 0.f, 0.f, 0.f };
 		}
-		BranchFacade.SetPoints(BranchIndex, MoveTemp(BranchPoints));
+		BranchFacade.SetPoints(FacadeBranchIndex, MoveTemp(FacadeBranchPoints));
 
-		if (Branch.ParentBranchIndex != INDEX_NONE)
+		if (FacadeParentBranchIndex != INDEX_NONE)
 		{
-			TArray<int32> Parents = BranchFacade.GetParents(Branch.ParentBranchIndex);
-			Parents.Add(BranchIndexToBranchNumber(Branch.ParentBranchIndex));
-			BranchFacade.SetParents(BranchIndex, MoveTemp(Parents));
+			TArray<int32> Parents = BranchFacade.GetParents(FacadeParentBranchIndex);
+			Parents.Add(BranchIndexToBranchNumber(FacadeParentBranchIndex));
+			BranchFacade.SetParents(FacadeBranchIndex, MoveTemp(Parents));
 		}
 
-		TArray<int32> Children;
-		Children.SetNumUninitialized(Branch.ChildBranchIndices.Num());
+		TArray<int32> FacadeBranchChildren;
+		FacadeBranchChildren.SetNumUninitialized(Branch.ChildBranchIndices.Num());
 		for (int32 i = 0; i < Branch.ChildBranchIndices.Num(); ++i)
 		{
-			check(Branch.ChildBranchIndices[i] != InBranchHierarchy.RootBranchIndex);
-			check(!AddedNodes.Contains(Branch.ChildBranchIndices[i]));
-			Children[i] = BranchIndexToBranchNumber(Branch.ChildBranchIndices[i]);
-		}
-		BranchFacade.SetChildren(BranchIndex, MoveTemp(Children));
+			const int32 FacadeChildBranchIndex = Branch.ChildBranchIndices[i] + FacadeBranchIndexOffset;
 
-		BranchFacade.SetParentBranchNumber(BranchIndex, BranchIndexToBranchNumber(Branch.ParentBranchIndex));
-		BranchFacade.SetBranchNumber(BranchIndex, BranchIndexToBranchNumber(BranchIndex));
+			check(FacadeChildBranchIndex != FacadeRootBranchIndex);
+			FacadeBranchChildren[i] = BranchIndexToBranchNumber(FacadeChildBranchIndex);
+		}
+		BranchFacade.SetChildren(FacadeBranchIndex, MoveTemp(FacadeBranchChildren));
 
-		BranchFacade.SetBranchSourceBudNumber(BranchIndex, -1);
-		BranchFacade.SetBranchHierarchyNumber(BranchIndex, Depth);
-		BranchFacade.SetBranchPlantNumber(BranchIndex, 0);
+		BranchFacade.SetParentBranchNumber(FacadeBranchIndex, BranchIndexToBranchNumber(FacadeParentBranchIndex));
+		BranchFacade.SetBranchNumber(FacadeBranchIndex, BranchIndexToBranchNumber(FacadeBranchIndex));
 
-		FirstPointIndex += LastPointIndex - FirstPointIndex + 1;
+		BranchFacade.SetBranchSourceBudNumber(FacadeBranchIndex, -1);
+		BranchFacade.SetBranchHierarchyNumber(FacadeBranchIndex, Depth);
+		BranchFacade.SetBranchPlantNumber(FacadeBranchIndex, PlantNumber);
 
 		for (int32 ChildIndex : Branch.ChildBranchIndices)
 		{
-			FillBranchesRecursive(ChildIndex, Depth + 1);
+			FillBranchesRecursive(
+				BranchHierarchy, 
+				PlantNumber,
+				ChildIndex, 
+				FacadeBranchIndexOffset,
+				FacadePointIndexOffset,
+				Depth + 1, 
+				BudNumber
+			);
 		}
 	};
 
-	FillBranchesRecursive(InBranchHierarchy.RootBranchIndex, 0);
-
-	TManagedArray<FVector3f>& Positions = PointFacade.ModifyPositions();
-	TManagedArray<float>& Scales = PointFacade.ModifyPointScales();
-	for (int32 i = 0; i < InBranchHierarchy.Points.Num(); ++i)
+	int32 FacadeBranchIndexOffset = 0;
+	int32 FacadePointIndexOffset = 0;
+	int32 PlantNumber = 0;
+	for (const FPVBranchHierarchyDescription& BranchHierarchy : InBranchHierarchies)
 	{
-		Positions[i] = InBranchHierarchy.Points[i];
-		Scales[i] = InBranchHierarchy.PointsRadii[i] > 0 ? InBranchHierarchy.PointsRadii[i] : 0;
+		int32 BudNumber = 1;
+
+		FillBranchesRecursive(
+			BranchHierarchy, 
+			PlantNumber,
+			BranchHierarchy.RootBranchIndex, 
+			FacadeBranchIndexOffset,
+			FacadePointIndexOffset,
+			0, 
+			BudNumber
+		);
+
+		TManagedArray<FVector3f>& Positions = PointFacade.ModifyPositions();
+		TManagedArray<float>& Scales = PointFacade.ModifyPointScales();
+		for (int32 i = 0; i < BranchHierarchy.Points.Num(); ++i)
+		{
+			Positions[i + FacadePointIndexOffset] = BranchHierarchy.Points[i];
+			Scales[i + FacadePointIndexOffset] = BranchHierarchy.PointsRadii[i] > 0 ? BranchHierarchy.PointsRadii[i] : 0;
+		}
+
+		FacadeBranchIndexOffset += BranchHierarchy.Branches.Num();
+		FacadePointIndexOffset += BranchHierarchy.Points.Num();
+		PlantNumber += 1;
 	}
 }
 
-void PV::Import::Helper::FillGrowthDataFromFoliageDescription(FManagedArrayCollection& Collection, const FPVFoliageDescription& InFoliageDescription, const FPVBranchHierarchyDescription& InBranchHierarchy)
+void PV::Import::Helper::FillGrowthDataFromFoliageDescription(
+	FManagedArrayCollection& Collection, 
+	const FPVFoliageDescription& InFoliageDescription, 
+	const TArray<FPVBranchHierarchyDescription>& InBranchHierarchies
+)
 {
 	PV::Facades::FFoliageFacade FoliageFacade(Collection);
 	PV::Facades::FPointFacade PointFacade(Collection);
@@ -228,20 +890,26 @@ void PV::Import::Helper::FillGrowthDataFromFoliageDescription(FManagedArrayColle
 		OutParentPointIndex = INDEX_NONE;
 		OutClosestPointOnBranch = FVector3f::ZeroVector;
 
-		for (int32 i = 0; i < InBranchHierarchy.Branches.Num(); ++i)
+		int32 BranchIndexOffset = 0;
+		for (const FPVBranchHierarchyDescription& BranchHierarchy : InBranchHierarchies)
 		{
-			FVector3f ClosestPointOnBranch;
-			float BranchDistSquared;
-			int32 PointIndex;
-			DistToBranch(Point, InBranchHierarchy.Points, InBranchHierarchy.Branches[i].PointIndices, ClosestPointOnBranch, BranchDistSquared, PointIndex);
-			if (BranchDistSquared < ClosestBranchDistSquared)
+			for (int32 i = 0; i < BranchHierarchy.Branches.Num(); ++i)
 			{
-				ClosestBranchDistSquared = BranchDistSquared;
+				FVector3f ClosestPointOnBranch;
+				float BranchDistSquared;
+				int32 PointIndex;
+				DistToBranch(Point, BranchHierarchy.Points, BranchHierarchy.Branches[i].PointIndices, ClosestPointOnBranch, BranchDistSquared, PointIndex);
+				if (BranchDistSquared < ClosestBranchDistSquared)
+				{
+					ClosestBranchDistSquared = BranchDistSquared;
 
-				OutBranchIndex = i;
-				OutParentPointIndex = PointIndex;
-				OutClosestPointOnBranch = ClosestPointOnBranch;
+					OutBranchIndex = i + BranchIndexOffset;
+					OutParentPointIndex = PointIndex;
+					OutClosestPointOnBranch = ClosestPointOnBranch;
+				}
 			}
+
+			BranchIndexOffset += BranchHierarchy.Branches.Num();
 		}
 	};
 
@@ -903,22 +1571,30 @@ void PV::Import::Helper::FuseCollinearBranches(FPVBranchHierarchyDescription& In
 
 bool PV::Import::Helper::GenerateGrowthDataFromPlantDescription(FManagedArrayCollection& OutCollection, const FPVPlantDescription& InPlantDescription)
 {
-	const FPVBranchHierarchyDescription& InBranchHierarchy = InPlantDescription.BranchHierarchy;
+	const TArray<FPVBranchHierarchyDescription>& InBranchHierarchies = InPlantDescription.BranchHierarchies;
 	const FPVFoliageDescription& InFoliageDescription = InPlantDescription.Foliage;
 
-	if (InBranchHierarchy.RootBranchIndex == INDEX_NONE)
+	if (InBranchHierarchies.Num() < 1)
 	{
 		return false;
 	}
 
+	for (const FPVBranchHierarchyDescription& BranchHierarchy : InBranchHierarchies)
+	{
+		if (BranchHierarchy.ValidateHierarchy() != FPVBranchHierarchyDescription::EValidateHierarchyResult::Success)
+		{
+			return false;
+		}
+	}
+
 	const int32 NumBranchPoints = InPlantDescription.GetNumBranchPoints();
 	const int32 NumFoliagePoints = InPlantDescription.GetNumFoliagePoints();
-	const int32 NumBranches = InBranchHierarchy.Branches.Num();
+	const int32 NumBranches = InPlantDescription.GetNumBranches();
 	const int32 NumFoliageMeshes = InFoliageDescription.FoliageMeshes.Num();
 
 	CreateEmptyGrowthData(OutCollection, NumBranchPoints, NumFoliagePoints, NumBranches, NumFoliageMeshes);
-	FillGrowthDataFromBranchHierarchyDescription(OutCollection, InBranchHierarchy);
-	FillGrowthDataFromFoliageDescription(OutCollection, InFoliageDescription, InBranchHierarchy);
+	FillGrowthDataFromBranchHierarchyDescription(OutCollection, InBranchHierarchies);
+	FillGrowthDataFromFoliageDescription(OutCollection, InFoliageDescription, InBranchHierarchies);
 
 	PV::Facades::FPointFacade PointFacade(OutCollection);
 	const PV::Facades::FBranchFacade BranchFacade(OutCollection);
@@ -1129,4 +1805,4 @@ FPVBranchHierarchyDescription PV::Import::Helper::CreateBranchHierarchyFromPoint
 	}
 
 	return OutBranchHierarchy;
-}
\ No newline at end of file
+}
diff --git a/Engine/Plugins/Experimental/ProceduralVegetationEditor/Source/ProceduralVegetationEditor/Private/Importer/PVImportHelpers.h b/Engine/Plugins/Experimental/ProceduralVegetationEditor/Source/ProceduralVegetationEditor/Private/Importer/PVImportHelpers.h
--- a/Engine/Plugins/Experimental/ProceduralVegetationEditor/Source/ProceduralVegetationEditor/Private/Importer/PVImportHelpers.h
+++ b/Engine/Plugins/Experimental/ProceduralVegetationEditor/Source/ProceduralVegetationEditor/Private/Importer/PVImportHelpers.h
@@ -14,13 +14,68 @@ struct FManagedArrayCollection;
 class UProceduralVegetationPreset;
 class UPackage;
 
+namespace PV
+{
+	template<typename T> class TDynamicMeshVertexAttributeDefinition;
+	template<typename AttribValueType> class TDynamicMeshVertexAttributeExt;
+	template<typename AttribValueType> class TVertexAttributeView;
+}
+
+namespace UE::Geometry
+{
+	class FDynamicMesh3;
+	template<class T> class TMeshAABBTree3;
+	using FDynamicMeshAABBTree3 = TMeshAABBTree3<FDynamicMesh3>;
+}
+
+namespace PV::DynamicMeshHelper
+{
+	using FBooleanVertexAttribute = TDynamicMeshVertexAttributeExt<bool>;
+	using FIntVertexAttribute     = TDynamicMeshVertexAttributeExt<int32>;
+	using FFloatVertexAttribute   = TDynamicMeshVertexAttributeExt<float>;
+
+	extern const TUniquePtr<TDynamicMeshVertexAttributeDefinition<float>> VertexScaleAttributeDefinition;
+	extern const TUniquePtr<TDynamicMeshVertexAttributeDefinition<bool>> EndPointAttributeDefinition;
+	extern const TUniquePtr<TDynamicMeshVertexAttributeDefinition<bool>> RootPointAttributeDefinition;
+	extern const TUniquePtr<TDynamicMeshVertexAttributeDefinition<float>> LengthFromRootAttributeDefinition;
+	extern const TUniquePtr<TDynamicMeshVertexAttributeDefinition<int32>> NextBranchPointAttributeDefinition;
+
+	void AddDynamicMeshToCollection(const UE::Geometry::FDynamicMesh3& DynamicMesh, FManagedArrayCollection& InOutCollection);
+	
+	void AddBranchHierarchiesToCollection(const TArray<FPVBranchHierarchyDescription>& BranchHierarchyDescriptions, FManagedArrayCollection& OutCollection);
+
+	TArray<int32> GetUniqueVertexIndices(const UE::Geometry::FDynamicMesh3& DynamicMesh, const TArray<int32>& TriangleIds);
+
+	TArray<int32> GetVertexNeighbours(const UE::Geometry::FDynamicMesh3& DynamicMesh, int32 VertexID);
+
+	TArray<int32> FindVerticesInRadius(const UE::Geometry::FDynamicMeshAABBTree3& DynamicMeshAABBTree, const FVector& Pos, float Radius);
+	
+	void ComputeBranchHierarchies(
+		const UE::Geometry::FDynamicMesh3& DynamicMesh,
+		const FBooleanVertexAttribute& RootPointsAttribute,
+		const FBooleanVertexAttribute& EndPointsAttribute,
+		const FFloatVertexAttribute& VertexScaleAttribute,
+		FFloatVertexAttribute& LengthFromRootAttribute,
+		FIntVertexAttribute& NextBranchPointAttribute,
+		const TFunction<bool(int32)>& VertexFilterPredicate,
+		TArray<FPVBranchHierarchyDescription>& OutBranchHierarchyDescriptions
+	);
+}
+
 namespace PV::Import::Helper
 {
 	void CreateEmptyGrowthData(FManagedArrayCollection& OutCollection, int32 NumBranchPoints, int32 NumFoliagePoints, int32 NumBranches, int32 NumFoliageMeshes);
 
-	void FillGrowthDataFromBranchHierarchyDescription(FManagedArrayCollection& Collection, const FPVBranchHierarchyDescription& InBranchHierarchy);
+	void FillGrowthDataFromBranchHierarchyDescription(
+		FManagedArrayCollection& Collection, 
+		const TArray<FPVBranchHierarchyDescription>& InBranchHierarchies
+	);
 
-	void FillGrowthDataFromFoliageDescription(FManagedArrayCollection& Collection, const FPVFoliageDescription& InFoliageDescription, const FPVBranchHierarchyDescription& InBranchHierarchy);
+	void FillGrowthDataFromFoliageDescription(
+		FManagedArrayCollection& Collection, 
+		const FPVFoliageDescription& InFoliageDescription, 
+		const TArray<FPVBranchHierarchyDescription>& InBranchHierarchies
+	);
 
 	void ComputeLengthFromRoot(PV::Facades::FPointFacade& PointFacade, const PV::Facades::FBranchFacade& BranchFacade, const PV::Facades::FPlantFacade& PlantFacade);
 
diff --git a/Engine/Plugins/Experimental/ProceduralVegetationEditor/Source/ProceduralVegetationEditor/Private/Importer/PVImportTests.cpp b/Engine/Plugins/Experimental/ProceduralVegetationEditor/Source/ProceduralVegetationEditor/Private/Importer/PVImportTests.cpp
--- a/Engine/Plugins/Experimental/ProceduralVegetationEditor/Source/ProceduralVegetationEditor/Private/Importer/PVImportTests.cpp
+++ b/Engine/Plugins/Experimental/ProceduralVegetationEditor/Source/ProceduralVegetationEditor/Private/Importer/PVImportTests.cpp
@@ -87,7 +87,7 @@ PV_IMPORTER_TEST(FillGrowthDataFromBranchHierarchyDescription)
 	UTEST_EQUAL("CorrectNrOfPointsInGroup", ActualCollection.NumElements(PV::GroupNames::PointGroup), NumBranchPoints);
 	UTEST_EQUAL("CorrectNrOfBranchesInGroup", ActualCollection.NumElements(PV::GroupNames::BranchGroup), NumBranches);
 
-	PV::Import::Helper::FillGrowthDataFromBranchHierarchyDescription(ActualCollection, *BranchHierarchy);
+	PV::Import::Helper::FillGrowthDataFromBranchHierarchyDescription(ActualCollection, { *BranchHierarchy });
 
 	const PV::Facades::FPointFacade ActualPointFacade(ActualCollection);
 	const PV::Facades::FBranchFacade ActualBranchFacade(ActualCollection);
@@ -136,7 +136,7 @@ PV_IMPORTER_TEST(ComputeLengthFromRoot)
 
 	FManagedArrayCollection Collection;
 	PV::Import::Helper::CreateEmptyGrowthData(Collection, NumBranchPoints, NumFoliagePoints, NumBranches, NumFoliageMeshes);
-	PV::Import::Helper::FillGrowthDataFromBranchHierarchyDescription(Collection, *BranchHierarchy);
+	PV::Import::Helper::FillGrowthDataFromBranchHierarchyDescription(Collection, { *BranchHierarchy });
 
 	PV::Facades::FPointFacade PointFacade(Collection);
 	const PV::Facades::FBranchFacade BranchFacade(Collection);
@@ -278,7 +278,7 @@ PV_IMPORTER_TEST(ComputeBudDevelopment)
 
 	FManagedArrayCollection Collection;
 	PV::Import::Helper::CreateEmptyGrowthData(Collection, NumBranchPoints, NumFoliagePoints, NumBranches, NumFoliageMeshes);
-	PV::Import::Helper::FillGrowthDataFromBranchHierarchyDescription(Collection, *BranchHierarchy);
+	PV::Import::Helper::FillGrowthDataFromBranchHierarchyDescription(Collection, { *BranchHierarchy });
 
 	PV::Facades::FPointFacade PointFacade(Collection);
 	const PV::Facades::FBranchFacade BranchFacade(Collection);
@@ -376,7 +376,7 @@ PV_IMPORTER_TEST(ComputeBudDirections)
 
 	FManagedArrayCollection Collection;
 	PV::Import::Helper::CreateEmptyGrowthData(Collection, NumBranchPoints, NumFoliagePoints, NumBranches, NumFoliageMeshes);
-	PV::Import::Helper::FillGrowthDataFromBranchHierarchyDescription(Collection, *BranchHierarchy);
+	PV::Import::Helper::FillGrowthDataFromBranchHierarchyDescription(Collection, { *BranchHierarchy });
 
 	PV::Facades::FPointFacade PointFacade(Collection);
 	const PV::Facades::FBranchFacade BranchFacade(Collection);
@@ -424,7 +424,7 @@ PV_IMPORTER_TEST(SplitFusedPoints)
 
 	FManagedArrayCollection Collection;
 	PV::Import::Helper::CreateEmptyGrowthData(Collection, NumBranchPoints, NumFoliagePoints, NumBranches, NumFoliageMeshes);
-	PV::Import::Helper::FillGrowthDataFromBranchHierarchyDescription(Collection, *BranchHierarchy);
+	PV::Import::Helper::FillGrowthDataFromBranchHierarchyDescription(Collection, { *BranchHierarchy });
 
 	PV::Facades::FPointFacade PointFacade(Collection);
 	const PV::Facades::FBranchFacade BranchFacade(Collection);
diff --git a/Engine/Plugins/Experimental/ProceduralVegetationEditor/Source/ProceduralVegetationEditor/Private/Importer/PVImporter_MockTree.cpp b/Engine/Plugins/Experimental/ProceduralVegetationEditor/Source/ProceduralVegetationEditor/Private/Importer/PVImporter_MockTree.cpp
--- a/Engine/Plugins/Experimental/ProceduralVegetationEditor/Source/ProceduralVegetationEditor/Private/Importer/PVImporter_MockTree.cpp
+++ b/Engine/Plugins/Experimental/ProceduralVegetationEditor/Source/ProceduralVegetationEditor/Private/Importer/PVImporter_MockTree.cpp
@@ -69,7 +69,7 @@ bool PV::Import::MockTree::ImportMockTreeAsProceduralVegetation(const FString& O
 	PlantDescriptionDataAsset = NewObject<UPVPlantDescriptionDataAsset>(PlantDescriptionDataAssetPackage, FPackageName::GetShortFName(PlantDescriptionDataAssetPath), RF_Standalone | RF_Public);
 
 	PlantDescriptionDataAsset->Modify();
-	PlantDescriptionDataAsset->PlantDescription.BranchHierarchy = MoveTemp(*BranchHierarchy);
+	PlantDescriptionDataAsset->PlantDescription.BranchHierarchies.Add(MoveTemp(*BranchHierarchy));
 	PlantDescriptionDataAsset->PlantDescription.Foliage = {};
 	PlantDescriptionDataAsset->PostEditChange();
 
diff --git a/Engine/Plugins/Experimental/ProceduralVegetationEditor/Source/ProceduralVegetationEditor/Private/Importer/PVImporter_StaticMesh.cpp b/Engine/Plugins/Experimental/ProceduralVegetationEditor/Source/ProceduralVegetationEditor/Private/Importer/PVImporter_StaticMesh.cpp
--- a/Engine/Plugins/Experimental/ProceduralVegetationEditor/Source/ProceduralVegetationEditor/Private/Importer/PVImporter_StaticMesh.cpp
+++ b/Engine/Plugins/Experimental/ProceduralVegetationEditor/Source/ProceduralVegetationEditor/Private/Importer/PVImporter_StaticMesh.cpp
@@ -4,168 +4,45 @@
 
 #include "PVPlantDescription.h"
 #include "PVImportHelpers.h"
+#include "PVDynamicMeshVertexAttribute.h"
 
 #include "Engine/StaticMesh.h"
 #include "DynamicMesh/DynamicMesh3.h"
-#include "DynamicMesh/DynamicVertexAttribute.h"
-#include "DynamicMesh/DynamicMeshAttributeSet.h"
 #include "DynamicMesh/DynamicMeshAABBTree3.h"
 #include "DynamicMesh/MeshNormals.h"
 #include "ConversionUtils/SceneComponentToDynamicMesh.h"
-#include "GeometryCollectionToDynamicMesh.h"
-#include "GeometryCollection/GeometryCollection.h"
 #include "Selections/MeshConnectedComponents.h"
 
-template<typename AttribValueType, int AttribDimension>
-class TDynamicMeshVertexAttributeExt 
-	: public UE::Geometry::TDynamicMeshVertexAttribute<AttribValueType, AttribDimension>
-{
-public:
-	using Super = UE::Geometry::TDynamicMeshVertexAttribute<AttribValueType, AttribDimension>;
-
-	TDynamicMeshVertexAttributeExt() : Super() {}
-	TDynamicMeshVertexAttributeExt(UE::Geometry::FDynamicMesh3* ParentIn, bool bAutoInit = true) : Super(ParentIn, bAutoInit) {}
-
-	inline AttribValueType GetValue(int32 VertexID) const
-	{
-		AttribValueType Value = (AttribValueType)0;
-		Super::GetValue(VertexID, &Value);
-		return Value;
-	}
-
-	inline void SetValue(int32 VertexID, const AttribValueType& NewValue)
-	{
-		Super::SetValue(VertexID, &NewValue);
-	}
-
-	template <typename Predicate>
-	TArray<int32> FindAllByPredicate(Predicate Pred) const
-	{
-		TArray<int32> Result;
-
-		if (Super::Parent)
-		{
-			Result.Reserve(Super::Parent->VertexCount());
-
-			for (int32 VertexID : Super::Parent->VertexIndicesItr())
-			{
-				const AttribValueType Value = GetValue(VertexID);
-				if (Pred(Value))
-				{
-					Result.Add(VertexID);
-				}
-			}
-		}
-
-		return Result;
-	}
+using FBooleanVertexAttribute = PV::DynamicMeshHelper::FBooleanVertexAttribute;
+using FIntVertexAttribute = PV::DynamicMeshHelper::FIntVertexAttribute;
+using FFloatVertexAttribute = PV::DynamicMeshHelper::FFloatVertexAttribute;
 
-	TArray<int32> FindAllNonZero() const
-	{
-		return FindAllByPredicate([](const auto& Value) { return Value != (AttribValueType)0; });
-	}
-
-	TArray<int32> FindAllZero() const
-	{
-		return FindAllByPredicate([](const auto& Value) { return Value == (AttribValueType)0; });
-	}
-};
+const static PV::TDynamicMeshVertexAttributeDefinition<bool> FillerVertexAttributeDefinition(TEXT("FillerVertex"));
+const static PV::TDynamicMeshVertexAttributeDefinition<int32> FuseToAttributeDefinition(TEXT("FuseTo"));
 
-template<typename T>
-class TDynamicMeshVertexAttributeDefinition
+static void AddBranchCurvesToCollection(FManagedArrayCollection& InCollection, const TArray<TArray<int32>>& BranchCurves, const TFunction<FVector3f(int32)>& GetVertexPosFunc)
 {
-public:
-	using FAttributeType = TDynamicMeshVertexAttributeExt<T, 1>;
-
-	const FName AttributeName;
-
-	TDynamicMeshVertexAttributeDefinition(FName InAttributeName)
-		: AttributeName(InAttributeName)
-	{
-	}
-
-	const FAttributeType* GetAttribute(const UE::Geometry::FDynamicMesh3& DynamicMesh) const
-	{
-		if (DynamicMesh.HasAttributes() && DynamicMesh.Attributes()->HasAttachedAttribute(AttributeName))
-		{
-			return static_cast<const FAttributeType*>(DynamicMesh.Attributes()->GetAttachedAttribute(AttributeName));
-		}
-		return nullptr;
-	}
-
-	FAttributeType* GetAttribute(UE::Geometry::FDynamicMesh3& DynamicMesh) const
-	{
-		if (DynamicMesh.HasAttributes() && DynamicMesh.Attributes()->HasAttachedAttribute(AttributeName))
-		{
-			return static_cast<FAttributeType*>(DynamicMesh.Attributes()->GetAttachedAttribute(AttributeName));
-		}
-		return nullptr;
-	}
+	const static FName BranchCurveGroupName = TEXT("BranchCurves");
+	const static FName CurveAttributeName = TEXT("Curve");
 
-	const FAttributeType& GetAttributeChecked(const UE::Geometry::FDynamicMesh3& DynamicMesh) const
+	if (!InCollection.HasGroup(BranchCurveGroupName))
 	{
-		const FAttributeType* Attribute = GetAttribute(DynamicMesh);
-		check(Attribute != nullptr);
-		return *Attribute;
+		InCollection.AddGroup(BranchCurveGroupName);
 	}
 
-	FAttributeType& GetAttributeChecked(UE::Geometry::FDynamicMesh3& DynamicMesh) const
-	{
-		FAttributeType* Attribute = GetAttribute(DynamicMesh);
-		check(Attribute != nullptr);
-		return *Attribute;
-	}
+	InCollection.AddElements(BranchCurves.Num(), BranchCurveGroupName);
 
-	void AttachAttribute(UE::Geometry::FDynamicMesh3& DynamicMesh) const
+	TManagedArray<TArray<FVector3f>>& CurveAttribute = InCollection.AddAttribute<TArray<FVector3f>>(CurveAttributeName, BranchCurveGroupName);
+	for (int32 i = 0; i < BranchCurves.Num(); ++i)
 	{
-		if (!DynamicMesh.HasAttributes())
-		{
-			DynamicMesh.EnableAttributes();
-		}
-
-		check(!DynamicMesh.Attributes()->HasAttachedAttribute(AttributeName));
-
-		auto* Attribute = new FAttributeType(&DynamicMesh);
-		Attribute->Initialize();
-		DynamicMesh.Attributes()->AttachAttribute(AttributeName, Attribute);
-	}
-
-	void CopyAttributeToCollection(const UE::Geometry::FDynamicMesh3& DynamicMesh, FManagedArrayCollection& OutCollection) const
-	{
-		const FAttributeType* AttachedAttribute = GetAttribute(DynamicMesh);
-		if (!AttachedAttribute)
-		{
-			return;
-		}
-
-		const static FName DebugVertexGroup = *FString::Printf(TEXT("%s_Debug"), *FGeometryCollection::VerticesGroup.ToString());
-
-		if (!OutCollection.HasGroup(DebugVertexGroup))
-		{
-			OutCollection.AddGroup(DebugVertexGroup);
-			OutCollection.AddElements(DynamicMesh.MaxVertexID(), DebugVertexGroup);
-
-			auto& DebugVertexPositionAttribute = OutCollection.AddAttribute<FVector3f>(FGeometryCollection::VertexPositionAttribute, DebugVertexGroup);
-			for (int32 VertexIndex : DynamicMesh.VertexIndicesItr())
-			{
-				DebugVertexPositionAttribute[VertexIndex] = FVector3f(DynamicMesh.GetVertex(VertexIndex));
-			}
-		}
-
-		TManagedArray<T>& CollectionAttribute = OutCollection.AddAttribute<T>(AttributeName, DebugVertexGroup);
-		for (int32 VertexIndex : DynamicMesh.VertexIndicesItr())
+		const TArray<int32>& BranchCurve = BranchCurves[i];
+		CurveAttribute[i].Reserve(BranchCurve.Num());
+		for (int32 VertexIndex : BranchCurve)
 		{
-			CollectionAttribute[VertexIndex] = AttachedAttribute->GetValue(VertexIndex);
+			CurveAttribute[i].Add(GetVertexPosFunc(VertexIndex));
 		}
 	}
-};
-
-const static TDynamicMeshVertexAttributeDefinition<float> VertexScaleAttributeDefinition(TEXT("VertexScale"));
-const static TDynamicMeshVertexAttributeDefinition<int32> FuseToAttributeDefinition(TEXT("FuseTo"));
-const static TDynamicMeshVertexAttributeDefinition<bool> FillerVertexAttributeDefinition(TEXT("FillerVertex"));
-const static TDynamicMeshVertexAttributeDefinition<bool> EndPointAttributeDefinition(TEXT("EndPoint"));
-const static TDynamicMeshVertexAttributeDefinition<bool> RootPointAttributeDefinition(TEXT("RootPoint"));
-const static TDynamicMeshVertexAttributeDefinition<float> LengthFromRootAttributeDefinition(TEXT("LengthFromRoot"));
+}
 
 static FVector GenerateFillerVertex(const FVector& Pos1, const FVector& Pos2)
 {
@@ -173,73 +50,6 @@ static FVector GenerateFillerVertex(const FVector& Pos1, const FVector& Pos2)
 	return ((Pos1 + Pos2) * 0.5)  + (Pos1 - Pos2).GetSafeNormal().Cross(FVector::ForwardVector).Cross(FVector::UpVector)/* * 0.01*/;
 }
 
-static TArray<int32> GetUniqueVertexIndices(const UE::Geometry::FDynamicMesh3& DynamicMesh, const TArray<int32>& TriangleIds)
-{
-	TSet<int32> UniqueVertexIndices;
-	UniqueVertexIndices.Reserve(TriangleIds.Num() * 3);
-	for (int32 tid : TriangleIds)
-	{
-		const UE::Geometry::FIndex3i VertexIndices = DynamicMesh.GetTriangle(tid);
-		UniqueVertexIndices.Add(VertexIndices[0]);
-		UniqueVertexIndices.Add(VertexIndices[1]);
-		UniqueVertexIndices.Add(VertexIndices[2]);
-	}
-
-	return UniqueVertexIndices.Array();
-}
-
-static TArray<int32> GetVertexNeighbours(const UE::Geometry::FDynamicMesh3& DynamicMesh, int32 VertexID)
-{
-	TArray<int32> OutNeighbours;
-	OutNeighbours.Reserve(DynamicMesh.GetVtxEdgeCount(VertexID));
-	for (int32 EdgeID : DynamicMesh.VtxEdgesItr(VertexID))
-	{
-		const UE::Geometry::FIndex2i VertexIDs = DynamicMesh.GetEdgeV(EdgeID);
-		const int32 OtherVertexID = VertexIDs.A == VertexID ? VertexIDs.B : VertexIDs.A;
-		OutNeighbours.Add(OtherVertexID);
-	}
-	return OutNeighbours;
-}
-
-static TArray<int32> FindVerticesInRadius(const UE::Geometry::FDynamicMeshAABBTree3& DynamicMeshAABBTree, const FVector& Pos, float Radius)
-{
-	TArray<int32> OutVertexIndices;
-
-	const UE::Geometry::FDynamicMesh3* DynamicMesh = DynamicMeshAABBTree.GetMesh();
-	if (DynamicMesh == nullptr)
-	{
-		return OutVertexIndices;
-	}
-
-	OutVertexIndices.Reserve(DynamicMesh->VertexCount());
-
-	const double RadiusSqr = Radius * Radius;
-
-	UE::Geometry::FDynamicMeshAABBTree3::FTreeTraversal TreeTraversal;
-
-	TreeTraversal.NextBoxF = [&](const UE::Geometry::FAxisAlignedBox3d& Box, int32 Depth)->bool
-	{
-		return Box.DistanceSquared(Pos) <= RadiusSqr;
-	};
-
-	TreeTraversal.NextTriangleF = [&](int32 TID)
-	{
-		const UE::Geometry::FIndex3i Triangle = DynamicMesh->GetTriangle(TID);
-		for (int32 j = 0; j < 3; ++j)
-		{
-			const double VertexDistSqr = (DynamicMesh->GetVertex(Triangle[j]) - Pos).SizeSquared();
-			if (VertexDistSqr <= RadiusSqr)
-			{
-				OutVertexIndices.AddUnique(Triangle[j]);
-			}
-		}
-	};
-
-	DynamicMeshAABBTree.DoTraversal(TreeTraversal);
-
-	return OutVertexIndices;
-}
-
 static void RemoveUnusedVertices(UE::Geometry::FDynamicMesh3& DynamicMesh)
 {
 	TRACE_CPUPROFILER_EVENT_SCOPE(OrphanRemoval);
@@ -318,8 +128,12 @@ static void ComputeNormals(UE::Geometry::FDynamicMesh3& DynamicMesh)
 	}
 }
 
-static void ComputeVerticesCenterAndRadius_Internal_V1(UE::Geometry::FDynamicMesh3& DynamicMesh, const TArray<int32>& TriangleIds)
-{
+static void ComputeVerticesCenterAndRadius_Internal_V1(
+	UE::Geometry::FDynamicMesh3& DynamicMesh, 
+	const TArray<int32>& TriangleIds,
+	FFloatVertexAttribute& VertexScaleAttribute
+)
+{	
 	// This function tries to find the center point of the branch/trunk for each vertex, it then moves the vertex to said center point and stores the
 	// branch radius in the VertexScale attribute
 	//
@@ -327,7 +141,7 @@ static void ComputeVerticesCenterAndRadius_Internal_V1(UE::Geometry::FDynamicMes
 	// center point to lay between the current vertex, and the found vertex. It then performs another overlap, and repeats the process until it no longer overlaps any vertices.
 	// Once an opposing vertex has been found, an extra check is performed to make sure the normal of the closest triangle is facing away from the center point.
 	
-	const TArray<int32> ComponentVertexIndices = GetUniqueVertexIndices(DynamicMesh, TriangleIds);
+	const TArray<int32> ComponentVertexIndices = PV::DynamicMeshHelper::GetUniqueVertexIndices(DynamicMesh, TriangleIds);
 
 	UE::Geometry::FDynamicMeshAABBTree3 DynamicMeshAABBTree(&DynamicMesh, false);
 	DynamicMeshAABBTree.Build(TriangleIds);
@@ -406,8 +220,6 @@ static void ComputeVerticesCenterAndRadius_Internal_V1(UE::Geometry::FDynamicMes
 		NewVertexPositions[i] = SphereCenter;
 	}
 
-	auto& VertexScaleAttribute = VertexScaleAttributeDefinition.GetAttributeChecked(DynamicMesh);
-
 	for (int32 i = 0; i < ComponentVertexIndices.Num(); ++i)
 	{
 		const int32 VertexIndex = ComponentVertexIndices[i];
@@ -416,7 +228,11 @@ static void ComputeVerticesCenterAndRadius_Internal_V1(UE::Geometry::FDynamicMes
 	}
 }
 
-static void ComputeVerticesCenterAndRadius_Internal_V2(UE::Geometry::FDynamicMesh3& DynamicMesh, const TArray<int32>& TriangleIds)
+static void ComputeVerticesCenterAndRadius_Internal_V2(
+	UE::Geometry::FDynamicMesh3& DynamicMesh, 
+	const TArray<int32>& TriangleIds,
+	FFloatVertexAttribute& VertexScaleAttribute
+)
 {
 	// This is a second attempt at creating an algorithm that can walk along each vertex on a tree mesh and find the corresponding center point and
 	// radius of the branch at the point of the vertex.
@@ -437,7 +253,7 @@ static void ComputeVerticesCenterAndRadius_Internal_V2(UE::Geometry::FDynamicMes
 	// 3. It also not work well if the mesh is built up using only faces (e.g flat geometry instead of a tube for the branches)
 	// 4. The normal of the vertex may not point directly towards the opposite side (see "Remove extreme radius values" section below for more information)
 
-	const TArray<int32> ComponentVertexIndices = GetUniqueVertexIndices(DynamicMesh, TriangleIds);
+	const TArray<int32> ComponentVertexIndices = PV::DynamicMeshHelper::GetUniqueVertexIndices(DynamicMesh, TriangleIds);
 
 	UE::Geometry::FDynamicMeshAABBTree3 DynamicMeshAABBTree(&DynamicMesh, false);
 	DynamicMeshAABBTree.Build(TriangleIds);
@@ -564,8 +380,6 @@ static void ComputeVerticesCenterAndRadius_Internal_V2(UE::Geometry::FDynamicMes
 		}
 	}
 
-	auto& VertexScaleAttribute = VertexScaleAttributeDefinition.GetAttributeChecked(DynamicMesh);
-
 	for (int32 i = 0; i < ComponentVertexIndices.Num(); ++i)
 	{
 		if (Scales[i] > 0)
@@ -584,22 +398,26 @@ static TAutoConsoleVariable<int32> ComputeVerticesCenterVersion(
 	ECVF_Default
 );
 
-static void ComputeVerticesCenterAndRadius_Internal(UE::Geometry::FDynamicMesh3& DynamicMesh, const TArray<int32>& TriangleIds)
+static void ComputeVerticesCenterAndRadius_Internal(
+	UE::Geometry::FDynamicMesh3& DynamicMesh, 
+	const TArray<int32>& TriangleIds,
+	FFloatVertexAttribute& VertexScaleAttribute
+)
 {
 	TRACE_CPUPROFILER_EVENT_SCOPE(ComputeVerticesCenterAndRadius_Internal);
 
 	const int32 UseVersion = ComputeVerticesCenterVersion.GetValueOnAnyThread();
 	if (UseVersion == 1)
 	{
-		ComputeVerticesCenterAndRadius_Internal_V1(DynamicMesh, TriangleIds);
+		ComputeVerticesCenterAndRadius_Internal_V1(DynamicMesh, TriangleIds, VertexScaleAttribute);
 	}
 	else
 	{
-		ComputeVerticesCenterAndRadius_Internal_V2(DynamicMesh, TriangleIds);
+		ComputeVerticesCenterAndRadius_Internal_V2(DynamicMesh, TriangleIds, VertexScaleAttribute);
 	}
 }
 
-static void ComputeVerticesCenterAndRadius(UE::Geometry::FDynamicMesh3& DynamicMesh)
+static void ComputeVerticesCenterAndRadius(UE::Geometry::FDynamicMesh3& DynamicMesh, FFloatVertexAttribute& VertexScaleAttribute)
 {
 	TRACE_CPUPROFILER_EVENT_SCOPE(ComputeVerticesCenterAndRadius);
 
@@ -610,21 +428,24 @@ static void ComputeVerticesCenterAndRadius(UE::Geometry::FDynamicMesh3& DynamicM
 
 	ParallelFor(ConnectedComponents.Num(), [&](int32 Index)
 	{
-		ComputeVerticesCenterAndRadius_Internal(DynamicMesh, ConnectedComponents.GetComponent(Index).Indices);
+		ComputeVerticesCenterAndRadius_Internal(DynamicMesh, ConnectedComponents.GetComponent(Index).Indices, VertexScaleAttribute);
 	});
 }
 
-static void ExtractReductionCloud_Internal(UE::Geometry::FDynamicMesh3& DynamicMesh, const TArray<int32>& TriangleIds, float SearchRadiusInflate)
+static void ExtractReductionCloud_Internal(
+	const UE::Geometry::FDynamicMesh3& DynamicMesh, 
+	const TArray<int32>& TriangleIds, 
+	float SearchRadiusInflate,
+	const FFloatVertexAttribute& VertexScaleAttribute,
+	FIntVertexAttribute& FuseToAttribute
+)
 {
 	TRACE_CPUPROFILER_EVENT_SCOPE(ExtractReductionCloud_Internal);
-
-	const auto& VertexScaleAttribute = VertexScaleAttributeDefinition.GetAttributeChecked(DynamicMesh);
-	auto& FuseToAttribute = FuseToAttributeDefinition.GetAttributeChecked(DynamicMesh);
 		
 	UE::Geometry::FDynamicMeshAABBTree3 DynamicMeshAABBTree(&DynamicMesh, false);
 	DynamicMeshAABBTree.Build(TriangleIds);
 
-	TArray<int32> VertexIndices = GetUniqueVertexIndices(DynamicMesh, TriangleIds);
+	TArray<int32> VertexIndices = PV::DynamicMeshHelper::GetUniqueVertexIndices(DynamicMesh, TriangleIds);
 
 	const int32 MaxVertexID = DynamicMesh.MaxVertexID();
 		
@@ -668,7 +489,7 @@ static void ExtractReductionCloud_Internal(UE::Geometry::FDynamicMesh3& DynamicM
 		}
 
 		const float SearchRadius = VertexRadius * SearchRadiusInflate;
-		const TArray<int32> OverlappingVertices = FindVerticesInRadius(DynamicMeshAABBTree, VertexPosition, SearchRadius);
+		const TArray<int32> OverlappingVertices = PV::DynamicMeshHelper::FindVerticesInRadius(DynamicMeshAABBTree, VertexPosition, SearchRadius);
 		for (int32 j = 0; j < OverlappingVertices.Num(); j++)
 		{
 			const int32 OverlappingVertexIndex = OverlappingVertices[j];
@@ -702,13 +523,14 @@ static void ExtractReductionCloud_Internal(UE::Geometry::FDynamicMesh3& DynamicM
 	}
 }
 	
-static void FuseVerticesUsingFuseAttribute(UE::Geometry::FDynamicMesh3& DynamicMesh)
+static void FuseVerticesUsingFuseAttribute(
+	UE::Geometry::FDynamicMesh3& DynamicMesh,
+	const FIntVertexAttribute& FuseToAttribute,
+	FBooleanVertexAttribute& FillerVertexAttribute
+)
 {
 	TRACE_CPUPROFILER_EVENT_SCOPE(FuseVerticesUsingFuseAttribute);
 
-	auto& FuseToAttribute = FuseToAttributeDefinition.GetAttributeChecked(DynamicMesh);
-	auto& FillerVertexAttribute = FillerVertexAttributeDefinition.GetAttributeChecked(DynamicMesh);
-
 	TArray<int32> TriangleIndices;
 	TriangleIndices.Reserve(DynamicMesh.TriangleCount());
 	for (int32 TID : DynamicMesh.TriangleIndicesItr())
@@ -756,7 +578,12 @@ static void FuseVerticesUsingFuseAttribute(UE::Geometry::FDynamicMesh3& DynamicM
 	RemoveUnusedVertices(DynamicMesh);
 }
 
-static void CollapseBranchesToLine(UE::Geometry::FDynamicMesh3& DynamicMesh)
+static void CollapseBranchesToLine(
+	UE::Geometry::FDynamicMesh3& DynamicMesh,
+	const FFloatVertexAttribute& VertexScaleAttribute,
+	FIntVertexAttribute& FuseToAttribute,
+	FBooleanVertexAttribute& FillerVertexAttribute
+)
 {
 	TRACE_CPUPROFILER_EVENT_SCOPE(CollapseBranchesToLine);
 
@@ -765,19 +592,20 @@ static void CollapseBranchesToLine(UE::Geometry::FDynamicMesh3& DynamicMesh)
 
 	ParallelFor(ConnectedComponents.Num(), [&](int32 Index)
 	{
-		ExtractReductionCloud_Internal(DynamicMesh, ConnectedComponents.GetComponent(Index).Indices, 2.f);
+		ExtractReductionCloud_Internal(DynamicMesh, ConnectedComponents.GetComponent(Index).Indices, 2.f, VertexScaleAttribute, FuseToAttribute);
 	});
 
-	FuseVerticesUsingFuseAttribute(DynamicMesh);
+	FuseVerticesUsingFuseAttribute(DynamicMesh, FuseToAttribute, FillerVertexAttribute);
 }
 
-static void FindEndPoints(UE::Geometry::FDynamicMesh3& DynamicMesh)
+static void FindEndPoints(
+	const UE::Geometry::FDynamicMesh3& DynamicMesh,
+	FBooleanVertexAttribute& EndPointAttribute,
+	const FBooleanVertexAttribute& FillerVertexAttribute
+)
 {
 	TRACE_CPUPROFILER_EVENT_SCOPE(FindEndPoints);
 	
-	auto& EndPointAttribute = EndPointAttributeDefinition.GetAttributeChecked(DynamicMesh);
-	auto& FillerVertexAttribute = FillerVertexAttributeDefinition.GetAttributeChecked(DynamicMesh);
-
 	TArray<double> xComps;
 	TArray<double> yComps;
 	TArray<double> zComps;
@@ -833,7 +661,12 @@ static void FindEndPoints(UE::Geometry::FDynamicMesh3& DynamicMesh)
 	}
 }
 
-static void ConnectBranchesToParent(UE::Geometry::FDynamicMesh3& DynamicMesh)
+static void ConnectBranchesToParent(
+	UE::Geometry::FDynamicMesh3& DynamicMesh,
+	const FFloatVertexAttribute& VertexScaleAttribute,
+	FBooleanVertexAttribute& EndPointAttribute,
+	FBooleanVertexAttribute& FillerVertexAttribute
+)
 {
 	TRACE_CPUPROFILER_EVENT_SCOPE(ConnectBranchesToParent);
 
@@ -849,10 +682,6 @@ static void ConnectBranchesToParent(UE::Geometry::FDynamicMesh3& DynamicMesh)
 	// connected if the closest vertex is too far away (even if the end-point is very close to the parent mesh itself). This would lead to the entire
 	// branch being excluded from the final tree.
 
-	const auto& VertexScaleAttribute = VertexScaleAttributeDefinition.GetAttributeChecked(DynamicMesh);
-	auto& EndPointAttribute = EndPointAttributeDefinition.GetAttributeChecked(DynamicMesh);
-	auto& FillerVertexAttribute = FillerVertexAttributeDefinition.GetAttributeChecked(DynamicMesh);
-
 	const UE::Geometry::FDynamicMeshAABBTree3 DynamicMeshAABBTree(&DynamicMesh);
 	UE::Geometry::FMeshConnectedComponents ConnectedComponents(&DynamicMesh);
 	ConnectedComponents.FindConnectedVertices();
@@ -878,7 +707,7 @@ static void ConnectBranchesToParent(UE::Geometry::FDynamicMesh3& DynamicMesh)
 		const FVector VertexPosition = DynamicMesh.GetVertex(EndPointVertexIndex);
 		const float VertexRadius = VertexScaleAttribute.GetValue(EndPointVertexIndex);
 		const float SearchRadius = VertexRadius * 4.f;
-		const TArray<int32> OverlappingVertices = FindVerticesInRadius(DynamicMeshAABBTree, VertexPosition, SearchRadius);
+		const TArray<int32> OverlappingVertices = PV::DynamicMeshHelper::FindVerticesInRadius(DynamicMeshAABBTree, VertexPosition, SearchRadius);
 
 		float ClosestVertexDistSqr = MAX_flt;
 		int32 VertexToMergeWith = INDEX_NONE;
@@ -937,12 +766,12 @@ static void ConnectBranchesToParent(UE::Geometry::FDynamicMesh3& DynamicMesh)
 
 		DynamicMesh.SetVertex(EndPointVertexIndex, NewVertexPosition);
 		
-		const TArray<int32> Neighbours = GetVertexNeighbours(DynamicMesh, EndPointVertexIndex);
+		const TArray<int32> Neighbours = PV::DynamicMeshHelper::GetVertexNeighbours(DynamicMesh, EndPointVertexIndex);
 		for (int32 Neighbour : Neighbours)
 		{
 			if (FillerVertexAttribute.GetValue(Neighbour))
 			{
-				const TArray<int32> FillerNeighbours = GetVertexNeighbours(DynamicMesh, Neighbour);
+				const TArray<int32> FillerNeighbours = PV::DynamicMeshHelper::GetVertexNeighbours(DynamicMesh, Neighbour);
 				const int32* AdjacentVertex = FillerNeighbours.FindByPredicate([&](const auto& Value) { return Value != EndPointVertexIndex && Value != Neighbour; });
 				check(AdjacentVertex != nullptr);
 				const FVector AdjacentVertexPosition = DynamicMesh.GetVertex(*AdjacentVertex);
@@ -959,15 +788,16 @@ static void ConnectBranchesToParent(UE::Geometry::FDynamicMesh3& DynamicMesh)
 	}
 }
 
-static void FindRootPoints(UE::Geometry::FDynamicMesh3& DynamicMesh)
+static void FindRootPoints(
+	const UE::Geometry::FDynamicMesh3& DynamicMesh,
+	FBooleanVertexAttribute& RootPointAttribute,
+	FBooleanVertexAttribute& EndPointAttribute
+)
 {
 	TRACE_CPUPROFILER_EVENT_SCOPE(FindRootPoints);
 
 	// This function currently just finds all the end points that are close to the bottom of the plant bounding box and converts them to root points.
 	// This could result in issues if a plant has multiple points touching the floor. 
-	
-	auto& RootPointAttribute = RootPointAttributeDefinition.GetAttributeChecked(DynamicMesh);
-	auto& EndPointAttribute = EndPointAttributeDefinition.GetAttributeChecked(DynamicMesh);
 
 	const UE::Geometry::FAxisAlignedBox3d BoxBounds = DynamicMesh.GetBounds();
 
@@ -985,323 +815,23 @@ static void FindRootPoints(UE::Geometry::FDynamicMesh3& DynamicMesh)
 	}
 }
 
-static void ComputeLengthFromRoot(UE::Geometry::FDynamicMesh3& DynamicMesh)
+PV::StaticMeshImport::EImportResult PV::StaticMeshImport::ImportGrowthDataFromStaticMesh(const FPVImportStaticMeshParams& InParams, FManagedArrayCollection& OutCollection)
 {
-	TRACE_CPUPROFILER_EVENT_SCOPE(ComputeLengthFromRoot);
-
-	// This algorithm starts at the root point and walks along every vertex, keeping track of its distance to the root point and stores that distance
-	// in the LengthFromRootAttribute. This distance is used to find the closest path between the root point and the tips in DrawBranchCurves.
-
-	auto& RootPointsAttribute = RootPointAttributeDefinition.GetAttributeChecked(DynamicMesh);
-	auto& LengthFromRootAttribute = LengthFromRootAttributeDefinition.GetAttributeChecked(DynamicMesh);
-	auto& FillerVertexAttribute = FillerVertexAttributeDefinition.GetAttributeChecked(DynamicMesh);
-
-	for (int32 VertexID : DynamicMesh.VertexIndicesItr())
-	{
-		LengthFromRootAttribute.SetValue(VertexID, -1);
-	}
+	using namespace PV::DynamicMeshHelper;
 	
-	TArray<bool> bHasComputedVertex;
-	bHasComputedVertex.SetNumZeroed(DynamicMesh.MaxVertexID());
-	const TFunction<void(int32, const FVector&, float)> ComputeLengthFromRootRecursive = [&](int32 VertexIndex, const FVector& PrevVertexLocation, float PrevVertexLengthFromRoot)
-	{
-		if (FillerVertexAttribute.GetValue(VertexIndex))
-		{
-			return;
-		}
-			
-		if (bHasComputedVertex[VertexIndex])
-		{
-			return;
-		}
-
-		bHasComputedVertex[VertexIndex] = true;
-			
-		const FVector VertexPosition = DynamicMesh.GetVertex(VertexIndex);
-		const float VertexLengthFromRoot = PrevVertexLengthFromRoot + (VertexPosition - PrevVertexLocation).Size();
-		LengthFromRootAttribute.SetValue(VertexIndex, VertexLengthFromRoot);
-		
-		const TArray<int32> Neighbours = GetVertexNeighbours(DynamicMesh, VertexIndex);
-		for (int32 Neighbour : Neighbours)
-		{
-			ComputeLengthFromRootRecursive(Neighbour, VertexPosition, VertexLengthFromRoot);
-		}
-	};
-
-	for (int32 VertexIndex : RootPointsAttribute.FindAllNonZero())
-	{
-		const FVector VertexPosition = DynamicMesh.GetVertex(VertexIndex);
-		ComputeLengthFromRootRecursive(VertexIndex, VertexPosition, 0.0);
-	}
-}
-
-static void DrawBranchCurves(const UE::Geometry::FDynamicMesh3& DynamicMesh, TArray<TArray<int32>>& OutBranchCurves)
-{
-	TRACE_CPUPROFILER_EVENT_SCOPE(DrawBranchCurves);
-
-	// This algorithm finds the shortest path from all the end points (tips) to the root point and stores that in OutBranchCurves
-
-	const auto& RootPointAttribute = RootPointAttributeDefinition.GetAttributeChecked(DynamicMesh);
-	const auto& EndPointAttribute = EndPointAttributeDefinition.GetAttributeChecked(DynamicMesh);
-	const auto& LengthFromRootAttribute = LengthFromRootAttributeDefinition.GetAttributeChecked(DynamicMesh);
-	const auto& FillerVertexAttribute = FillerVertexAttributeDefinition.GetAttributeChecked(DynamicMesh);
-
-	const auto FindNextPoint = [&](int32 PointIndex, float MinLengthFromRoot)->int32
-	{
-		int32 OutNextPointIndex = INDEX_NONE;
-		const TArray<int32> Neighbours = GetVertexNeighbours(DynamicMesh, PointIndex);
-		for (int32 Neighbour : Neighbours)
-		{
-			if (FillerVertexAttribute.GetValue(Neighbour))
-			{
-				continue;
-			}
-
-			const float LengthFromRoot = LengthFromRootAttribute.GetValue(Neighbour);
-			check(LengthFromRoot >= 0);
-			
-			if (LengthFromRoot >= MinLengthFromRoot)
-			{ 
-				continue; 
-			}
-			MinLengthFromRoot = LengthFromRoot;
-			OutNextPointIndex = Neighbour;
-		}
-		return OutNextPointIndex;
-	};
-
-	const TArray<int32> EndPoints = EndPointAttribute.FindAllNonZero();
-
-	TArray<bool> bHasProcessedVertex;
-	for (int32 i = 0; i < EndPoints.Num(); i++)
-	{
-		int32 CurrentPointIndex = EndPoints[i];
-			
-		float MinLengthFromRoot = LengthFromRootAttribute.GetValue(CurrentPointIndex);
-		if (MinLengthFromRoot < 0)
-		{
-			continue; // This end point was unable to trace down to a root point, perhaps log a warning?
-		}
-			
-		TArray<int32> CurvePoints = { CurrentPointIndex };
-
-		bHasProcessedVertex.Reset();
-		bHasProcessedVertex.SetNumZeroed(DynamicMesh.MaxVertexID());
-		
-		while (true)
-		{
-			const int32 NextPointIndex = FindNextPoint(CurrentPointIndex, MinLengthFromRoot);
-			if (NextPointIndex == INDEX_NONE)
-			{ 
-				break; 
-			}
-
-			if (!ensureMsgf(!bHasProcessedVertex[NextPointIndex], TEXT("Infinite loop detected")))
-			{
-				CurvePoints.Empty();
-				break;
-			}
-			
-			bHasProcessedVertex[NextPointIndex] = true;
-
-			CurvePoints.Add(NextPointIndex);
-
-			if (RootPointAttribute.GetValue(NextPointIndex))
-			{
-				break; 
-			}
-
-			CurrentPointIndex = NextPointIndex;
-			MinLengthFromRoot = LengthFromRootAttribute.GetValue(CurrentPointIndex);
-		}
-
-		if (CurvePoints.Num() > 1)
-		{
-			OutBranchCurves.Add(MoveTemp(CurvePoints));
-		}
-	}
-}
-
-static void RemoveBranchCurveOverlaps(const UE::Geometry::FDynamicMesh3& DynamicMesh, TArray<TArray<int32>>& InOutBranchCurves)
-{
-	TRACE_CPUPROFILER_EVENT_SCOPE(RemoveBranchCurveOverlaps);
-
-	// Since DrawBranchCurves finds the full path from all end points to the root point, there will be overlaps in the branch curves
-	//
-	// In the below example, the curves between a->root and b->root would overlap from c->root. This algorithm finds those overlapping sections
-	// and only keeps that section on the longest curve (a->root in this case).
-	//
-	// a     
-	//  \   b
-	//   \ /
-	//    |c
-	//    |
-	//    |
-	//   root
-	//
-	// Issues with this algorithm: Using the length of the branch to determine the branch parent only works well on small trees. As the tree grows larger
-	// the likelihood of a smaller branches have a longer distance to root gets bigger. A better value for selecting the parent would probably be the
-	// branch radius.
-
-	TArray<float> BranchLengths;
-	BranchLengths.SetNum(InOutBranchCurves.Num());
-
-	int32 MaxNumVertices = 0;
-	for (int32 i = 0; i < InOutBranchCurves.Num(); ++i)
-	{
-		const TArray<int32>& BranchCurve = InOutBranchCurves[i];
-		MaxNumVertices += BranchCurve.Num();
-		for (int32 j = 0; j < BranchCurve.Num() - 1; ++j)
-		{
-			const FVector& Vertex0 = DynamicMesh.GetVertex(BranchCurve[j]);
-			const FVector& Vertex1 = DynamicMesh.GetVertex(BranchCurve[j + 1]);
-			BranchLengths[i] += (Vertex0 - Vertex1).Size();
-		}
-	}
-
-	struct FVertexLongestBranch
-	{
-		int32 BranchIndex = INDEX_NONE;
-		float BranchLength = 0;
-	};
-
-	TMap<int32, FVertexLongestBranch> VerticesLongestBranch;
-	VerticesLongestBranch.Reserve(MaxNumVertices);
-
-	for (int32 i = 0; i < InOutBranchCurves.Num(); ++i)
-	{
-		const TArray<int32>& BranchCurve = InOutBranchCurves[i];
-		const float BranchLength = BranchLengths[i];
-			
-		for (int32 j = 0; j < BranchCurve.Num(); ++j)
-		{
-			const int32 VertexIndex = BranchCurve[j];
-			FVertexLongestBranch& VertexLongestBranch = VerticesLongestBranch.FindOrAdd(VertexIndex);
-
-			if (BranchLength > VertexLongestBranch.BranchLength)
-			{
-				VertexLongestBranch.BranchLength = BranchLength;
-				VertexLongestBranch.BranchIndex = i;
-			}
-		}
-	}
-
-	for (int32 i = InOutBranchCurves.Num() - 1; i >= 0; --i)
-	{
-		TArray<int32>& BranchCurve = InOutBranchCurves[i];
-			
-		for (int32 j = BranchCurve.Num() - 1; j >= 0; --j)
-		{
-			const int32 ThisVertexIndex = BranchCurve[j];
-			const int32 ThisLongestBranchIndex = VerticesLongestBranch[ThisVertexIndex].BranchIndex;
-				
-			if (ThisLongestBranchIndex != i)
-			{
-				if (j == 0)
-				{
-					BranchCurve.RemoveAt(j);
-				}
-				else
-				{
-					const int32 NextVertexIndex = BranchCurve[j - 1];
-					const int32 NextLongestBranchIndex = VerticesLongestBranch[NextVertexIndex].BranchIndex;
-					if (NextLongestBranchIndex != i)
-					{
-						BranchCurve.RemoveAt(j);
-					}
-				}
-			}
-		}
-
-		if (BranchCurve.Num() < 2)
-		{
-			InOutBranchCurves.RemoveAtSwap(i);
-		}
-	}
-}
-
-static void ExtractBranchesFromBranchCurves(const UE::Geometry::FDynamicMesh3& DynamicMesh,
-	const TArray<TArray<int32>>& InBranchCurves, 
-	TArray<FVector3f>& OutPoints, 
-	TArray<float>& OutPointsRadii, 
-	TArray<TArray<int32>>& OutBranchPointIndices
-)
-{
-	TRACE_CPUPROFILER_EVENT_SCOPE(ExtractBranchesFromBranchCurves);
-
-	const auto& VertexScaleAttribute = VertexScaleAttributeDefinition.GetAttributeChecked(DynamicMesh);
-	const auto& EndPointAttribute = EndPointAttributeDefinition.GetAttributeChecked(DynamicMesh);
-
-	int32 MaxNrOfPoints = 0;
-	for (const TArray<int32>& BranchCurve : InBranchCurves)
-	{
-		MaxNrOfPoints += BranchCurve.Num();
-	}
-
-	OutPoints.Reserve(MaxNrOfPoints);
-	OutPointsRadii.Reserve(MaxNrOfPoints);
-	OutBranchPointIndices.Reserve(InBranchCurves.Num());
-
-	TMap<int32, int32> VertexIndexToPointIndex;
-	VertexIndexToPointIndex.Reserve(MaxNrOfPoints);
-
-	TMap<int32, TArray<int32, TInlineAllocator<2>>> PointIndexToBranchIndices;
-	PointIndexToBranchIndices.Reserve(MaxNrOfPoints);
-
-	for (const TArray<int32>& BranchCurve : InBranchCurves)
-	{
-		TArray<int32>& Branch = OutBranchPointIndices.AddDefaulted_GetRef();
-		Branch.Reserve(BranchCurve.Num());
-
-		for (int32 i = BranchCurve.Num() - 1; i >= 0; --i) // Reverse since curves go from tip to base while branches to from base to tip
-		{
-			const int32 VertexIndex = BranchCurve[i];
-			int32 PointIndex = INDEX_NONE;
-
-			int32* ExistingPointIndex = VertexIndexToPointIndex.Find(VertexIndex);
-			if (!ExistingPointIndex)
-			{
-				PointIndex = OutPoints.Add(FVector3f(DynamicMesh.GetVertex(VertexIndex)));
-					
-				const bool bIsEndPoint = EndPointAttribute.GetValue(VertexIndex);
-				if (bIsEndPoint)
-				{
-					OutPointsRadii.Add(0.01f);
-				}
-				else
-				{
-					const float VertexRadius = VertexScaleAttribute.GetValue(VertexIndex);
-					OutPointsRadii.Add(VertexRadius);
-				}
-
-				VertexIndexToPointIndex.Add(VertexIndex, PointIndex);
-			}
-			else
-			{
-				PointIndex = *ExistingPointIndex;
-			}
-
-			Branch.Add(PointIndex);
-			PointIndexToBranchIndices.FindOrAdd(VertexIndex).Add(OutBranchPointIndices.Num() - 1);
-		}
-	}
-}
-
-PV::StaticMeshImport::EImportResult PV::StaticMeshImport::ImportGrowthDataFromStaticMesh(const TObjectPtr<UStaticMesh>& InSourceStaticMesh, FManagedArrayCollection& OutCollection, const FImportOptions& InOptions)
-{
 	TRACE_CPUPROFILER_EVENT_SCOPE(PV::StaticMeshImport::ImportGrowthDataFromStaticMesh);
 	
-	if (!InSourceStaticMesh)
+	if (!InParams.StaticMeshAsset)
 	{
 		return EImportResult::InvalidSourceMesh;
 	}
 
-	if (InSourceStaticMesh->IsCompiling())
+	if (InParams.StaticMeshAsset->IsCompiling())
 	{
 		return EImportResult::MeshNotReady;
 	}
 
-	if (InOptions.MaterialsToKeep.Num() == 0)
+	if (InParams.MaterialsToKeep.Num() == 0)
 	{
 		return EImportResult::InvalidMaterialFilter;
 	}
@@ -1312,167 +842,128 @@ PV::StaticMeshImport::EImportResult PV::StaticMeshImport::ImportGrowthDataFromSt
 
 		FText OutErrorMessage;
 		UE::Conversion::FStaticMeshConversionOptions Options;
-		if (!UE::Conversion::StaticMeshToDynamicMesh(InSourceStaticMesh, DynamicMesh, OutErrorMessage, Options))
+		if (!UE::Conversion::StaticMeshToDynamicMesh(InParams.StaticMeshAsset, DynamicMesh, OutErrorMessage, Options))
 		{
 			return EImportResult::InvalidMeshData;
 		}
 	}
 
-	const auto ReturnWithDebugMesh = [&OutCollection](const UE::Geometry::FDynamicMesh3& DynamicMesh)->EImportResult
+	if (InParams.DebugState == EPVImportStaticMeshDebugState::InitialMesh)
 	{
-		FGeometryCollection MeshGeomCollection;
-		UE::Geometry::FGeometryCollectionToDynamicMeshes::FToCollectionOptions Options2;
-		UE::Geometry::FGeometryCollectionToDynamicMeshes::AppendMeshToCollection(MeshGeomCollection, DynamicMesh, FTransform::Identity, Options2);
-
-		FManagedArrayCollection MeshArrayCollection;
-		MeshGeomCollection.CopyTo(&OutCollection);
-
+		AddDynamicMeshToCollection(DynamicMesh, OutCollection);
 		return EImportResult::Success;
-	};
-
-	if (InOptions.DebugState == FImportOptions::EDebugState::InitialMesh)
-	{
-		return ReturnWithDebugMesh(DynamicMesh);
 	}
 	
-	if (!FilterTrianglesUsingMaterials(DynamicMesh, InSourceStaticMesh, InOptions.MaterialsToKeep))
+	if (!FilterTrianglesUsingMaterials(DynamicMesh, InParams.StaticMeshAsset, InParams.MaterialsToKeep))
 	{
 		return EImportResult::InvalidMaterialFilter;
 	}
 	
-	if (InOptions.DebugState == FImportOptions::EDebugState::KeepMaterial)
+	if (InParams.DebugState == EPVImportStaticMeshDebugState::KeepMaterial)
 	{
-		return ReturnWithDebugMesh(DynamicMesh);
+		AddDynamicMeshToCollection(DynamicMesh, OutCollection);
+		return EImportResult::Success;
 	}
 
 	ComputeNormals(DynamicMesh);
-	if (InOptions.DebugState == FImportOptions::EDebugState::ComputeNormals)
+	if (InParams.DebugState == EPVImportStaticMeshDebugState::ComputeNormals)
 	{
-		return ReturnWithDebugMesh(DynamicMesh);
+		AddDynamicMeshToCollection(DynamicMesh, OutCollection);
+		return EImportResult::Success;
 	}
 
-	const UE::Geometry::FDynamicMesh3 ReferenceMesh = DynamicMesh;
-	VertexScaleAttributeDefinition.AttachAttribute(DynamicMesh);
-
-	ComputeVerticesCenterAndRadius(DynamicMesh);
-	if (InOptions.DebugState == FImportOptions::EDebugState::ComputeVerticesCenterAndRadius)
+	ComputeVerticesCenterAndRadius(
+		DynamicMesh,
+		VertexScaleAttributeDefinition->GetOrAttachAttribute(DynamicMesh)
+	);
+	if (InParams.DebugState == EPVImportStaticMeshDebugState::ComputeVerticesCenterAndRadius)
 	{
-		VertexScaleAttributeDefinition.CopyAttributeToCollection(DynamicMesh, OutCollection);
-		return ReturnWithDebugMesh(DynamicMesh);
+		AddDynamicMeshToCollection(DynamicMesh, OutCollection);
+		VertexScaleAttributeDefinition->CopyAttributeToCollection(DynamicMesh, OutCollection);
+		return EImportResult::Success;
 	}
 
-	FuseToAttributeDefinition.AttachAttribute(DynamicMesh);
-	FillerVertexAttributeDefinition.AttachAttribute(DynamicMesh);
-
-	CollapseBranchesToLine(DynamicMesh);
-	if (InOptions.DebugState == FImportOptions::EDebugState::CollapseBranchesToLine)
+	CollapseBranchesToLine(
+		DynamicMesh,
+		VertexScaleAttributeDefinition->GetOrAttachAttribute(DynamicMesh),
+		FuseToAttributeDefinition.GetOrAttachAttribute(DynamicMesh),
+		FillerVertexAttributeDefinition.GetOrAttachAttribute(DynamicMesh)
+	);
+	if (InParams.DebugState == EPVImportStaticMeshDebugState::CollapseBranchesToLine)
 	{
+		AddDynamicMeshToCollection(DynamicMesh, OutCollection);
 		FuseToAttributeDefinition.CopyAttributeToCollection(DynamicMesh, OutCollection);
-		VertexScaleAttributeDefinition.CopyAttributeToCollection(DynamicMesh, OutCollection);
-		return ReturnWithDebugMesh(DynamicMesh);
-	}
-
-	EndPointAttributeDefinition.AttachAttribute(DynamicMesh);
-
-	FindEndPoints(DynamicMesh);
-	if (InOptions.DebugState == FImportOptions::EDebugState::FindEndPoints)
-	{
-		EndPointAttributeDefinition.CopyAttributeToCollection(DynamicMesh, OutCollection);
-		return ReturnWithDebugMesh(DynamicMesh);
-	}
-
-	ConnectBranchesToParent(DynamicMesh);
-	if (InOptions.DebugState == FImportOptions::EDebugState::ConnectBranchesToParent)
-	{
-		EndPointAttributeDefinition.CopyAttributeToCollection(DynamicMesh, OutCollection);
-		return ReturnWithDebugMesh(DynamicMesh);	
-	}
-
-	RootPointAttributeDefinition.AttachAttribute(DynamicMesh);
-	FindRootPoints(DynamicMesh);
-	if (InOptions.DebugState == FImportOptions::EDebugState::FindRootPoints)
-	{
-		RootPointAttributeDefinition.CopyAttributeToCollection(DynamicMesh, OutCollection);
-		EndPointAttributeDefinition.CopyAttributeToCollection(DynamicMesh, OutCollection);
-		return ReturnWithDebugMesh(DynamicMesh);
+		VertexScaleAttributeDefinition->CopyAttributeToCollection(DynamicMesh, OutCollection);
+		return EImportResult::Success;
 	}
-	
-	LengthFromRootAttributeDefinition.AttachAttribute(DynamicMesh);
 
-	ComputeLengthFromRoot(DynamicMesh);
-	if (InOptions.DebugState == FImportOptions::EDebugState::ComputeLengthFromRoot)
+	FindEndPoints(
+		DynamicMesh,
+		EndPointAttributeDefinition->GetOrAttachAttribute(DynamicMesh),
+		FillerVertexAttributeDefinition.GetOrAttachAttribute(DynamicMesh)
+	);
+	if (InParams.DebugState == EPVImportStaticMeshDebugState::FindEndPoints)
 	{
-		LengthFromRootAttributeDefinition.CopyAttributeToCollection(DynamicMesh, OutCollection);
-		return ReturnWithDebugMesh(DynamicMesh);
+		AddDynamicMeshToCollection(DynamicMesh, OutCollection);
+		EndPointAttributeDefinition->CopyAttributeToCollection(DynamicMesh, OutCollection);
+		return EImportResult::Success;
 	}
 
-	const static auto Debug_AddBranchCurvesToCollection = [](
-		FManagedArrayCollection& InCollection, 
-		const TArray<TArray<int32>>& BranchCurves, 
-		const TFunction<FVector3f(int32)>& GetVertexPosFunc)
+	ConnectBranchesToParent(
+		DynamicMesh,
+		VertexScaleAttributeDefinition->GetOrAttachAttribute(DynamicMesh),
+		EndPointAttributeDefinition->GetOrAttachAttribute(DynamicMesh),
+		FillerVertexAttributeDefinition.GetOrAttachAttribute(DynamicMesh)
+	);
+	if (InParams.DebugState == EPVImportStaticMeshDebugState::ConnectBranchesToParent)
 	{
-		const static FName BranchCurveGroupName = TEXT("BranchCurves");
-		const static FName CurveAttributeName = TEXT("Curve");
-
-		InCollection.AddGroup(BranchCurveGroupName);
-		InCollection.AddElements(BranchCurves.Num(), BranchCurveGroupName);
-
-		TManagedArray<TArray<FVector3f>>& CurveAttribute = InCollection.AddAttribute<TArray<FVector3f>>(CurveAttributeName, BranchCurveGroupName);
-		for (int32 i = 0; i < BranchCurves.Num(); ++i)
-		{
-			const TArray<int32>& BranchCurve = BranchCurves[i];
-			CurveAttribute[i].Reserve(BranchCurve.Num());
-			for (int32 VertexIndex : BranchCurve)
-			{
-				CurveAttribute[i].Add(GetVertexPosFunc(VertexIndex));
-			}
-		}
-	};
-
-	TArray<TArray<int32>> BranchCurves;
-	DrawBranchCurves(DynamicMesh, BranchCurves);
-	if (InOptions.DebugState == FImportOptions::EDebugState::DrawBranchCurves)
-	{
-		Debug_AddBranchCurvesToCollection(OutCollection, BranchCurves, [&](int32 VertexIndex) { return FVector3f(DynamicMesh.GetVertex(VertexIndex)); });
-		return ReturnWithDebugMesh(ReferenceMesh);
+		AddDynamicMeshToCollection(DynamicMesh, OutCollection);
+		EndPointAttributeDefinition->CopyAttributeToCollection(DynamicMesh, OutCollection);
+		return EImportResult::Success;
 	}
 
-	RemoveBranchCurveOverlaps(DynamicMesh, BranchCurves);
-	if (InOptions.DebugState == FImportOptions::EDebugState::RemoveBranchCurveOverlaps)
+	FindRootPoints(
+		DynamicMesh,
+		RootPointAttributeDefinition->GetOrAttachAttribute(DynamicMesh),
+		EndPointAttributeDefinition->GetOrAttachAttribute(DynamicMesh)
+	);
+	if (InParams.DebugState == EPVImportStaticMeshDebugState::FindRootPoints)
 	{
-		Debug_AddBranchCurvesToCollection(OutCollection, BranchCurves, [&](int32 VertexIndex) { return FVector3f(DynamicMesh.GetVertex(VertexIndex)); });
-		return ReturnWithDebugMesh(ReferenceMesh);
+		AddDynamicMeshToCollection(DynamicMesh, OutCollection);
+		RootPointAttributeDefinition->CopyAttributeToCollection(DynamicMesh, OutCollection);
+		EndPointAttributeDefinition->CopyAttributeToCollection(DynamicMesh, OutCollection);
+		return EImportResult::Success;
 	}
 
-	TArray<FVector3f> Points;
-	TArray<float> PointsRadii;
-	TArray<TArray<int32>> BranchPointIndices;
-	ExtractBranchesFromBranchCurves(DynamicMesh, BranchCurves, Points, PointsRadii, BranchPointIndices);
-	if (InOptions.DebugState == FImportOptions::EDebugState::ExtractBranches)
-	{
-		const static FName PointsGroupName = TEXT("Points");
-		const static FName PositionAttributeName = TEXT("Position");
-		const static FName ScaleAttributeName = TEXT("Scale");
-
-		OutCollection.AddGroup(PointsGroupName);
-		OutCollection.AddElements(Points.Num(), PointsGroupName);
-
-		TManagedArray<FVector3f>& PositionAttribute = OutCollection.AddAttribute<FVector3f>(PositionAttributeName, PointsGroupName);
-		TManagedArray<float>& ScaleAttribute = OutCollection.AddAttribute<float>(ScaleAttributeName, PointsGroupName);
-		for (int32 i = 0; i < Points.Num(); ++i)
-		{
-			PositionAttribute[i] = Points[i];
-			ScaleAttribute[i] = PointsRadii[i];
-		}
+	TArray<FPVBranchHierarchyDescription> BranchHierarchyDescriptions;
+	auto& FillerVertexAttribute = FillerVertexAttributeDefinition.GetOrAttachAttribute(DynamicMesh);
+	ComputeBranchHierarchies(
+		DynamicMesh,
+		RootPointAttributeDefinition->GetOrAttachAttribute(DynamicMesh),
+		EndPointAttributeDefinition->GetOrAttachAttribute(DynamicMesh),
+		VertexScaleAttributeDefinition->GetOrAttachAttribute(DynamicMesh),
+		LengthFromRootAttributeDefinition->GetOrAttachAttribute(DynamicMesh),
+		NextBranchPointAttributeDefinition->GetOrAttachAttribute(DynamicMesh),
+		[&](int32 VertexID) { return !FillerVertexAttribute.GetValue(VertexID); },
+		BranchHierarchyDescriptions
+	);
+
+	if (InParams.DebugState == EPVImportStaticMeshDebugState::ComputeBranchHierarchies)
+	{
+		AddDynamicMeshToCollection(DynamicMesh, OutCollection);
+		AddBranchHierarchiesToCollection(BranchHierarchyDescriptions, OutCollection);
+		RootPointAttributeDefinition->CopyAttributeToCollection(DynamicMesh, OutCollection);
+		EndPointAttributeDefinition->CopyAttributeToCollection(DynamicMesh, OutCollection);
+		VertexScaleAttributeDefinition->CopyAttributeToCollection(DynamicMesh, OutCollection);
+		LengthFromRootAttributeDefinition->CopyAttributeToCollection(DynamicMesh, OutCollection);
+		NextBranchPointAttributeDefinition->CopyAttributeToCollection(DynamicMesh, OutCollection);
 
-		Debug_AddBranchCurvesToCollection(OutCollection, BranchPointIndices, [&](int32 VertexIndex) { return FVector3f(Points[VertexIndex]); });
-
-		return ReturnWithDebugMesh(ReferenceMesh);
+		return EImportResult::Success;
 	}
 
 	FPVPlantDescription PlantDescription;
-	PlantDescription.BranchHierarchy = PV::Import::Helper::CreateBranchHierarchyFromPoints(Points, PointsRadii, BranchPointIndices);
-	
+	PlantDescription.BranchHierarchies = MoveTemp(BranchHierarchyDescriptions);
+
 	if (!PV::Import::Helper::GenerateGrowthDataFromPlantDescription(OutCollection, PlantDescription))
 	{
 		return EImportResult::FailedToGenerateGrowthData;
diff --git a/Engine/Plugins/Experimental/ProceduralVegetationEditor/Source/ProceduralVegetationEditor/Private/Importer/PVImporter_StaticMesh.h b/Engine/Plugins/Experimental/ProceduralVegetationEditor/Source/ProceduralVegetationEditor/Private/Importer/PVImporter_StaticMesh.h
--- a/Engine/Plugins/Experimental/ProceduralVegetationEditor/Source/ProceduralVegetationEditor/Private/Importer/PVImporter_StaticMesh.h
+++ b/Engine/Plugins/Experimental/ProceduralVegetationEditor/Source/ProceduralVegetationEditor/Private/Importer/PVImporter_StaticMesh.h
@@ -3,6 +3,7 @@
 
 #include "CoreMinimal.h"
 
+#include "PVImportStaticMeshParams.h"
 #include "GeometryCollection/ManagedArrayCollection.h"
 
 class UStaticMesh;
@@ -20,29 +21,5 @@ namespace PV::StaticMeshImport
 		FailedToGenerateGrowthData
 	};
 
-	struct FImportOptions
-	{
-		TArray<FName> MaterialsToKeep;
-
-		enum class EDebugState : uint8
-		{
-			None,
-			InitialMesh,
-			KeepMaterial,
-			ComputeNormals,
-			ComputeVerticesCenterAndRadius,
-			CollapseBranchesToLine,
-			FindEndPoints,
-			ConnectBranchesToParent,
-			FindRootPoints,
-			ComputeLengthFromRoot,
-			DrawBranchCurves,
-			RemoveBranchCurveOverlaps,
-			ExtractBranches,
-			MAX
-		};
-		EDebugState DebugState = EDebugState::None;
-	};
-
-	EImportResult ImportGrowthDataFromStaticMesh(const TObjectPtr<UStaticMesh>& InSourceStaticMesh, FManagedArrayCollection& OutCollection, const FImportOptions& InOptions);
+	EImportResult ImportGrowthDataFromStaticMesh(const FPVImportStaticMeshParams& InParams, FManagedArrayCollection& OutCollection);
 };
diff --git a/Engine/Plugins/Experimental/ProceduralVegetationEditor/Source/ProceduralVegetationEditor/Private/Importer/PVImporter_Texture2D.cpp b/Engine/Plugins/Experimental/ProceduralVegetationEditor/Source/ProceduralVegetationEditor/Private/Importer/PVImporter_Texture2D.cpp
--- a/Engine/Plugins/Experimental/ProceduralVegetationEditor/Source/ProceduralVegetationEditor/Private/Importer/PVImporter_Texture2D.cpp
+++ b/Engine/Plugins/Experimental/ProceduralVegetationEditor/Source/ProceduralVegetationEditor/Private/Importer/PVImporter_Texture2D.cpp
@@ -1,41 +1,53 @@
 // Copyright Epic Games, Inc. All Rights Reserved.
 
 #include "PVImporter_Texture2D.h"
+#include "PVImportHelpers.h"
+#include "PVPlantDescription.h"
+#include "PVDynamicMeshVertexAttribute.h"
 
 #include "IntVectorTypes.h"
+#include "PCGSettings.h"
 #include "GeometryCollection/ManagedArrayCollection.h"
 
 #include "Engine/Texture2D.h"
 #include "Materials/MaterialInstanceDynamic.h"
 #include "Materials/MaterialInterface.h"
 
+#include "DynamicMesh/DynamicMesh3.h"
+#include "DynamicMesh/DynamicMeshAABBTree3.h"
 #include "GeometryCollection/GeometryCollection.h"
 #include "GeometryCollection/GeometryCollectionUtility.h"
 #include "SpatialAlgo/PCGMarchingSquares.h"
+#include "CompGeom/Delaunay2.h"
+#include "Selections/MeshConnectedComponents.h"
 
 namespace Internal
 {
+	using FBooleanVertexAttribute = PV::DynamicMeshHelper::FBooleanVertexAttribute;
+	using FFloatVertexAttribute = PV::DynamicMeshHelper::FFloatVertexAttribute; using FIndex2i = UE::Geometry::FIndex2i;
+
+	using FIndex3i = UE::Geometry::FIndex3i;
 	using FVector2i = UE::Geometry::FVector2i;
-	constexpr float DebugQuadSize = 100;
+	using FMeshConnectedComponents = UE::Geometry::FMeshConnectedComponents;
+	using FDynamicMeshAABBTree3 = UE::Geometry::FDynamicMeshAABBTree3;
+	using FDynamicMesh3 = UE::Geometry::FDynamicMesh3;
+	using FDynamicMeshQueries3 = UE::Geometry::TMeshQueries<FDynamicMesh3>;
 
 	const static FName PixelsDebugGroupName = TEXT("PixelDebug");
 	const static FName PixelPositionsAttributeName = TEXT("Position");
 	const static FName PixelValueAttributeName = TEXT("Value");
 	const static FName CurveGroupName = TEXT("Curves");
 	const static FName CurveAttributeName = TEXT("Curve");
+	const static FName InputRootPointsGroupName = TEXT("InputRootPoints");
+	const static FName RootPointsGroupName = TEXT("RootPoints");
+	const static FName RootPointsAttributeName = TEXT("Point");
+
+	const static PV::TDynamicMeshVertexAttributeDefinition<bool> PerimiterAttributeDefinition(TEXT("Perimiter"));
 	
 	static int32 GetPixelIndex(const FVector2i& PixelCoords, int32 ImageWidth)
 	{
 		return PixelCoords.X + PixelCoords.Y * ImageWidth;
 	}
-	
-	template<typename T>
-	static const T& SampleImageBuffer(const TArray<T>& Buffer, const int32 Stride, const FVector2i& XYCoord)
-	{
-		const int32 Index = GetPixelIndex(XYCoord, Stride);
-		check(Buffer.IsValidIndex(Index));
-		return Buffer[Index];
-	}
 
 	static FVector2f PixelCoordToUV(const FVector2i& PixelCoord, const FVector2i& ImageSize)
 	{
@@ -43,8 +55,26 @@ namespace Internal
 		const float V = PixelCoord.Y / static_cast<float>(ImageSize.Y - 1);
 		return FVector2f(U, V);
 	}
+	
+	template<typename Vector3dType = float, typename Vector2dType = float>
+	UE::Math::TVector<Vector3dType> ToVector3(const UE::Math::TVector2<Vector2dType>& Vec2D, float Scale)
+	{
+		return UE::Math::TVector<Vector3dType>(0, Vec2D.X * Scale, (1.0 - Vec2D.Y) * Scale);
+	}
 
-	static void AddDebugPixelValuesToCollection(FManagedArrayCollection& InOutCollection, const TArray<bool>& ImageData, const FVector2i& ImageSize)
+	static void AddDebugRootPointsToCollection(FManagedArrayCollection& InOutCollection, const TArray<FPVImportTexture2DRootPoint>& InRootPoints, float Scale)
+	{
+		InOutCollection.AddGroup(InputRootPointsGroupName);
+		InOutCollection.AddElements(InRootPoints.Num(), InputRootPointsGroupName);
+	 
+		auto& InputRootPointAttribute = InOutCollection.AddAttribute<FVector3f>(RootPointsAttributeName, InputRootPointsGroupName);
+		for (int32 i = 0; i < InRootPoints.Num(); i++)
+		{
+			InputRootPointAttribute[i] = ToVector3(InRootPoints[i].Position, Scale);
+		}
+	}
+	
+	static void AddDebugPixelValuesToCollection(FManagedArrayCollection& InOutCollection, const TArray<bool>& ImageData, const FVector2i& ImageSize, float Scale)
 	{
 		check(ImageData.Num() == ImageSize.X * ImageSize.Y);
 		
@@ -61,14 +91,14 @@ namespace Internal
 				const FVector2f UV = PixelCoordToUV(PixelCoords, ImageSize);
 				
 				const int32 PixelIndex = GetPixelIndex(PixelCoords, ImageSize.X);
-				PixelPositionsAttribute[PixelIndex] = FVector3f(0, UV.X * DebugQuadSize, (1.0 - UV.Y) * DebugQuadSize);
+				PixelPositionsAttribute[PixelIndex] = ToVector3(UV, Scale);
 				PixelValueAttribute[PixelIndex] = ImageData[PixelIndex];
 			}
 		}
 	}
 
-	static void AddDebugCurvesToCollection(FManagedArrayCollection& InOutCollection, const TArray<TArray<FVector2f>>& InCurves)
-	{	
+	static void AddDebugCurvesToCollection(FManagedArrayCollection& InOutCollection, const TArray<TArray<FVector2f>>& InCurves, float Scale)
+	{
 		InOutCollection.AddGroup(CurveGroupName);
 		InOutCollection.AddElements(InCurves.Num(), CurveGroupName);
     	
@@ -78,12 +108,12 @@ namespace Internal
 			const TArray<FVector2f>& Curve = InCurves[i];
 			for (const FVector2f& CurvePos : Curve)
 			{
-				CurveAttribute[i].Emplace(0, CurvePos.X * DebugQuadSize, DebugQuadSize - (CurvePos.Y * DebugQuadSize));
+				CurveAttribute[i].Add(ToVector3(CurvePos, Scale));
 			}
 		}
 	}
 
-	static TSharedRef<FGeometryCollection> GenerateDebugQuad(UTexture2D& Texture2D)
+	static void AddDebugTextureQuadToCollection(FManagedArrayCollection& InOutCollection, UTexture2D& Texture2D, float Scale)
 	{
 		const TArray<FVector> Points = {
 			FVector(0, 0, 1),
@@ -112,7 +142,7 @@ namespace Internal
 			Normals,
 			Triangles,
 			UVs,
-			FTransform(FRotator::ZeroRotator, FVector::ZeroVector, FVector(Internal::DebugQuadSize)),
+			FTransform(FRotator::ZeroRotator, FVector::ZeroVector, FVector(Scale)),
 			FTransform::Identity,
 			1
 		);
@@ -133,7 +163,7 @@ namespace Internal
 			MaterialArray[0] = MaterialInstance->GetPathName();
 		}
 
-		return MeshGeomCollection.ToSharedRef();
+		MeshGeomCollection->CopyTo(&InOutCollection);
 	}
 
 	static void SampleImage(FTextureSource& InTextureSource, FVector2i& InOutSampleSize, TArray<bool>& OutBlackAndWhiteImageData, bool bInvertImage, float WhiteLevel)
@@ -335,34 +365,504 @@ namespace Internal
 			Curve.SetNum(KeepPoints.Num());
 		}
 	}
+	
+	static void TriangulateCurvesAndSplitEdges(const TArray<TArray<FVector2f>>& InCurves, UE::Geometry::FDynamicMesh3& OutMesh, float Scale)
+	{
+		/*
+		* This function convertes the perimiter curves into a DynamicMesh and adds a vertex at the center of each edge spanning two perimiter points.
+		* 
+		* It does this by first triangulating the curve using FDelaunay2, it then looks at each edge and if said edge has two perimiter points, it inserts a new vertex
+		* at the center of the edge. It does this for all edges and then re-triangulates the entire mesh with this new vertex set.
+		* 
+		* The end result will look something like the following diagram:
+		*      *      
+		*     / \           /|\
+		*    /   \         / | \
+		*   /     \       /__|__\
+		*   *     *       |\ | /|
+		*   |     |  -->  | \|/ |
+		*   *     *       |--|--|
+		*   |     |       | /|\ |
+		*   *     *       |/ | \|
+		*/
+
+		PerimiterAttributeDefinition.AttachAttribute(OutMesh);
+
+		for (int32 CurveIndex = 0; CurveIndex < InCurves.Num(); CurveIndex++)
+		{
+			const TArray<FVector2f>& Curve = InCurves[CurveIndex];
+			
+			const TArray<FIndex2i> OriginalCurveEdges = Invoke([&]() 
+			{
+				TArray<FIndex2i> Edges;
+				Edges.Reserve(Curve.Num());
+				for (int32 i = 0; i < Curve.Num() - 1; i++)
+				{
+					Edges.Emplace(i, i + 1);
+				}
+				Edges.Emplace(0, Curve.Num() - 1);
+				return Edges;
+			});
+			
+			UE::Geometry::FDelaunay2 Delaunay;
+			bool bTriangulateSucess = Delaunay.Triangulate(Curve, OriginalCurveEdges);
+			if (!bTriangulateSucess)
+			{
+				continue;
+			}
+			
+			const TArray<FIndex3i> OriginalCurveTriangles = Delaunay.GetFilledTriangles(OriginalCurveEdges, UE::Geometry::FDelaunay2::EFillMode::Solid);
+			
+			TArray<FVector2f> SubdividedMeshPoints;
+			SubdividedMeshPoints.Reserve(Curve.Num() + OriginalCurveTriangles.Num());
+			SubdividedMeshPoints.Append(Curve);
+			
+			TArray<FIndex2i> SubdividedMeshEdges;
+			SubdividedMeshEdges.Reserve(OriginalCurveTriangles.Num() * 4);
+			SubdividedMeshEdges.Append(OriginalCurveEdges);
+			
+			TSet<FIndex2i> CheckedEdges;
+			CheckedEdges.Reserve(OriginalCurveTriangles.Num() * 3);
+			CheckedEdges.Append(OriginalCurveEdges);
+			
+			for (const FIndex3i& Triangle : OriginalCurveTriangles)
+			{
+				for (int32 i = 0; i < 3; i++)
+				{
+					FIndex2i EdgeIndex(Triangle[i], Triangle[(i + 1) % 3]);
+					EdgeIndex.Sort();
+					
+					bool bAlreadyAdded = false;
+					CheckedEdges.Add(EdgeIndex, &bAlreadyAdded);
+					if (!bAlreadyAdded)
+					{
+						const int32 PointIndex1 = EdgeIndex[0];
+						const int32 PointIndex2 = EdgeIndex[1];
+						const FVector2f& Pos1 = Curve[PointIndex1];
+						const FVector2f& Pos2 = Curve[PointIndex2];
+						const int32 NewPointIndex = SubdividedMeshPoints.Add((Pos1 + Pos2) / 2);
+						
+						const FIndex2i NewEdge1(PointIndex1, NewPointIndex);
+						const FIndex2i NewEdge2(PointIndex2, NewPointIndex);
+						
+						SubdividedMeshEdges.Add(NewEdge1);
+						SubdividedMeshEdges.Add(NewEdge2);
+					}
+				}
+			}
+			
+			bTriangulateSucess = Delaunay.Triangulate(SubdividedMeshPoints, SubdividedMeshEdges);
+			if (!bTriangulateSucess)
+			{
+				continue;
+			}
+			
+			UE::Geometry::FDynamicMesh3 NewDynamicMesh;
+			PerimiterAttributeDefinition.AttachAttribute(NewDynamicMesh);
+			auto& PerimiterAttribute = PerimiterAttributeDefinition.GetAttributeChecked(NewDynamicMesh);
+			
+			const TArray<FIndex3i> Triangles = Delaunay.GetFilledTriangles(OriginalCurveEdges, UE::Geometry::FDelaunay2::EFillMode::Solid);
+			const TArray<FVector2f>& Vertices = SubdividedMeshPoints;
+			
+			for (int32 i = 0; i < Vertices.Num(); ++i)
+			{
+				const FVector2f& Vertex = Vertices[i];
+				UE::Geometry::FVertexInfo VertexInfo(
+					ToVector3<double, float>(Vertex, Scale),
+					FVector3f(-1, 0, 0)
+				);
+				
+				const int32 VertexID = NewDynamicMesh.AppendVertex(VertexInfo);
+				PerimiterAttribute.SetValue(VertexID, i < Curve.Num());
+			}
+			
+			for (const FIndex3i& Triangle : Triangles)
+			{
+				NewDynamicMesh.AppendTriangle(Triangle);
+			}
+			
+			OutMesh.AppendWithOffsets(NewDynamicMesh);
+		}
+	}
+
+	static void ComputeVertexScale2D(
+		const UE::Geometry::FDynamicMesh3& DynamicMesh,
+		const FBooleanVertexAttribute& PerimiterAttribute,
+		FFloatVertexAttribute& VertexScaleAttribute
+	)
+	{
+		/*
+		* This function computes the vertex scale for each vertex in the mesh by finding the nearest point on the perimiter curve.
+		*/
+
+		const auto FindDistToPerimiter = [&](const FVector3d& Point, const TArray<int32>& PerimiterVertices)
+		{
+			// TODO: This algorithm is N^2, we can probably reduce that by using FSparseDynamicOctree3
+
+			double ClosestDistSqr = MAX_flt;
+
+			for (int32 i = 0; i < PerimiterVertices.Num(); ++i)
+			{
+				const int32 PointIndex0 = PerimiterVertices[i];
+				const int32 PointIndex1 = PerimiterVertices[FMath::WrapExclusive(i + 1, 0, PerimiterVertices.Num())];
+
+				const FVector3d& VertexLocation1 = DynamicMesh.GetVertex(PointIndex0);
+				const FVector3d& VertexLocation2 = DynamicMesh.GetVertex(PointIndex1);
+
+				const FVector3d ClosestPointOnLine = FMath::ClosestPointOnSegment(Point, VertexLocation1, VertexLocation2);
+				const double DistToLineSqr = (Point - ClosestPointOnLine).SizeSquared();
+				if (DistToLineSqr < ClosestDistSqr)
+				{
+					ClosestDistSqr = DistToLineSqr;
+				}
+			}
+
+			return ClosestDistSqr > 0 ? FMath::Sqrt(ClosestDistSqr) : 0;
+		};
+
+		UE::Geometry::FMeshConnectedComponents ConnectedComponents(&DynamicMesh);
+		ConnectedComponents.FindConnectedVertices();
+
+		for (const auto& Component : ConnectedComponents)
+		{
+			TArray<int32> PerimiterVertices = PerimiterAttribute.FindAllNonZero(Component.Indices);
+			PerimiterVertices.Sort(); // This assumes perimiter vertices are ordered sequentially, which is currently the case.
+
+			for (int32 VertexID : Component.Indices)
+			{
+				if (!PerimiterAttribute.GetValue(VertexID))
+				{
+					const float VertexScale = FindDistToPerimiter(DynamicMesh.GetVertex(VertexID), PerimiterVertices);
+					VertexScaleAttribute.SetValue(VertexID, VertexScale);
+				}
+			}
+		}
+	}
+
+	static void FindEndPoints2D(
+		const UE::Geometry::FDynamicMesh3& DynamicMesh,
+		float MaxTipAngleSearchDist,
+		float TipAngleThresholdDeg,
+		FBooleanVertexAttribute& EndPointAttribute,
+		const FBooleanVertexAttribute& PerimiterAttribute
+	)
+	{
+		/*
+		* This function finds the end points (or "tips") on a 2D mesh by walking the perimiter vertices and checking the angle between itself and its nearest 
+		* two neighbours on the perimiter. To allow for curved (or flat-tipped) end points to be detected, the function takes a second pass where it walks outwards along the
+		* perimited (limited by MaxTipAngleSearchDist), and checks the angle of the current vertex towards the subsequent neighbours. (see diagram a for example if flat-tipped end point)
+		* 
+		*     a: ___     b:
+		*       /   \      /\
+		*       |   |     |  |
+		*       |   |     |  |
+		*/
+
+		using namespace PV::DynamicMeshHelper;
+
+		const static auto ComputeAngleDeg = [](
+			const UE::Geometry::FDynamicMesh3& Mesh,
+			const TArray<int32>& PerimiterVertices, 
+			int32 PerimeterIndex, 
+			int32 LeftPerimeterIndex, 
+			int32 RightPerimeterIndex)->float
+		{
+			check(PerimiterVertices.IsValidIndex(PerimeterIndex));
+			check(Mesh.IsVertex(PerimiterVertices[PerimeterIndex]));
+
+			LeftPerimeterIndex = FMath::WrapExclusive(LeftPerimeterIndex, 0, PerimiterVertices.Num());
+			RightPerimeterIndex = FMath::WrapExclusive(RightPerimeterIndex, 0, PerimiterVertices.Num());
+
+			const int32 VertexIndex = PerimiterVertices[PerimeterIndex];
+			const int32 NeighbourVertexIndex1 = PerimiterVertices[LeftPerimeterIndex];
+			const int32 NeighbourVertexIndex2 = PerimiterVertices[RightPerimeterIndex];
+
+			const FVector3d& PointPos = Mesh.GetVertex(VertexIndex);
+			const FVector3d& PerimeterNeighbour1Pos = Mesh.GetVertex(NeighbourVertexIndex1);
+			const FVector3d& PerimeterNeighbour2Pos = Mesh.GetVertex(NeighbourVertexIndex2);
+
+			const FVector3d Normal1 = (PointPos - PerimeterNeighbour1Pos).GetSafeNormal();
+			const FVector3d Normal2 = (PointPos - PerimeterNeighbour2Pos).GetSafeNormal();
+
+			const double Dot = FVector3d::DotProduct(Normal1, Normal2);
+			return (Dot + 1) / 2 * 180.0;
+		};
+
+		UE::Geometry::FMeshConnectedComponents ConnectedComponents(&DynamicMesh);
+		ConnectedComponents.FindConnectedVertices();
+
+		for (const auto& Component : ConnectedComponents)
+		{
+			const TArray<int32> PerimiterVertices = PerimiterAttribute.FindAllNonZero(Component.Indices);
+			 	
+			TArray<int32> ExcludedPerimiterVertices;
+			ExcludedPerimiterVertices.Reserve(PerimiterVertices.Num());
+			 	
+			TArray<float> PerimiterAngles;
+			PerimiterAngles.SetNumZeroed(PerimiterVertices.Num());
+			 	
+			for (int32 PerimiterIndex = 0; PerimiterIndex < PerimiterVertices.Num(); ++PerimiterIndex)
+			{
+				const int32 PointIndex = PerimiterVertices[PerimiterIndex];
+			 	check(DynamicMesh.IsVertex(PointIndex));
+			 		
+			 	int32 NumInternalNeighbours = 0;
+				const TArray<int32> PointNeighbours = PV::DynamicMeshHelper::GetVertexNeighbours(DynamicMesh, PointIndex);
+			 	for (int32 NeighbourIndex : PointNeighbours)
+			 	{
+			 		if (!PerimiterAttribute.GetValue(NeighbourIndex))
+			 		{
+			 			NumInternalNeighbours += 1;
+			 		}
+			 	}
+			 	
+				// TODO: Since we do not know whether the tip is pointing inwards or outwards, we use the fact that most outward pointing tips only have one
+				// internal connection (e.g. a connection to a vertec which is not a perimiter vertex). This is however not a very reliable method of
+				// determining the direction of the tip. A better alternative would be to calculate the normal for each vertex in the perimiter curve
+				// and use said normal to determine whether the tip is pointing outwads.
+			 	if (NumInternalNeighbours == 1)
+			 	{
+			 		const double TipAngle = ComputeAngleDeg(DynamicMesh, PerimiterVertices, PerimiterIndex, PerimiterIndex - 1, PerimiterIndex + 1);
+			 		if (TipAngle >= TipAngleThresholdDeg)
+			 		{
+						EndPointAttribute.SetValue(PointIndex, true);
+						PerimiterAngles[PerimiterIndex] = TipAngle;
+			 		}
+			 		else 
+			 		{
+						ExcludedPerimiterVertices.Add(PerimiterIndex);
+			 		}
+			 	}
+			}
+			 	
+			const float MaxTipAngleSearchDistSqr = MaxTipAngleSearchDist * MaxTipAngleSearchDist;
+			for (int32 PerimiterIndex : ExcludedPerimiterVertices)
+			{
+				const int32 PointIndex = PerimiterVertices[PerimiterIndex];
+
+			 	float LeftDistSqr = 0.f;
+			 	float RightDistSqr = 0.f;
+			 		
+			 	int32 LeftPerimiterIndex = PerimiterIndex;
+			 	int32 RightPerimiterIndex = PerimiterIndex;
+			 		
+			 	while (LeftDistSqr < MaxTipAngleSearchDistSqr || RightDistSqr < MaxTipAngleSearchDistSqr)
+			 	{
+			 		if (LeftDistSqr < MaxTipAngleSearchDistSqr)
+			 		{
+						LeftPerimiterIndex = FMath::WrapExclusive(LeftPerimiterIndex - 1, 0, PerimiterVertices.Num());
+			 				
+			 			if (LeftPerimiterIndex == PerimiterIndex)
+			 			{
+			 				break;
+			 			}
+			 			
+						const int32 LeftVertexIndex = PerimiterVertices[LeftPerimiterIndex];
+
+			 			if (EndPointAttribute.GetValue(LeftVertexIndex))
+			 			{
+			 				LeftDistSqr = MaxTipAngleSearchDistSqr;
+			 			}
+			 			else
+			 			{
+			 				LeftDistSqr = (DynamicMesh.GetVertex(PointIndex) - DynamicMesh.GetVertex(LeftVertexIndex)).SizeSquared();
+			 			}
+			 		}
+			 			
+			 		if (RightDistSqr < MaxTipAngleSearchDistSqr)
+			 		{
+						RightPerimiterIndex = FMath::WrapExclusive(RightPerimiterIndex + 1, 0, PerimiterVertices.Num());
+			 			
+			 			if (RightPerimiterIndex == PerimiterIndex)
+			 			{
+			 				break;
+			 			}
+
+						const int32 RightPointIndex = PerimiterVertices[RightPerimiterIndex];
+			 				
+			 			if (EndPointAttribute.GetValue(RightPointIndex))
+			 			{
+			 				RightDistSqr = MaxTipAngleSearchDistSqr;
+			 			}
+			 			else
+			 			{
+							RightDistSqr = (DynamicMesh.GetVertex(PointIndex) - DynamicMesh.GetVertex(RightPointIndex)).SizeSquared();
+			 			}
+			 		}
+			 			
+			 		const float TipAngle = ComputeAngleDeg(DynamicMesh, PerimiterVertices, PerimiterIndex, LeftPerimiterIndex, RightPerimiterIndex);
+			 		if (TipAngle >= TipAngleThresholdDeg)
+			 		{
+						EndPointAttribute.SetValue(PointIndex, true);
+						PerimiterAngles[PerimiterIndex] = TipAngle;
+			 			break;
+			 		}
+			 	}
+			}
+			 	
+			// When we have 2 neighbouring end points (e.g. when the tip is a plane with 2 45* corners that can both be considered end points),
+			// then select the one with the steepest angle
+			for (int32 PerimiterIndex = 0; PerimiterIndex < PerimiterVertices.Num(); ++PerimiterIndex)
+			{
+				const int32 PointIndex = PerimiterVertices[PerimiterIndex];
+				if (!EndPointAttribute.GetValue(PointIndex))
+				{
+					continue;
+				}
+
+			 	const int32 NextPerimiterIndex = FMath::WrapExclusive(PerimiterIndex + 1, 0, PerimiterVertices.Num());
+				const int32 NextPointIndex = PerimiterVertices[NextPerimiterIndex];
+
+			 	if (EndPointAttribute.GetValue(NextPointIndex))
+			 	{
+			 		const int32 EndPointToRemove = PerimiterAngles[NextPerimiterIndex] >= PerimiterAngles[PerimiterIndex] ? PointIndex : NextPointIndex;
+					EndPointAttribute.SetValue(EndPointToRemove, false);
+			 	}
+			}
+		}
+	}
+
+	static void ApplyRootPoints2D(
+		UE::Geometry::FDynamicMesh3& DynamicMesh, 
+		float Scale,
+		const TArray<FPVImportTexture2DRootPoint>& InRootPoints,
+		FBooleanVertexAttribute& RootPointAttribute,
+		FBooleanVertexAttribute& EndPointAttribute
+	)
+	{
+		// Finds the closest mesh component, and vertex in said mesh component, for each point in InRootPoints
+
+		using namespace PV::DynamicMeshHelper;
+
+		FMeshConnectedComponents Components(&DynamicMesh);
+		Components.FindConnectedTriangles();
+
+		TArray<FDynamicMeshAABBTree3> AABBTrees;
+		AABBTrees.Reserve(Components.Components.Num());
+		
+		for (const auto& Component : Components)
+		{
+			FDynamicMeshAABBTree3 DynamicMeshAABBTree(&DynamicMesh, false);
+			DynamicMeshAABBTree.Build(Component.Indices);
+
+			AABBTrees.Add(MoveTemp(DynamicMeshAABBTree));
+		}
+	
+		struct FMeshIndexDistPair
+		{
+			int32 ComponentIndex = INDEX_NONE;
+			double DistSqr = TNumericLimits<double>::Max();
+			int32 ClosestTriangle = INDEX_NONE;
+		};
+		TArray<FMeshIndexDistPair> RootPositionsClosestMesh;
+		RootPositionsClosestMesh.SetNum(InRootPoints.Num());
+		
+		for (int32 i = 0; i < InRootPoints.Num(); i++)
+		{
+			FMeshIndexDistPair& ClosestMesh = RootPositionsClosestMesh[i];
+		
+			for (int32 ComponentIndex = 0; ComponentIndex < Components.Num(); ComponentIndex++)
+			{
+				const FDynamicMeshAABBTree3& AABBTree = AABBTrees[ComponentIndex];
+				
+				double DistSqr = 0;
+				const FVector3d RootPoint = ToVector3<double, float>(InRootPoints[i].Position, Scale);
+				const int32 NearTriID = AABBTree.FindNearestTriangle(RootPoint, DistSqr);
+				if (NearTriID != INDEX_NONE)
+				{
+					UE::Geometry::FDistPoint3Triangle3d Query = FDynamicMeshQueries3::TriangleDistance(DynamicMesh, NearTriID, RootPoint);
+					DistSqr = Query.GetSquared();
+					if (DistSqr < ClosestMesh.DistSqr)
+					{
+						ClosestMesh.ComponentIndex = ComponentIndex;
+						ClosestMesh.DistSqr = DistSqr;
+						ClosestMesh.ClosestTriangle = NearTriID;
+					}
+				}
+			}
+		}
+		
+		struct FRootPositionIndexDistPair
+		{
+			int32 RootPositionIndex = INDEX_NONE;
+			double DistSqr = TNumericLimits<double>::Max();
+		};
+		TArray<FRootPositionIndexDistPair> MeshClosestRootPosition;
+		MeshClosestRootPosition.SetNum(Components.Num());
+		
+		for (int32 i = 0; i < RootPositionsClosestMesh.Num(); i++)
+		{
+			const int32 ComponentIndex = RootPositionsClosestMesh[i].ComponentIndex;
+			if (ComponentIndex == INDEX_NONE)
+			{
+				// Unable to find a closest mesh (mesh generation has most likely failed).
+				continue;
+			}
+			
+			const double DistSqr = RootPositionsClosestMesh[i].DistSqr;
+			
+			if (DistSqr < MeshClosestRootPosition[ComponentIndex].DistSqr)
+			{
+				MeshClosestRootPosition[ComponentIndex].RootPositionIndex = i;
+				MeshClosestRootPosition[ComponentIndex].DistSqr = DistSqr;
+			}
+		}
+		
+		for (int32 ComponentIndex = 0; ComponentIndex < Components.Num(); ComponentIndex++)
+		{
+			if (MeshClosestRootPosition[ComponentIndex].RootPositionIndex == INDEX_NONE)
+			{
+				// Unable to find a root position for this mesh
+				continue;
+			}
+			
+			const int32 RootPositionIndex = MeshClosestRootPosition[ComponentIndex].RootPositionIndex;
+			
+			const FVector3d& RootPosition = ToVector3<double, float>(InRootPoints[RootPositionIndex].Position, Scale);
+			const int32 TriangleID = RootPositionsClosestMesh[RootPositionIndex].ClosestTriangle;
+			const FIndex3i& Triangle = DynamicMesh.GetTriangle(TriangleID);
+			const int32 ClosestTriangleVertex = FMath::Min3Index(
+				(RootPosition - DynamicMesh.GetVertex(Triangle[0])).SizeSquared(),
+				(RootPosition - DynamicMesh.GetVertex(Triangle[1])).SizeSquared(),
+				(RootPosition - DynamicMesh.GetVertex(Triangle[2])).SizeSquared()
+			);
+			
+			const int32 RootVertexID = Triangle[ClosestTriangleVertex];
+
+			RootPointAttribute.SetValue(RootVertexID, true);
+			EndPointAttribute.SetValue(RootVertexID, false);
+
+			const FVector VertexTranslation = DynamicMesh.GetVertex(RootVertexID);
+			const FQuat VertexRotation = FQuat(FVector::ForwardVector, FMath::DegreesToRadians(InRootPoints[RootPositionIndex].Rotation));
+
+			for (int32 VertexIndex : PV::DynamicMeshHelper::GetUniqueVertexIndices(DynamicMesh, Components[ComponentIndex].Indices))
+			{
+				FVector Vertex = DynamicMesh.GetVertex(VertexIndex);
+				Vertex -= VertexTranslation;
+				Vertex = VertexRotation.UnrotateVector(Vertex);
+				DynamicMesh.SetVertex(VertexIndex, Vertex);
+			}
+		}
+	}
 };
 
-PV::Texture2DImport::EImportResult PV::Texture2DImport::ImportGrowthDataFromTexture2D(const TObjectPtr<UTexture2D>& InSourceTexture2D,
-	FManagedArrayCollection& OutCollection, const FImportOptions& InOptions)
+PV::Texture2DImport::EImportResult PV::Texture2DImport::ImportGrowthDataFromTexture2D(const FPVImportTexture2DParams& InParams, FManagedArrayCollection& OutCollection)
 {
+	using namespace PV::DynamicMeshHelper;
 	using namespace Internal;
 	
-	if (!InSourceTexture2D || !InSourceTexture2D->Source.IsValid())
+	if (!InParams.Texture2DAsset || !InParams.Texture2DAsset->Source.IsValid())
 	{
 		return EImportResult::InvalidSourceTexture;
 	}
 
-	const auto ReturnWithDebugMesh = [&]()->EImportResult
+	if (InParams.DebugState == EPVImportTexture2DDebugState::Texture)
 	{
-		const TSharedRef<FGeometryCollection> DebugMeshGeomCollection = GenerateDebugQuad(*InSourceTexture2D);
-
-		FManagedArrayCollection MeshArrayCollection;
-		DebugMeshGeomCollection->CopyTo(&OutCollection);
-	
+		AddDebugTextureQuadToCollection(OutCollection, *InParams.Texture2DAsset, InParams.Scale);
 		return EImportResult::Success;
-	};
-
-	if (InOptions.DebugState == FImportOptions::EDebugState::Texture)
-	{
-		return ReturnWithDebugMesh();
 	}
 
-	FVector2i SampleSize = FVector2i(InOptions.SampleWidth, InOptions.SampleHeight);
+	FVector2i SampleSize = FVector2i(InParams.SampleWidth, InParams.SampleHeight);
 	if (SampleSize.X <= 1 || SampleSize.Y <= 1)
 	{
 		return EImportResult::InvalidSampleSize;
@@ -370,46 +870,136 @@ PV::Texture2DImport::EImportResult PV::Texture2DImport::ImportGrowthDataFromText
 	
 	TArray<bool> BlackAndWhiteImageData;
 
-	SampleImage(InSourceTexture2D->Source, SampleSize, BlackAndWhiteImageData, InOptions.bInvertImage, InOptions.WhiteLevel);
+	SampleImage(InParams.Texture2DAsset->Source, SampleSize, BlackAndWhiteImageData, InParams.bInvertImage, InParams.WhiteLevel);
 	if (SampleSize.X <= 1 || SampleSize.Y <= 1)
 	{
 		return EImportResult::InvalidTextureSize;
 	}
 	
-	if (InOptions.DebugState == FImportOptions::EDebugState::SampleImage)
+	if (InParams.DebugState == EPVImportTexture2DDebugState::SampleImage)
 	{
-		AddDebugPixelValuesToCollection(OutCollection, BlackAndWhiteImageData, SampleSize);
-		return ReturnWithDebugMesh();
+		AddDebugPixelValuesToCollection(OutCollection, BlackAndWhiteImageData, SampleSize, InParams.Scale);
+		AddDebugTextureQuadToCollection(OutCollection, *InParams.Texture2DAsset, InParams.Scale);
+		return EImportResult::Success;
 	}
 
 	TArray<TArray<FVector2f>> Curves;
 	TraceImage(BlackAndWhiteImageData, SampleSize, Curves);
-	if (InOptions.DebugState == FImportOptions::EDebugState::TraceImage)
+	if (InParams.DebugState == EPVImportTexture2DDebugState::TraceImage)
 	{
-		AddDebugCurvesToCollection(OutCollection, Curves);
-		return ReturnWithDebugMesh();
+		AddDebugCurvesToCollection(OutCollection, Curves, InParams.Scale);
+		AddDebugTextureQuadToCollection(OutCollection, *InParams.Texture2DAsset, InParams.Scale);
+		return EImportResult::Success;
 	}
 	
-	RemoveSmallCurves(Curves, InOptions.MinBoundsArea);
-	if (InOptions.DebugState == FImportOptions::EDebugState::RemoveSmallCurves)
+	RemoveSmallCurves(Curves, InParams.MinBoundsArea);
+	if (InParams.DebugState == EPVImportTexture2DDebugState::RemoveSmallCurves)
 	{
-		AddDebugCurvesToCollection(OutCollection, Curves);
-		return ReturnWithDebugMesh();
+		AddDebugCurvesToCollection(OutCollection, Curves, InParams.Scale);
+		AddDebugTextureQuadToCollection(OutCollection, *InParams.Texture2DAsset, InParams.Scale);
+		return EImportResult::Success;
 	}
 	
-	SmoothCurves(Curves, InOptions.SmoothingIterations);
-	if (InOptions.DebugState == FImportOptions::EDebugState::SmoothCurves)
+	SmoothCurves(Curves, InParams.SmoothingIterations);
+	if (InParams.DebugState == EPVImportTexture2DDebugState::SmoothCurves)
 	{
-		AddDebugCurvesToCollection(OutCollection, Curves);
-		return ReturnWithDebugMesh();
+		AddDebugCurvesToCollection(OutCollection, Curves, InParams.Scale);
+		AddDebugTextureQuadToCollection(OutCollection, *InParams.Texture2DAsset, InParams.Scale);
+		return EImportResult::Success;
 	}
 	
-	SimplifyCurves(Curves, InOptions.SimplificationEpsilon);
-	if (InOptions.DebugState == FImportOptions::EDebugState::SimplifyCurves)
+	SimplifyCurves(Curves, InParams.SimplificationEpsilon);
+	if (InParams.DebugState == EPVImportTexture2DDebugState::SimplifyCurves)
 	{
-		AddDebugCurvesToCollection(OutCollection, Curves);
-		return ReturnWithDebugMesh();
+		AddDebugCurvesToCollection(OutCollection, Curves, InParams.Scale);
+		AddDebugTextureQuadToCollection(OutCollection, *InParams.Texture2DAsset, InParams.Scale);
+		return EImportResult::Success;
 	}
 	
+	UE::Geometry::FDynamicMesh3 DynamicMesh;
+	TriangulateCurvesAndSplitEdges(Curves, DynamicMesh, InParams.Scale);
+	if (InParams.DebugState == EPVImportTexture2DDebugState::TriangulateCurves)
+	{
+		AddDynamicMeshToCollection(DynamicMesh, OutCollection);
+		PerimiterAttributeDefinition.CopyAttributeToCollection(DynamicMesh, OutCollection);
+		return EImportResult::Success;
+	}
+
+	ComputeVertexScale2D(
+		DynamicMesh, 
+		PerimiterAttributeDefinition.GetOrAttachAttribute(DynamicMesh),
+		VertexScaleAttributeDefinition->GetOrAttachAttribute(DynamicMesh)
+	);
+	if (InParams.DebugState == EPVImportTexture2DDebugState::ComputeVertexScale)
+	{
+		AddDynamicMeshToCollection(DynamicMesh, OutCollection);
+		VertexScaleAttributeDefinition->CopyAttributeToCollection(DynamicMesh, OutCollection);
+		return EImportResult::Success;
+	}
+
+	FindEndPoints2D(
+		DynamicMesh,
+		InParams.MaxTipAngleSearchDist * InParams.Scale,
+		InParams.TipAngleThresholdInDegrees,
+		EndPointAttributeDefinition->GetOrAttachAttribute(DynamicMesh),
+		PerimiterAttributeDefinition.GetOrAttachAttribute(DynamicMesh)
+	);
+	if (InParams.DebugState == EPVImportTexture2DDebugState::FindTips)
+	{
+		AddDynamicMeshToCollection(DynamicMesh, OutCollection);
+		EndPointAttributeDefinition->CopyAttributeToCollection(DynamicMesh, OutCollection);
+ 		return EImportResult::Success;
+	}
+
+	ApplyRootPoints2D(
+		DynamicMesh, 
+		InParams.Scale,
+		InParams.RootPoints,
+		RootPointAttributeDefinition->GetOrAttachAttribute(DynamicMesh),
+		EndPointAttributeDefinition->GetOrAttachAttribute(DynamicMesh)
+	);
+	if (InParams.DebugState == EPVImportTexture2DDebugState::ApplyRootPoints)
+	{
+		AddDebugTextureQuadToCollection(OutCollection, *InParams.Texture2DAsset, InParams.Scale);
+		AddDynamicMeshToCollection(DynamicMesh, OutCollection);
+		AddDebugRootPointsToCollection(OutCollection, InParams.RootPoints, InParams.Scale);
+		EndPointAttributeDefinition->CopyAttributeToCollection(DynamicMesh, OutCollection);
+		RootPointAttributeDefinition->CopyAttributeToCollection(DynamicMesh, OutCollection);
+ 		return EImportResult::Success;
+	}
+
+	auto& PerimiterAttribute = PerimiterAttributeDefinition.GetOrAttachAttribute(DynamicMesh);
+	TArray<FPVBranchHierarchyDescription> BranchHierarchyDescriptions;
+	ComputeBranchHierarchies(
+		DynamicMesh,
+		RootPointAttributeDefinition->GetOrAttachAttribute(DynamicMesh),
+		EndPointAttributeDefinition->GetOrAttachAttribute(DynamicMesh),
+		VertexScaleAttributeDefinition->GetOrAttachAttribute(DynamicMesh),
+		LengthFromRootAttributeDefinition->GetOrAttachAttribute(DynamicMesh),
+		NextBranchPointAttributeDefinition->GetOrAttachAttribute(DynamicMesh),
+		[&](int32 VertexID) { return !PerimiterAttribute.GetValue(VertexID); },
+		BranchHierarchyDescriptions
+	);
+	if (InParams.DebugState == EPVImportTexture2DDebugState::ComputeBranchHierarchies)
+	{
+		AddDynamicMeshToCollection(DynamicMesh, OutCollection);
+		AddBranchHierarchiesToCollection(BranchHierarchyDescriptions, OutCollection);
+		RootPointAttributeDefinition->CopyAttributeToCollection(DynamicMesh, OutCollection);
+		EndPointAttributeDefinition->CopyAttributeToCollection(DynamicMesh, OutCollection);
+		VertexScaleAttributeDefinition->CopyAttributeToCollection(DynamicMesh, OutCollection);
+		LengthFromRootAttributeDefinition->CopyAttributeToCollection(DynamicMesh, OutCollection);
+		NextBranchPointAttributeDefinition->CopyAttributeToCollection(DynamicMesh, OutCollection);
+
+		return EImportResult::Success;
+	}
+
+	FPVPlantDescription PlantDescription;
+	PlantDescription.BranchHierarchies = MoveTemp(BranchHierarchyDescriptions);
+
+	if (!PV::Import::Helper::GenerateGrowthDataFromPlantDescription(OutCollection, PlantDescription))
+	{
+		return EImportResult::FailedToGenerateGrowthData;
+	}
+
 	return EImportResult::Success;
-}
+}
\ No newline at end of file
diff --git a/Engine/Plugins/Experimental/ProceduralVegetationEditor/Source/ProceduralVegetationEditor/Private/Importer/PVImporter_Texture2D.h b/Engine/Plugins/Experimental/ProceduralVegetationEditor/Source/ProceduralVegetationEditor/Private/Importer/PVImporter_Texture2D.h
--- a/Engine/Plugins/Experimental/ProceduralVegetationEditor/Source/ProceduralVegetationEditor/Private/Importer/PVImporter_Texture2D.h
+++ b/Engine/Plugins/Experimental/ProceduralVegetationEditor/Source/ProceduralVegetationEditor/Private/Importer/PVImporter_Texture2D.h
@@ -2,7 +2,7 @@
 #pragma once
 
 #include "CoreMinimal.h"
-
+#include "PVImportTexture2DParams.h"
 #include "GeometryCollection/ManagedArrayCollection.h"
 
 class UTexture2D;
@@ -14,32 +14,9 @@ namespace PV::Texture2DImport
 		Success,
 		InvalidSourceTexture,
 		InvalidSampleSize,
-		InvalidTextureSize
-	};
-
-	struct FImportOptions
-	{
-		int32 SampleWidth = 128;
-		int32 SampleHeight = 128;
-		float MinBoundsArea = 0;
-		int32 SmoothingIterations = 1;
-		float SimplificationEpsilon = 0.001;
-		bool bInvertImage = false;
-		float WhiteLevel = 0.01;
-		
-		enum class EDebugState : uint8
-		{
-			None,
-			Texture,
-			SampleImage,
-			TraceImage,
-			RemoveSmallCurves,
-			SmoothCurves,
-			SimplifyCurves,
-			MAX
-		};
-		EDebugState DebugState = EDebugState::None;
+		InvalidTextureSize,
+		FailedToGenerateGrowthData
 	};
 
-	EImportResult ImportGrowthDataFromTexture2D(const TObjectPtr<UTexture2D>& InSourceTexture2D, FManagedArrayCollection& OutCollection, const FImportOptions& InOptions);
+	EImportResult ImportGrowthDataFromTexture2D(const FPVImportTexture2DParams& InParams, FManagedArrayCollection& OutCollection);
 };
diff --git a/Engine/Plugins/Experimental/ProceduralVegetationEditor/Source/ProceduralVegetationEditor/Private/Importer/PVPlantDescription.cpp b/Engine/Plugins/Experimental/ProceduralVegetationEditor/Source/ProceduralVegetationEditor/Private/Importer/PVPlantDescription.cpp
--- a/Engine/Plugins/Experimental/ProceduralVegetationEditor/Source/ProceduralVegetationEditor/Private/Importer/PVPlantDescription.cpp
+++ b/Engine/Plugins/Experimental/ProceduralVegetationEditor/Source/ProceduralVegetationEditor/Private/Importer/PVPlantDescription.cpp
@@ -198,12 +198,15 @@ FPrimitiveSceneProxy* UPVPlantDescriptionVisualizationComponent::CreateSceneProx
 				}
 			};
 			
-			if (PlantDescription.BranchHierarchy.ValidateHierarchy() != FPVBranchHierarchyDescription::EValidateHierarchyResult::Success)
+			for (const FPVBranchHierarchyDescription& BranchHierarchy : PlantDescription.BranchHierarchies)
 			{
-				return;
-			}
+				if (BranchHierarchy.ValidateHierarchy() != FPVBranchHierarchyDescription::EValidateHierarchyResult::Success)
+				{
+					return;
+				}
 
-			DrawBranchesRecursive(PDI, PlantDescription.BranchHierarchy, PlantDescription.BranchHierarchy.RootBranchIndex, 0);
+				DrawBranchesRecursive(PDI, BranchHierarchy, BranchHierarchy.RootBranchIndex, 0);
+			}
 		}
 		
 		void DrawFoliage(FPrimitiveDrawInterface* PDI) const
@@ -298,9 +301,12 @@ FBoxSphereBounds UPVPlantDescriptionVisualizationComponent::CalcBounds(const FTr
 
 	if (IsValid(PlantDescriptionDataAsset))
 	{
-		for (const FVector3f& Point : PlantDescriptionDataAsset->PlantDescription.BranchHierarchy.Points)
+		for (const FPVBranchHierarchyDescription& BranchHierarchy : PlantDescriptionDataAsset->PlantDescription.BranchHierarchies)
 		{
-			BBox += FVector(Point);
+			for (const FVector3f& Point : BranchHierarchy.Points)
+			{
+				BBox += FVector(Point);
+			}
 		}
 	}
 
diff --git a/Engine/Plugins/Experimental/ProceduralVegetationEditor/Source/ProceduralVegetationEditor/Private/Importer/PVPlantDescription.h b/Engine/Plugins/Experimental/ProceduralVegetationEditor/Source/ProceduralVegetationEditor/Private/Importer/PVPlantDescription.h
--- a/Engine/Plugins/Experimental/ProceduralVegetationEditor/Source/ProceduralVegetationEditor/Private/Importer/PVPlantDescription.h
+++ b/Engine/Plugins/Experimental/ProceduralVegetationEditor/Source/ProceduralVegetationEditor/Private/Importer/PVPlantDescription.h
@@ -117,12 +117,30 @@ struct FPVPlantDescription
 	GENERATED_BODY()
 
 	UPROPERTY(VisibleAnywhere, Category = "Plant Data")
-	FPVBranchHierarchyDescription BranchHierarchy;
+	TArray<FPVBranchHierarchyDescription> BranchHierarchies;
 
 	UPROPERTY(VisibleAnywhere, Category = "Plant Data")
 	FPVFoliageDescription Foliage;
 
-	int32 GetNumBranchPoints() const { return BranchHierarchy.GetNumBranchPoints(); }
+	int32 GetNumBranchPoints() const 
+	{ 
+		int32 NumPoints = 0;
+		for (const FPVBranchHierarchyDescription& BranchHierarchy : BranchHierarchies)
+		{
+			NumPoints += BranchHierarchy.Points.Num();
+		}
+		return NumPoints; 
+	}
+
+	int32 GetNumBranches() const
+	{
+		int32 NumBranches = 0;
+		for (const FPVBranchHierarchyDescription& BranchHierarchy : BranchHierarchies)
+		{
+			NumBranches += BranchHierarchy.Branches.Num();
+		}
+		return NumBranches;
+	}
 
 	int32 GetNumFoliagePoints() const { return Foliage.GetNumFoliagePoints(); }
 };
diff --git a/Engine/Plugins/Experimental/ProceduralVegetationEditor/Source/ProceduralVegetationEditor/Private/PVEditorInterfaceImpl.cpp b/Engine/Plugins/Experimental/ProceduralVegetationEditor/Source/ProceduralVegetationEditor/Private/PVEditorInterfaceImpl.cpp
--- a/Engine/Plugins/Experimental/ProceduralVegetationEditor/Source/ProceduralVegetationEditor/Private/PVEditorInterfaceImpl.cpp
+++ b/Engine/Plugins/Experimental/ProceduralVegetationEditor/Source/ProceduralVegetationEditor/Private/PVEditorInterfaceImpl.cpp
@@ -12,30 +12,12 @@
 
 bool PV::FPVEditorInterfaceImpl::ImportGrowthDataFromStaticMesh(const FPVImportStaticMeshParams& InParams, FManagedArrayCollection& OutCollection)
 {
-	PV::StaticMeshImport::FImportOptions ImportMeshOptions;
-	ImportMeshOptions.MaterialsToKeep = InParams.MaterialsToKeep;
-	
-	static_assert(static_cast<int32>(EPVImportStaticMeshDebugState::MAX) == static_cast<int32>(PV::StaticMeshImport::FImportOptions::EDebugState::MAX));
-	ImportMeshOptions.DebugState = static_cast<PV::StaticMeshImport::FImportOptions::EDebugState>(static_cast<int32>(InParams.DebugState));
-
-	const PV::StaticMeshImport::EImportResult ImportResult = PV::StaticMeshImport::ImportGrowthDataFromStaticMesh(InParams.StaticMeshAsset, OutCollection, ImportMeshOptions);
+	const PV::StaticMeshImport::EImportResult ImportResult = PV::StaticMeshImport::ImportGrowthDataFromStaticMesh(InParams, OutCollection);
 	return ImportResult == PV::StaticMeshImport::EImportResult::Success;
 }
 
 bool PV::FPVEditorInterfaceImpl::ImportGrowthDataFromTexture2D(const FPVImportTexture2DParams& InParams, FManagedArrayCollection& OutCollection)
 {
-	PV::Texture2DImport::FImportOptions ImportOptions;
-	ImportOptions.SampleWidth = InParams.SampleWidth;
-	ImportOptions.SampleHeight = InParams.SampleHeight;
-	ImportOptions.MinBoundsArea = InParams.MinBoundsArea;
-	ImportOptions.SmoothingIterations = InParams.SmoothingIterations;
-	ImportOptions.SimplificationEpsilon = InParams.SimplificationEpsilon;
-	ImportOptions.bInvertImage = InParams.bInvertImage;
-	ImportOptions.WhiteLevel = InParams.WhiteLevel;
-	
-	static_assert(static_cast<int32>(EPVImportTexture2DDebugState::MAX) == static_cast<int32>(PV::Texture2DImport::FImportOptions::EDebugState::MAX));
-	ImportOptions.DebugState = static_cast<PV::Texture2DImport::FImportOptions::EDebugState>(static_cast<int32>(InParams.DebugState));
-
-	const PV::Texture2DImport::EImportResult ImportResult = PV::Texture2DImport::ImportGrowthDataFromTexture2D(InParams.Texture2DAsset, OutCollection, ImportOptions);
+	const PV::Texture2DImport::EImportResult ImportResult = PV::Texture2DImport::ImportGrowthDataFromTexture2D(InParams, OutCollection);
 	return ImportResult == PV::Texture2DImport::EImportResult::Success;
 }
\ No newline at end of file
