Add debug visualization setting for drawing points at pivot point and using attribute for scale

#Synced-CL 49723000

diff --git a/Engine/Plugins/Experimental/ProceduralVegetationEditor/Source/ProceduralVegetation/Public/DataTypes/PVData.h b/Engine/Plugins/Experimental/ProceduralVegetationEditor/Source/ProceduralVegetation/Public/DataTypes/PVData.h
--- a/Engine/Plugins/Experimental/ProceduralVegetationEditor/Source/ProceduralVegetation/Public/DataTypes/PVData.h
+++ b/Engine/Plugins/Experimental/ProceduralVegetationEditor/Source/ProceduralVegetation/Public/DataTypes/PVData.h
@@ -66,16 +66,20 @@ struct FPVVisualizationSettings
 	UPROPERTY(EditAnywhere, Category = DebugSettings, meta = (EditCondition = "VisualizationMode==EPVDebugValueVisualizationMode::Point || bShowAnchorPoints"))
 	bool bDrawPointAsMesh = true;
 
+	// If true, then the AttributeToFilter will be used as point scale and the CustomPivotPositionAttributeName will be used for the point position. If false then AttributeToFilter is used for position and pivot attribute is ignored. 
+	UPROPERTY(EditAnywhere, Category = DebugSettings, meta = (EditCondition = "DebugType==EPVDebugType::Custom && VisualizationMode==EPVDebugValueVisualizationMode::Point"))
+	bool bUsePivotAsPosition = false;
+
 	UPROPERTY(EditAnywhere, Category = DebugSettings)
 	FName AttributeToFilter;
 
 	UPROPERTY(EditAnywhere, Category = DebugSettings, meta = (EditCondition = "DebugType==EPVDebugType::Custom"))
 	FName CustomGroupToFilter = TEXT("Custom");
 
-	UPROPERTY(EditAnywhere, Category = DebugSettings, meta = (EditCondition = "DebugType==EPVDebugType::Custom && VisualizationMode != EPVDebugValueVisualizationMode::Curve"))
+	UPROPERTY(EditAnywhere, Category = DebugSettings, meta = (EditCondition = "DebugType==EPVDebugType::Custom && VisualizationMode != EPVDebugValueVisualizationMode::Curve && (VisualizationMode != EPVDebugValueVisualizationMode::Point || bUsePivotAsPosition)"))
 	FName CustomPivotPositionAttributeName = TEXT("Position");
 
-	UPROPERTY(EditAnywhere, Category = DebugSettings, meta = (EditCondition = "DebugType==EPVDebugType::Custom && VisualizationMode != EPVDebugValueVisualizationMode::Curve"))
+	UPROPERTY(EditAnywhere, Category = DebugSettings, meta = (EditCondition = "DebugType==EPVDebugType::Custom && VisualizationMode != EPVDebugValueVisualizationMode::Curve && (VisualizationMode != EPVDebugValueVisualizationMode::Point || bUsePivotAsPosition)"))
 	FName CustomPivotScaleAttributeName = TEXT("Scale");
 
 	UPROPERTY(EditAnywhere, Category = DebugSettings)
diff --git a/Engine/Plugins/Experimental/ProceduralVegetationEditor/Source/ProceduralVegetationEditor/Private/Visualizations/DebugVisualization/PVDebugVisualizationBase.cpp b/Engine/Plugins/Experimental/ProceduralVegetationEditor/Source/ProceduralVegetationEditor/Private/Visualizations/DebugVisualization/PVDebugVisualizationBase.cpp
--- a/Engine/Plugins/Experimental/ProceduralVegetationEditor/Source/ProceduralVegetationEditor/Private/Visualizations/DebugVisualization/PVDebugVisualizationBase.cpp
+++ b/Engine/Plugins/Experimental/ProceduralVegetationEditor/Source/ProceduralVegetationEditor/Private/Visualizations/DebugVisualization/PVDebugVisualizationBase.cpp
@@ -533,19 +533,41 @@ void FPVDebugVisualizationBase::DrawAttributes(FVisualizerDrawContext& InContext
 		{
 			using namespace PVE::AttributeToVector3;
 
-			const TArray<FVectorValue> AttributeAsVector = ConvertAttributeToVector3(Collection, AttributeName, GroupName);
-			if (AttributeAsVector.Num() != NumElements)
+			if (VisualizationSettings.bUsePivotAsPosition)
 			{
-				break;
-			}
+				TArray<double> AttributeAsScalar = PVE::AttributeToScalar::ConvertAttributeToScalar(Collection, AttributeName, GroupName);
+				if (AttributeAsScalar.Num() != NumElements)
+				{
+					break;
+				}
 
-			for (int32 i = 0; i < NumElements; i++)
+				for (int32 i = 0; i < NumElements; i++)
+				{
+					FVector3f OutPos;
+					float OutScale;
+					GetPivot(Collection, i, OutPos, OutScale);
+
+					const FColor DrawColor = VisualizationSettings.bRandomizeColors ? FColor::MakeRandomColor() : VisualizationSettings.Color;
+					AddDrawAsPoint(FVector(OutPos), AttributeAsScalar[i] * VisualizationSettings.GizmoScale * OutScale);
+				}
+			}
+			else 
 			{
-				for (int32 j = 0; j < AttributeAsVector[i].Num(); ++j)
+				const TArray<FVectorValue> AttributeAsVector = ConvertAttributeToVector3(Collection, AttributeName, GroupName);
+				if (AttributeAsVector.Num() != NumElements)
 				{
-					AddDrawAsPoint(AttributeAsVector[i][j], VisualizationSettings.GizmoScale);
+					break;
 				}
+
+				for (int32 i = 0; i < NumElements; i++)
+				{
+					for (int32 j = 0; j < AttributeAsVector[i].Num(); ++j)
+					{
+						AddDrawAsPoint(AttributeAsVector[i][j], VisualizationSettings.GizmoScale);
+					}
+				}	
 			}
+
 			break;
 		}
 		case EPVDebugValueVisualizationMode::Sphere:
@@ -565,7 +587,7 @@ void FPVDebugVisualizationBase::DrawAttributes(FVisualizerDrawContext& InContext
 				GetPivot(Collection, i, OutPos, OutScale);
 
 				const FColor DrawColor = VisualizationSettings.bRandomizeColors ? FColor::MakeRandomColor() : VisualizationSettings.Color;
-				DrawAsSphere(LineBatchComponent, FVector(OutPos), AttributeAsScalar[i] * VisualizationSettings.GizmoScale, DrawColor);
+				DrawAsSphere(LineBatchComponent, FVector(OutPos), AttributeAsScalar[i] * VisualizationSettings.GizmoScale * OutScale, DrawColor);
 			}
 			break;
 		}
@@ -593,6 +615,9 @@ void FPVDebugVisualizationBase::DrawAttributes(FVisualizerDrawContext& InContext
 			}
 			break;
 		}
+		case EPVDebugValueVisualizationMode::Vector:
+			ensureMsgf(false, TEXT("Unexpected visualization mode"));
+			return;
 	}
 
 	if (PointsToDraw.Num() > 0)
@@ -619,7 +644,10 @@ void FPVDebugVisualizationBase::DrawAttributes(FVisualizerDrawContext& InContext
 
 void FPVDebugVisualizationBase::AddDrawAsPoint(const FVector& InPos, float PointScale)
 {
-	PointsToDraw.Add({ InPos, PointScale });
+	if (!FMath::IsNearlyZero(PointScale))
+	{
+		PointsToDraw.Add({ InPos, PointScale });
+	}
 }
 
 void FPVDebugVisualizationBase::DrawAsText(FPCGSceneSetupParams& InOutParams, const FText& InTextToDraw, const FVector3f& InPos, float InScale, float TextSize, const FLinearColor& InColor, const FVector3f TextOffset)
