[PVE] Object Interaction node
- Update to Object Avoidance with trim/clip functionality
- Extended Object Avoidance to trim the branches when a ray hit is detected with the collider

#Synced-CL 49647823

diff --git a/Engine/Plugins/Experimental/ProceduralVegetationEditor/Source/ProceduralVegetation/Private/Implementations/PVCarve.h b/Engine/Plugins/Experimental/ProceduralVegetationEditor/Source/ProceduralVegetation/Private/Implementations/PVCarve.h
--- a/Engine/Plugins/Experimental/ProceduralVegetationEditor/Source/ProceduralVegetation/Private/Implementations/PVCarve.h
+++ b/Engine/Plugins/Experimental/ProceduralVegetationEditor/Source/ProceduralVegetation/Private/Implementations/PVCarve.h
@@ -36,7 +36,6 @@ struct FPVCarve
 	static void ApplyCarve(FManagedArrayCollection& OutCollection, const FManagedArrayCollection& SourceCollection, const ECarveBasis CarveBasis,
 	                       const float Carve);
 
-private:
 	static void UpdatePointScales(PV::Facades::FPointFacade& PointFacadeOut, const PV::Facades::FPointFacade& PointFacadeSource,
 	                              const TArray<int>& BranchPoints, const float LastPointScale, const int32 LastPointIndex, const float CarveRatio,
 	                              const int32 EndIndex, const float FirstPointTargetPScale);
diff --git a/Engine/Plugins/Experimental/ProceduralVegetationEditor/Source/ProceduralVegetation/Private/Implementations/PVObjectInteraction.cpp b/Engine/Plugins/Experimental/ProceduralVegetationEditor/Source/ProceduralVegetation/Private/Implementations/PVObjectInteraction.cpp
new file mode 100644
--- /dev/null
+++ b/Engine/Plugins/Experimental/ProceduralVegetationEditor/Source/ProceduralVegetation/Private/Implementations/PVObjectInteraction.cpp
@@ -0,0 +1,399 @@
+// Copyright Epic Games, Inc. All Rights Reserved.
+#include "PVObjectInteraction.h"
+
+#include "MeshDescriptionAdapter.h"
+
+#include "DataTypes/PVGrowthData.h"
+
+#include "DynamicMesh/DynamicMeshAABBTree3.h"
+
+#include "Engine/StaticMesh.h"
+
+#include "Facades/PVBranchFacade.h"
+#include "Facades/PVBudVectorsFacade.h"
+#include "Facades/PVPointFacade.h"
+
+#include "Helpers/PVUtilities.h"
+
+#include "Implementations/PVCarve.h"
+
+FVector FPVObjectInteraction::GetTriangleNormal(FMeshDescriptionTriangleMeshAdapter& Mesh, const int32 TID, TOptional<FVector> BaryCoords)
+{
+	if (Mesh.HasNormals() && BaryCoords.IsSet())
+	{
+		FVector3f NormalA, NormalB, NormalC;
+		Mesh.GetTriNormals(TID, NormalA, NormalB, NormalC);
+		return FVector(BaryCoords->X * NormalA + BaryCoords->Y * NormalB + BaryCoords->Z * NormalC).GetUnsafeNormal();
+	}
+	else
+	{
+		FVector VertexA, VertexB, VertexC;
+		Mesh.GetTriVertices(TID, VertexA, VertexB, VertexC);
+		return UE::Geometry::VectorUtil::Normal(VertexA, VertexB, VertexC);
+	}
+}
+
+void FPVObjectInteraction::ObjectInteraction(const FManagedArrayCollection& InCollection, FPVColliderParams Collider,
+                                             FManagedArrayCollection& OutCollection)
+{
+	OutCollection = InCollection;
+
+	const FTransform& Transform = Collider.Transform;
+
+	const UStaticMesh* ColliderMesh = Collider.Mesh.LoadSynchronous();
+	if (!ColliderMesh)
+	{
+		return;
+	}
+
+#if WITH_EDITORONLY_DATA
+	const FMeshDescription* MeshDescription = ColliderMesh->GetMeshDescription(0);
+	FMeshDescriptionTriangleMeshAdapter Mesh(MeshDescription);
+	UE::Geometry::TMeshAABBTree3 Octree(&Mesh);
+
+	PV::Facades::FPointFacade PointFacade(OutCollection);
+	PV::Facades::FBranchFacade BranchFacade(OutCollection);
+	PV::Facades::FFoliageFacade FoliageFacade(OutCollection);
+	PV::Facades::FBudVectorsFacade BudVectorFacade = PV::Facades::FBudVectorsFacade(OutCollection);
+
+	TArray<bool> PointsToRemove;
+	PointsToRemove.Init(false, PointFacade.GetElementCount());
+	TArray<bool> BranchesToRemove;
+	BranchesToRemove.Init(false, BranchFacade.GetElementCount());
+	TArray<bool> FoliageInstancesToRemove;
+	FoliageInstancesToRemove.Init(false, FoliageFacade.NumFoliageEntries());
+
+	TManagedArray<FVector3f>& PointPositions = PointFacade.ModifyPositions();
+	TManagedArray<TArray<FVector3f>>& BudDirections = BudVectorFacade.ModifyBudDirections();
+	TManagedArray<FVector3f>& FoliagePositions = FoliageFacade.ModifyPivotPositions();
+
+	const auto ApplyOffsetToBranchSuccessors = [&](const int32 BranchIndex, const FVector3f& Offset)
+		{
+			for (const int32 BranchPoint : BranchFacade.GetPoints(BranchIndex))
+			{
+				PointPositions[BranchPoint] += Offset;
+			}
+			for (const int32 FoliageEntryId : FoliageFacade.GetFoliageEntryIdsForBranch(BranchIndex))
+			{
+				FoliagePositions[FoliageEntryId] += Offset;
+			}
+
+			const TArray<int32> BranchChildren = BranchFacade.GetChildren(BranchIndex);
+			for (int32 BranchChild : BranchChildren)
+			{
+				const int32 ChildIndex = BranchFacade.GetBranchNumbers().Find(BranchChild);
+				if (ChildIndex != INDEX_NONE)
+				{
+					for (const int32 ChildBranchPoint : BranchFacade.GetPoints(ChildIndex))
+					{
+						PointPositions[ChildBranchPoint] += Offset;
+					}
+					for (const int32 ChildFoliageEntryId : FoliageFacade.GetFoliageEntryIdsForBranch(ChildIndex))
+					{
+						FoliagePositions[ChildFoliageEntryId] += Offset;
+					}
+				}
+			}
+		};
+
+	const auto RemoveBranchSuccessors = [&](const int32 BranchIndex)
+		{
+			BranchesToRemove[BranchIndex] = true;
+			for (const int32 BranchPoint : BranchFacade.GetPoints(BranchIndex))
+			{
+				PointsToRemove[BranchPoint] = true;
+			}
+			for (const int32 FoliageEntryId : FoliageFacade.GetFoliageEntryIdsForBranch(BranchIndex))
+			{
+				FoliageInstancesToRemove[FoliageEntryId] = true;
+			}
+
+			const TArray<int32> BranchChildren = BranchFacade.GetChildren(BranchIndex);
+			for (int32 BranchChild : BranchChildren)
+			{
+				const int32 ChildIndex = BranchFacade.GetBranchNumbers().Find(BranchChild);
+				if (ChildIndex != INDEX_NONE)
+				{
+					BranchesToRemove[ChildIndex] = true;
+					for (const int32 ChildBranchPoint : BranchFacade.GetPoints(ChildIndex))
+					{
+						PointsToRemove[ChildBranchPoint] = true;
+					}
+					for (const int32 ChildFoliageEntryId : FoliageFacade.GetFoliageEntryIdsForBranch(ChildIndex))
+					{
+						FoliageInstancesToRemove[ChildFoliageEntryId] = true;
+					}
+				}
+			}
+		};
+
+	TArray<int32> BranchIndices;
+	BranchIndices.Reserve(BranchFacade.GetElementCount());
+	BranchFacade.GetSortedBranchIndicesByHierarchy(BranchIndices);
+
+	for (int32 BranchIndex : BranchIndices)
+	{
+		PVE_OUTER_LOOP_DEBUG_CHECK(break);
+
+		if (BranchesToRemove[BranchIndex] && Collider.CollisionType == EPVCollisionType::TRIM_OUTSIDE)
+		{
+			continue;
+		}
+
+		const float MaxPointScale = *Algo::MaxElement(PointFacade.GetPointScales());
+
+		const TArray<int32> BranchPoints = BranchFacade.GetPoints(BranchIndex);
+		for (int32 BranchPointIndex = 1; BranchPointIndex < BranchPoints.Num(); ++BranchPointIndex)
+		{
+			const int32 PointIndex = BranchPoints[BranchPointIndex - 1];
+			const int32 NextPointIndex = BranchPoints[BranchPointIndex];
+
+			if (PointsToRemove[PointIndex] && Collider.CollisionType == EPVCollisionType::TRIM_OUTSIDE)
+			{
+				continue;
+			}
+
+			const FVector PointPosition = static_cast<FVector>(PointPositions[PointIndex]);
+			const FVector NextPointPosition = static_cast<FVector>(PointPositions[NextPointIndex]);
+			const FVector SegmentVector = NextPointPosition - PointPosition;
+
+			// Get the previous direction
+			FVector PreviousDirection;
+			if (BranchPoints.IsValidIndex(BranchPointIndex - 2))
+			{
+				PreviousDirection = (PointPosition - BranchPoints[BranchPointIndex - 2]).
+					GetSafeNormal(UE_SMALL_NUMBER, FVector::UpVector);
+			}
+			else
+			{
+				PreviousDirection = (NextPointPosition - PointPosition).GetSafeNormal(UE_SMALL_NUMBER, FVector::UpVector);
+			}
+
+			const float PointScale = Collider.CollisionType == EPVCollisionType::AVOID
+				? MaxPointScale
+				: 0.0f;
+
+			FRay3d Ray(PointPosition, SegmentVector);
+			FRay3d InverseRay(Transform.InverseTransformPosition(Ray.Origin), Transform.InverseTransformVector(Ray.Direction));
+
+			const float PointLFR = PointFacade.GetLengthFromRoot(PointIndex);
+			const float NextPointLFR = PointFacade.GetLengthFromRoot(NextPointIndex);
+
+			const double MaxLength = SegmentVector.Length() + PointScale;
+			const FVector MaxLengthVector = Ray.Direction * MaxLength;
+			const double InverseMaxLength = Transform.InverseTransformVector(MaxLengthVector).Length();
+
+			int32 HitTriangle = INDEX_NONE;
+			FVector BaryCoords;
+			if (double RayT; Octree.FindNearestHitTriangle(InverseRay, RayT, HitTriangle, BaryCoords, {InverseMaxLength}))
+			{
+				const FVector HitNormal = Transform.TransformVectorNoScale(GetTriangleNormal(Mesh, HitTriangle, BaryCoords)).
+				                                    GetSafeNormal(UE_SMALL_NUMBER, FVector::UpVector);
+
+				if (Collider.CollisionType == EPVCollisionType::AVOID)
+				{
+					const FVector HitPoint = Transform.TransformPosition(InverseRay.PointAt(RayT));
+					const FVector InvertedNormal = -HitNormal;
+
+					const double DistanceToHit = FVector::Distance(PointPosition, HitPoint);
+					const double DistanceAfterHit = FVector::Distance(HitPoint, Ray.PointAt(MaxLength));
+					const FVector BiNormal = FVector::CrossProduct(Ray.Direction, InvertedNormal);
+					FVector Tangent = FVector::CrossProduct(InvertedNormal, BiNormal);
+					const FVector PreviousBiNormal = FVector::CrossProduct(PreviousDirection, InvertedNormal);
+					const FVector PreviousTangent = FVector::CrossProduct(InvertedNormal, PreviousBiNormal);
+					// Get the mean tangent of this and previous tangent
+					Tangent = (PreviousTangent + Tangent).GetSafeNormal(UE_SMALL_NUMBER, FVector::UpVector);
+					const FVector AvoidanceOffset = Tangent * DistanceToHit * (FVector::DotProduct(Ray.Direction, InvertedNormal) + 1);
+					const FVector AvoidancePosition = HitPoint + HitNormal * (PointScale + DistanceAfterHit) + AvoidanceOffset;
+					const FVector AvoidanceDirection = (AvoidancePosition - PointPosition).GetSafeNormal(UE_SMALL_NUMBER,
+						FVector::DownVector);
+
+					const FVector OffsetPosition = PointPosition + AvoidanceDirection * FVector::Distance(PointPosition,
+						NextPointPosition);
+					const FVector3f Offset = FVector3f(OffsetPosition - NextPointPosition);
+
+					PVE_LOOP_DEBUG_PARAMS_START()
+						PVE_LOOP_DEBUG_DIRECTION_PARAM("Ray", Ray.Origin, MaxLengthVector);
+						PVE_LOOP_DEBUG_DIRECTION_PARAM("Hit Normal", HitPoint, HitNormal);
+						PVE_LOOP_DEBUG_DIRECTION_PARAM("Avoidance Offset", HitPoint, AvoidanceOffset);
+						PVE_LOOP_DEBUG_DIRECTION_PARAM("Avoidance Direction", NextPointPosition, AvoidanceDirection);
+						PVE_LOOP_DEBUG_VECTOR_PARAM("Avoidance Vector", PointPosition, AvoidancePosition);
+						PVE_LOOP_DEBUG_DIRECTION_PARAM("Offset", NextPointPosition, Offset);
+					PVE_LOOP_DEBUG_PARAMS_END()
+
+					PVE_LOOP_DEBUG_STEP(break);
+
+					constexpr int32 Smoothness = 10;
+					constexpr float SmoothnessFactor = 1.5f;
+
+					const int32 StartIndex = FMath::Max(1, BranchPointIndex - Smoothness);
+					FVector3f PreviousSiblingPosition;
+					for (int32 SiblingPointIndex = StartIndex; SiblingPointIndex < BranchPointIndex; ++SiblingPointIndex)
+					{
+						const int32 SiblingPoint = BranchPoints[SiblingPointIndex];
+						const int32 NextSiblingPoint = BranchPoints[SiblingPointIndex + 1];
+
+						const float SiblingPointLFR = PointFacade.GetLengthFromRoot(SiblingPoint);
+						const float NextSiblingPointLFR = PointFacade.GetLengthFromRoot(NextSiblingPoint);
+
+						const FVector3f PointOffset = Offset / FMath::Pow(SmoothnessFactor, BranchPointIndex - SiblingPointIndex);
+
+						if (SiblingPointIndex != StartIndex)
+						{
+							const FVector3f SiblingSegmentVector = PointPositions[SiblingPoint] - PreviousSiblingPosition;
+							const FQuat4f Rotation = FQuat4f::FindBetween(SiblingSegmentVector, SiblingSegmentVector + PointOffset);
+							for (const int32 FoliageEntryId : FoliageFacade.GetFoliageEntryIdsForBranch(BranchIndex))
+							{
+								const float FoliageLFR = FoliageFacade.GetLengthFromRoot(FoliageEntryId);
+								if (FoliageLFR > SiblingPointLFR && FoliageLFR < NextSiblingPointLFR)
+								{
+									FoliagePositions[FoliageEntryId] += PointOffset;
+									FoliageFacade.SetUpVector(
+										FoliageEntryId,
+										Rotation.RotateVector(FoliageFacade.GetUpVector(FoliageEntryId))
+									);
+									FoliageFacade.SetNormalVector(
+										FoliageEntryId,
+										Rotation.RotateVector(FoliageFacade.GetNormalVector(FoliageEntryId))
+									);
+								}
+							}
+							for (FVector3f& BudDirection : BudDirections[SiblingPoint])
+							{
+								BudDirection = Rotation.RotateVector(BudDirection);
+							}
+						}
+
+						PreviousSiblingPosition = PointPositions[SiblingPoint];
+						PointPositions[SiblingPoint] += PointOffset;
+
+						const TArray<int32> BranchChildren = BranchFacade.GetImmediateChildren(BranchIndex);
+						for (int32 ChildBranch : BranchChildren)
+						{
+							const int32 ChildBranchIndex = BranchFacade.GetBranchNumbers().Find(ChildBranch);
+							if (ChildBranchIndex != INDEX_NONE)
+							{
+								const int32 RootPoint = BranchFacade.GetRootPoint(ChildBranchIndex);
+								const float ChildRootPointLFR = PointFacade.GetLengthFromRoot(RootPoint);
+								if (ChildRootPointLFR >= SiblingPointLFR && ChildRootPointLFR < NextSiblingPointLFR)
+								{
+									ApplyOffsetToBranchSuccessors(ChildBranchIndex, PointOffset);
+								}
+							}
+						}
+					}
+
+					for (int32 SiblingPointIndex = BranchPointIndex; SiblingPointIndex < BranchPoints.Num(); ++SiblingPointIndex)
+					{
+						PointPositions[BranchPoints[SiblingPointIndex]] += Offset;
+					}
+
+					const FQuat4f Rotation = FQuat4f::FindBetween(FVector3f(SegmentVector), FVector3f(SegmentVector) + Offset);
+					for (const int32 FoliageEntryId : FoliageFacade.GetFoliageEntryIdsForBranch(BranchIndex))
+					{
+						const float FoliageLFR = FoliageFacade.GetLengthFromRoot(FoliageEntryId);
+						if (FoliageLFR > PointLFR)
+						{
+							FoliagePositions[FoliageEntryId] += Offset;
+							if (FoliageLFR < NextPointLFR)
+							{
+								FoliageFacade.SetUpVector(
+									FoliageEntryId,
+									Rotation.RotateVector(FoliageFacade.GetUpVector(FoliageEntryId))
+								);
+								FoliageFacade.SetNormalVector(
+									FoliageEntryId,
+									Rotation.RotateVector(FoliageFacade.GetNormalVector(FoliageEntryId))
+								);
+							}
+						}
+					}
+					for (FVector3f& BudDirection : BudDirections[NextPointIndex])
+					{
+						BudDirection = Rotation.RotateVector(BudDirection);
+					}
+
+					const TArray<int32> BranchChildren = BranchFacade.GetImmediateChildren(BranchIndex);
+					for (int32 ChildBranch : BranchChildren)
+					{
+						const int32 ChildBranchIndex = BranchFacade.GetBranchNumbers().Find(ChildBranch);
+						if (ChildBranchIndex != INDEX_NONE)
+						{
+							const int32 RootPoint = BranchFacade.GetRootPoint(ChildBranchIndex);
+							if (PointFacade.GetLengthFromRoot(RootPoint) > PointLFR)
+							{
+								ApplyOffsetToBranchSuccessors(ChildBranchIndex, Offset);
+							}
+						}
+					}
+				}
+				else if (Collider.CollisionType == EPVCollisionType::TRIM_OUTSIDE)
+				{
+					// In case the edge hit from inside the collider (Normal and Ray are opposite direction)
+					// Trim the remaining branch, all child branches and all foliage instances in case of hit
+					if (FVector::DotProduct(Ray.Direction, HitNormal) > 0)
+					{
+						for (int32 SiblingPointIndex = BranchPointIndex; SiblingPointIndex < BranchPoints.Num(); ++SiblingPointIndex)
+						{
+							PointsToRemove[BranchPoints[SiblingPointIndex]] = true;
+						}
+						for (const int32 FoliageEntryId : FoliageFacade.GetFoliageEntryIdsForBranch(BranchIndex))
+						{
+							const float FoliageLFR = FoliageFacade.GetLengthFromRoot(FoliageEntryId);
+							if (FoliageLFR > PointLFR)
+							{
+								FoliageInstancesToRemove[FoliageEntryId] = true;
+							}
+						}
+
+						const float OldMax = PointFacade.GetPointScale(BranchPoints[0]);
+						const float OldMin = PointFacade.GetPointScale(BranchPoints[BranchPointIndex - 1]);
+						const float NewMax = OldMax * (static_cast<float>(BranchPointIndex) / BranchPoints.Num());
+						const float NewMin = OldMin * (NewMax / OldMax);
+
+						// Rescale the branches based on the cut point
+						for (int32 Index = 0; Index <= BranchPointIndex; ++Index)
+						{
+							const float UpdatedScale = FMath::GetMappedRangeValueClamped(
+								FVector2f(OldMin, OldMax),
+								FVector2f(NewMin, NewMax),
+								PointFacade.GetPointScale(BranchPoints[Index]));
+							PointFacade.SetPointScale(BranchPoints[Index], UpdatedScale);
+						}
+
+						// Remove all branch children that are after the cut point
+						const TArray<int32> BranchChildren = BranchFacade.GetImmediateChildren(BranchIndex);
+						for (int32 ChildBranch : BranchChildren)
+						{
+							const int32 ChildBranchIndex = BranchFacade.GetBranchNumbers().Find(ChildBranch);
+							if (ChildBranchIndex != INDEX_NONE)
+							{
+								const int32 RootPoint = BranchFacade.GetRootPoint(ChildBranchIndex);
+								const float ChildRootPointLFR = PointFacade.GetLengthFromRoot(RootPoint);
+								if (ChildRootPointLFR >= PointLFR)
+								{
+									RemoveBranchSuccessors(ChildBranchIndex);
+								}
+							}
+						}
+
+						break;
+					}
+				}
+			}
+		}
+	}
+
+	// Remove all the ToRemove entries using the same algorithm being used in Carve algorithm
+	if (Collider.CollisionType == EPVCollisionType::TRIM_OUTSIDE)
+	{
+		FPVCarve::RemoveEntriesAndRecomputeAttributes(
+			OutCollection,
+			InCollection,
+			PointsToRemove,
+			BranchesToRemove,
+			FoliageInstancesToRemove
+		);
+	}
+#endif
+}
diff --git a/Engine/Plugins/Experimental/ProceduralVegetationEditor/Source/ProceduralVegetation/Private/Implementations/PVObjectInteraction.h b/Engine/Plugins/Experimental/ProceduralVegetationEditor/Source/ProceduralVegetation/Private/Implementations/PVObjectInteraction.h
new file mode 100644
--- /dev/null
+++ b/Engine/Plugins/Experimental/ProceduralVegetationEditor/Source/ProceduralVegetation/Private/Implementations/PVObjectInteraction.h
@@ -0,0 +1,18 @@
+// Copyright Epic Games, Inc. All Rights Reserved.
+
+#pragma once
+
+#include "Nodes/PVObjectInteractionSettings.h"
+
+struct FMeshDescriptionTriangleMeshAdapter;
+
+struct FPVObjectInteraction
+{
+	static FVector GetTriangleNormal(
+		FMeshDescriptionTriangleMeshAdapter& Mesh,
+		const int32 TID,
+		TOptional<FVector> BaryCoords = TOptional<FVector>()
+	);
+	
+	static void ObjectInteraction(const FManagedArrayCollection& InCollection, FPVColliderParams Collider, FManagedArrayCollection& OutCollection);
+};
diff --git a/Engine/Plugins/Experimental/ProceduralVegetationEditor/Source/ProceduralVegetation/Private/Nodes/PVObjectAvoidance.cpp b/Engine/Plugins/Experimental/ProceduralVegetationEditor/Source/ProceduralVegetation/Private/Nodes/PVObjectAvoidance.cpp
deleted file mode 100644
diff --git a/Engine/Plugins/Experimental/ProceduralVegetationEditor/Source/ProceduralVegetation/Private/Nodes/PVObjectInteractionSettings.cpp b/Engine/Plugins/Experimental/ProceduralVegetationEditor/Source/ProceduralVegetation/Private/Nodes/PVObjectInteractionSettings.cpp
new file mode 100644
--- /dev/null
+++ b/Engine/Plugins/Experimental/ProceduralVegetationEditor/Source/ProceduralVegetation/Private/Nodes/PVObjectInteractionSettings.cpp
@@ -0,0 +1,62 @@
+// Copyright Epic Games, Inc. All Rights Reserved.
+
+#include "Nodes/PVObjectInteractionSettings.h"
+
+#include "DataTypes/PVGrowthData.h"
+
+#include "Helpers/PVUtilities.h"
+
+#include "Implementations/PVObjectInteraction.h"
+
+FPCGDataTypeIdentifier UPVObjectInteractionSettings::GetInputPinTypeIdentifier() const
+{
+	return FPCGDataTypeIdentifier{FPVDataTypeInfoGrowth::AsId()};
+}
+
+FPCGDataTypeIdentifier UPVObjectInteractionSettings::GetOutputPinTypeIdentifier() const
+{
+	return FPCGDataTypeIdentifier{FPVDataTypeInfoGrowth::AsId()};
+}
+
+FPCGElementPtr UPVObjectInteractionSettings::CreateElement() const
+{
+	return MakeShared<FPVObjectInteractionElement>();
+}
+
+bool FPVObjectInteractionElement::ExecuteInternal(FPCGContext* InContext) const
+{
+	TRACE_CPUPROFILER_EVENT_SCOPE(FPVObjectInteractionElement::ExecuteInternal);
+
+	check(InContext);
+
+	const UPVObjectInteractionSettings* InputSettings = InContext->GetInputSettings<UPVObjectInteractionSettings>();
+	check(InputSettings);
+
+	const TArray<FPCGTaggedData>& Inputs = InContext->InputData.GetInputsByPin(PCGPinConstants::DefaultInputLabel);
+	if (!Inputs.IsEmpty())
+	{
+		if (const UPVGrowthData* InputData = Cast<UPVGrowthData>(Inputs[0].Data))
+		{
+			FManagedArrayCollection OutCollection;
+			UPVData* OutManagedArrayCollectionData = FPCGContext::NewObject_AnyThread<UPVGrowthData>(InContext);
+
+#if WITH_EDITORONLY_DATA
+			PVE_LOOP_DEBUG_INIT(InputSettings, LoopDebugStepper, OutManagedArrayCollectionData);
+#endif
+
+			for (const auto& Collider : InputSettings->ObjectInteractionSettings.Colliders)
+			{
+				PVE_OUTER_LOOP_DEBUG_CHECK(break);
+
+				FPVObjectInteraction::ObjectInteraction(InputData->GetCollection(), Collider, OutCollection);
+			}
+
+			OutManagedArrayCollectionData->Initialize(MoveTemp(OutCollection));
+			InContext->OutputData.TaggedData.Emplace(OutManagedArrayCollectionData);
+
+			PVE_LOOP_DEBUG_END();
+		}
+	}
+
+	return true;
+}
diff --git a/Engine/Plugins/Experimental/ProceduralVegetationEditor/Source/ProceduralVegetation/Public/Nodes/PVObjectAvoidance.h b/Engine/Plugins/Experimental/ProceduralVegetationEditor/Source/ProceduralVegetation/Public/Nodes/PVObjectAvoidance.h
deleted file mode 100644
diff --git a/Engine/Plugins/Experimental/ProceduralVegetationEditor/Source/ProceduralVegetation/Public/Nodes/PVObjectInteractionSettings.h b/Engine/Plugins/Experimental/ProceduralVegetationEditor/Source/ProceduralVegetation/Public/Nodes/PVObjectInteractionSettings.h
new file mode 100644
--- /dev/null
+++ b/Engine/Plugins/Experimental/ProceduralVegetationEditor/Source/ProceduralVegetation/Public/Nodes/PVObjectInteractionSettings.h
@@ -0,0 +1,88 @@
+// Copyright Epic Games, Inc. All Rights Reserved.
+
+#pragma once
+
+#include "Nodes/PVBaseSettings.h"
+
+#include "PVObjectInteractionSettings.generated.h"
+
+UENUM()
+enum class EPVCollisionType
+{
+	AVOID UMETA(DisplayName = "Avoid"),
+	TRIM_OUTSIDE UMETA(DisplayName = "Trim Outside")
+};
+
+USTRUCT()
+struct FPVColliderParams
+{
+	GENERATED_BODY()
+
+	UPROPERTY(EditAnywhere, Category = "Collider")
+	TSoftObjectPtr<UStaticMesh> Mesh;
+
+	UPROPERTY(EditAnywhere, Category = "Collider")
+	FTransform Transform;
+
+	UPROPERTY(EditAnywhere, Category = "Collider")
+	EPVCollisionType CollisionType;
+};
+
+USTRUCT()
+struct FPVObjectInteractionParams
+{
+	GENERATED_BODY()
+
+	UPROPERTY(EditAnywhere, Category = "Colliders")
+	TArray<FPVColliderParams> Colliders;
+};
+
+UCLASS(BlueprintType, ClassGroup = (Procedural))
+class UPVObjectInteractionSettings : public UPVBaseSettings
+{
+	GENERATED_BODY()
+
+public:
+#if WITH_EDITOR
+	virtual FName GetDefaultNodeName() const override
+	{
+		return FName(TEXT("Object Interaction"));
+	}
+
+	virtual FText GetNodeTooltipText() const override
+	{
+		return FText::FromString("Object Interaction");
+	}
+
+	virtual FLinearColor GetNodeTitleColor() const override
+	{
+		return FLinearColor::Yellow;
+	}
+
+	virtual TArray<EPVRenderType> GetDefaultRenderType() const override
+	{
+		return TArray{EPVRenderType::PointData};
+	}
+#endif
+
+protected:
+	virtual FPCGDataTypeIdentifier GetInputPinTypeIdentifier() const override;
+	virtual FPCGDataTypeIdentifier GetOutputPinTypeIdentifier() const override;
+
+	virtual FPCGElementPtr CreateElement() const override;
+
+public:
+	UPROPERTY(EditAnywhere, Category = "Object Interaction", meta = (ShowOnlyInnerProperties))
+	FPVObjectInteractionParams ObjectInteractionSettings;
+
+#if WITH_EDITORONLY_DATA
+	UPROPERTY(EditAnywhere, Category = "Loop Debug", meta=(EditCondition = "bDebug", EditConditionHides))
+	FLoopDebugStepper LoopDebugStepper;
+#endif
+};
+
+class FPVObjectInteractionElement : public FPVBaseElement
+{
+protected:
+	virtual bool ExecuteInternal(FPCGContext* InContext) const override;
+};
diff --git a/Engine/Plugins/Experimental/ProceduralVegetationEditor/Source/ProceduralVegetationEditor/Private/Visualizations/PVDataVisualization.cpp b/Engine/Plugins/Experimental/ProceduralVegetationEditor/Source/ProceduralVegetationEditor/Private/Visualizations/PVDataVisualization.cpp
--- a/Engine/Plugins/Experimental/ProceduralVegetationEditor/Source/ProceduralVegetationEditor/Private/Visualizations/PVDataVisualization.cpp
+++ b/Engine/Plugins/Experimental/ProceduralVegetationEditor/Source/ProceduralVegetationEditor/Private/Visualizations/PVDataVisualization.cpp
@@ -34,7 +34,7 @@
 #include "Materials/Material.h"
  
 #include "Nodes/PVBaseSettings.h"
-#include "Nodes/PVObjectAvoidance.h"
+#include "Nodes/PVObjectInteractionSettings.h"
  
 #include "Rendering/SkeletalMeshRenderData.h"
  
@@ -114,16 +114,16 @@ FPCGSetupSceneFunc FPVDataVisualization::GetViewportSetupFunc(const UPCGSettings
 			}
 		}
  
-		if (const UPVObjectAvoidanceSettings* ObjectAvoidanceSettings = Cast<UPVObjectAvoidanceSettings>(WeakSettings))
+		if (const UPVObjectInteractionSettings* ObjectAvoidanceSettings = Cast<UPVObjectInteractionSettings>(WeakSettings))
 		{
-			for (const auto& [Mesh, Transform] : ObjectAvoidanceSettings->ObjectAvoidanceSettings.Colliders)
+			for (const auto& Collider : ObjectAvoidanceSettings->ObjectInteractionSettings.Colliders)
 			{
 				TObjectPtr<UStaticMeshComponent> StaticMeshComponent = NewObject<UStaticMeshComponent>(GetTransientPackage(), NAME_None, RF_Transient);
-				StaticMeshComponent->SetStaticMesh(Mesh.LoadSynchronous());
+				StaticMeshComponent->SetStaticMesh(Collider.Mesh.LoadSynchronous());
 				StaticMeshComponent->UpdateBounds();
  
 				InOutParams.ManagedResources.Add(StaticMeshComponent);
-				InOutParams.Scene->AddComponent(StaticMeshComponent, Transform);
+				InOutParams.Scene->AddComponent(StaticMeshComponent, Collider.Transform);
 			}
 		}
  
