Set up basic unit testing framework for PVE

- Move test related code to ProceduralVegetation/Source/Test & ProceduralVegetationEditor/Source/Test
- PVImportTests.cpp
- PVMockTree.cpp
- Add PVTestsCommon as a testing helper library
- Remove old debug logic used for manual testing (PVPlantDescription & PVImportHelpers)
- Add simple unit test for Slope and Carve node

#Synced-CL 49880577

diff --git a/Engine/Plugins/Experimental/ProceduralVegetationEditor/Source/ProceduralVegetation/Private/Nodes/PVCarveSettings.h b/Engine/Plugins/Experimental/ProceduralVegetationEditor/Source/ProceduralVegetation/Private/Nodes/PVCarveSettings.h
--- a/Engine/Plugins/Experimental/ProceduralVegetationEditor/Source/ProceduralVegetation/Private/Nodes/PVCarveSettings.h
+++ b/Engine/Plugins/Experimental/ProceduralVegetationEditor/Source/ProceduralVegetation/Private/Nodes/PVCarveSettings.h
@@ -30,7 +30,7 @@ protected:
 
 	virtual FPCGElementPtr CreateElement() const override;
 
-private:
+public:
 	UPROPERTY(EditAnywhere, Category = "Carve", meta = (ShowOnlyInnerProperties))
 	FPVCarveParams CarveSettings;
 };
diff --git a/Engine/Plugins/Experimental/ProceduralVegetationEditor/Source/ProceduralVegetation/Private/Tests/Elements/PVCarveTest.cpp b/Engine/Plugins/Experimental/ProceduralVegetationEditor/Source/ProceduralVegetation/Private/Tests/Elements/PVCarveTest.cpp
new file mode 100644
--- /dev/null
+++ b/Engine/Plugins/Experimental/ProceduralVegetationEditor/Source/ProceduralVegetation/Private/Tests/Elements/PVCarveTest.cpp
@@ -0,0 +1,62 @@
+// Copyright Epic Games, Inc. All Rights Reserved.
+
+#if WITH_DEV_AUTOMATION_TESTS
+
+#include "Tests/PVTestsCommon.h"
+#include "DataTypes/PVGrowthData.h"
+#include "Nodes/PVCarveSettings.h"
+#include "PCGContext.h"
+
+namespace PVCarveTest
+{
+	TUniquePtr<FPCGContext> GenerateMockTreeDataAndRun(const FPVCarveParams& InCarveParams)
+	{
+		PCGTestsCommon::FTestData TestData;
+		PCGTestsCommon::GenerateSettings<UPVCarveSettings>(TestData);
+		UPVCarveSettings* Settings = CastChecked<UPVCarveSettings>(TestData.Settings);
+		Settings->CarveSettings = InCarveParams;
+
+		FPCGTaggedData& Inputs = TestData.InputData.TaggedData.Emplace_GetRef();
+		Inputs.Pin = PCGPinConstants::DefaultInputLabel;
+
+		UPVGrowthData* PVGrowthData = NewObject<UPVGrowthData>();
+		PVGrowthData->Initialize(MoveTemp(*PVMockTreeCollection::CreateCollection()));
+		Inputs.Data = PVGrowthData;
+
+		FPCGElementPtr TestElement = TestData.Settings->GetElement();
+
+		TUniquePtr<FPCGContext> Context = TestData.InitializeTestContext();
+
+		while (!TestElement->Execute(Context.Get())) {}
+
+		return Context;
+	}
+};
+
+#define PV_CARVE_TEST(TestName) PV_SIMPLE_AUTOMATION_TEST(Carve, TestName)
+
+PV_CARVE_TEST(Default)
+{
+	FPVCarveParams CarveParams;
+	CarveParams.Carve = 0;
+
+	TUniquePtr<FPCGContext> Context = PVCarveTest::GenerateMockTreeDataAndRun(CarveParams);
+
+	const TArray<FPCGTaggedData>& Outputs = Context->OutputData.TaggedData;
+	UTEST_EQUAL("Output count", Outputs.Num(), 1);
+
+	const UPVGrowthData* OutPVGrowthData = Cast<UPVGrowthData>(Outputs[0].Data);
+	UTEST_NOT_NULL("Output PVGrowthData", OutPVGrowthData);
+
+	const FManagedArrayCollection& Collection = OutPVGrowthData->GetCollection();
+	const PVMockTreeCollection::CompareResult Result = PVMockTreeCollection::Compare(Collection, true);
+	if (Result != PVMockTreeCollection::CompareResult::Success)
+	{
+		AddError(FString::Printf(TEXT("Collection does not contain mock tree data: %s"), PVMockTreeCollection::CompareResultToString(Result)));
+		return false;
+	}
+
+	return true;
+}
+
+#endif //WITH_DEV_AUTOMATION_TESTS
\ No newline at end of file
diff --git a/Engine/Plugins/Experimental/ProceduralVegetationEditor/Source/ProceduralVegetation/Private/Tests/Elements/PVSlopeTest.cpp b/Engine/Plugins/Experimental/ProceduralVegetationEditor/Source/ProceduralVegetation/Private/Tests/Elements/PVSlopeTest.cpp
new file mode 100644
--- /dev/null
+++ b/Engine/Plugins/Experimental/ProceduralVegetationEditor/Source/ProceduralVegetation/Private/Tests/Elements/PVSlopeTest.cpp
@@ -0,0 +1,64 @@
+// Copyright Epic Games, Inc. All Rights Reserved.
+
+#if WITH_DEV_AUTOMATION_TESTS
+
+#include "Tests/PVTestsCommon.h"
+#include "DataTypes/PVGrowthData.h"
+#include "Nodes/PVSlopeSettings.h"
+#include "PCGContext.h"
+
+namespace PVSlopeTest
+{
+	TUniquePtr<FPCGContext> GenerateMockTreeDataAndRun(const FPVSlopeParams& InSlopeParams)
+	{
+		PCGTestsCommon::FTestData TestData;
+		PCGTestsCommon::GenerateSettings<UPVSlopeSettings>(TestData);
+		UPVSlopeSettings* Settings = CastChecked<UPVSlopeSettings>(TestData.Settings);
+		Settings->SlopeParams = InSlopeParams;
+
+		FPCGTaggedData& Inputs = TestData.InputData.TaggedData.Emplace_GetRef();
+		Inputs.Pin = PCGPinConstants::DefaultInputLabel;
+
+		UPVGrowthData* PVGrowthData = NewObject<UPVGrowthData>();
+		PVGrowthData->Initialize(MoveTemp(*PVMockTreeCollection::CreateCollection()));
+		Inputs.Data = PVGrowthData;
+
+		FPCGElementPtr TestElement = TestData.Settings->GetElement();
+
+		TUniquePtr<FPCGContext> Context = TestData.InitializeTestContext();
+
+		while (!TestElement->Execute(Context.Get())) {}
+
+		return Context;
+	}
+};
+
+#define PV_SLOPE_TEST(TestName) PV_SIMPLE_AUTOMATION_TEST(Slope, TestName)
+
+PV_SLOPE_TEST(Default)
+{
+	FPVSlopeParams SlopeParams;
+	SlopeParams.BendStrength = 0;
+	SlopeParams.SlopeAngle = 0;
+	SlopeParams.SlopeDirection = 0;
+
+	TUniquePtr<FPCGContext> Context = PVSlopeTest::GenerateMockTreeDataAndRun(SlopeParams);
+
+	const TArray<FPCGTaggedData>& Outputs = Context->OutputData.TaggedData;
+	UTEST_EQUAL("Output count", Outputs.Num(), 1);
+
+	const UPVGrowthData* OutPVGrowthData = Cast<UPVGrowthData>(Outputs[0].Data);
+	UTEST_NOT_NULL("Output PVGrowthData", OutPVGrowthData);
+
+	const FManagedArrayCollection& Collection = OutPVGrowthData->GetCollection();
+	const PVMockTreeCollection::CompareResult Result = PVMockTreeCollection::Compare(Collection, true);
+	if (Result != PVMockTreeCollection::CompareResult::Success)
+	{
+		AddError(FString::Printf(TEXT("Collection does not contain mock tree data: %s"), PVMockTreeCollection::CompareResultToString(Result)));
+		return false;
+	}
+
+	return true;
+}
+
+#endif //WITH_DEV_AUTOMATION_TESTS
\ No newline at end of file
diff --git a/Engine/Plugins/Experimental/ProceduralVegetationEditor/Source/ProceduralVegetation/Private/Tests/PVMockTree.cpp b/Engine/Plugins/Experimental/ProceduralVegetationEditor/Source/ProceduralVegetation/Private/Tests/PVMockTree.cpp
new file mode 100644
--- /dev/null
+++ b/Engine/Plugins/Experimental/ProceduralVegetationEditor/Source/ProceduralVegetation/Private/Tests/PVMockTree.cpp
@@ -0,0 +1,86 @@
+// Copyright Epic Games, Inc. All Rights Reserved.
+
+#if WITH_DEV_AUTOMATION_TESTS
+
+#include "Tests/PVMockTree.h"
+#include "IPVEditorInterface.h"
+#include "GeometryCollection/ManagedArrayCollection.h"
+
+#include "Facades/PVBranchFacade.h"
+#include "Facades/PVPointFacade.h"
+
+const TCHAR* PVMockTreeCollection::CompareResultToString(PVMockTreeCollection::CompareResult InValue)
+{
+	switch (InValue)
+	{
+	case CompareResult::Success:
+		return TEXT("Success");
+	case CompareResult::InvalidCollection:
+		return TEXT("InvalidCollection");
+	case CompareResult::MismatchNrOfPoints:
+		return TEXT("MismatchNrOfPoints");
+	case CompareResult::MismatchNrOfBranches:
+		return TEXT("MismatchNrOfBranches");
+	case CompareResult::MismatchPointPositions:
+		return TEXT("MismatchPointPositions");
+	default:
+		ensureMsgf(false, TEXT("Unknown result type"));
+		return TEXT("Unknown");
+	}
+}
+
+PVMockTreeCollection::CompareResult PVMockTreeCollection::Compare(const FManagedArrayCollection& Collection, bool bCheckPointPositions)
+{
+#if WITH_EDITOR
+	const PV::Facades::FPointFacade InputPointFacade(Collection);
+	const PV::Facades::FBranchFacade InputBranchFacade(Collection);
+	if (!InputPointFacade.IsValid() || !InputBranchFacade.IsValid())
+	{
+		return CompareResult::InvalidCollection;
+	}
+
+	const TSharedRef<FManagedArrayCollection> BaselineCollection = PVMockTreeCollection::CreateCollection();
+
+	const PV::Facades::FPointFacade BaselinePointFacade(*BaselineCollection);
+	const PV::Facades::FBranchFacade BaselineBranchFacade(*BaselineCollection);
+
+	if (InputPointFacade.GetPositions().Num() != BaselinePointFacade.GetPositions().Num())
+	{
+		return CompareResult::MismatchNrOfPoints;
+	}
+
+	if (InputBranchFacade.GetBranchNumbers().Num() != BaselineBranchFacade.GetBranchNumbers().Num())
+	{
+		return CompareResult::MismatchNrOfBranches;
+	}
+
+	if (bCheckPointPositions)
+	{
+		for (int32 i = 0; i < InputPointFacade.GetPositions().Num(); ++i)
+		{
+			const FVector3f& PointA = InputPointFacade.GetPosition(i);
+			const FVector3f& PointB = BaselinePointFacade.GetPosition(i);
+
+			if (!PointA.Equals(PointB))
+			{
+				return CompareResult::MismatchPointPositions;
+			}
+		}
+	}
+#endif
+	return CompareResult::Success;
+}
+
+TSharedRef<FManagedArrayCollection> PVMockTreeCollection::CreateCollection()
+{
+#if WITH_EDITOR
+	if (PV::GPVEditorInterface)
+	{
+		return PV::GPVEditorInterface->CreateMockTreeGrowthDataCollection();
+	}
+#endif // WITH_EDITOR
+
+	return MakeShared<FManagedArrayCollection>();
+}
+
+#endif // WITH_DEV_AUTOMATION_TESTS
\ No newline at end of file
diff --git a/Engine/Plugins/Experimental/ProceduralVegetationEditor/Source/ProceduralVegetation/Private/Tests/PVTestsCommon.cpp b/Engine/Plugins/Experimental/ProceduralVegetationEditor/Source/ProceduralVegetation/Private/Tests/PVTestsCommon.cpp
new file mode 100644
--- /dev/null
+++ b/Engine/Plugins/Experimental/ProceduralVegetationEditor/Source/ProceduralVegetation/Private/Tests/PVTestsCommon.cpp
@@ -0,0 +1,9 @@
+// Copyright Epic Games, Inc. All Rights Reserved.
+
+#if WITH_DEV_AUTOMATION_TESTS
+
+#include "Tests/PVTestsCommon.h"
+
+DEFINE_LOG_CATEGORY(LogPVTest);
+
+#endif //WITH_DEV_AUTOMATION_TESTS
\ No newline at end of file
diff --git a/Engine/Plugins/Experimental/ProceduralVegetationEditor/Source/ProceduralVegetation/Public/IPVEditorInterface.h b/Engine/Plugins/Experimental/ProceduralVegetationEditor/Source/ProceduralVegetation/Public/IPVEditorInterface.h
--- a/Engine/Plugins/Experimental/ProceduralVegetationEditor/Source/ProceduralVegetation/Public/IPVEditorInterface.h
+++ b/Engine/Plugins/Experimental/ProceduralVegetationEditor/Source/ProceduralVegetation/Public/IPVEditorInterface.h
@@ -19,6 +19,7 @@ namespace PV
 
 		virtual bool ImportGrowthDataFromStaticMesh(const FPVImportStaticMeshParams& InParams, FManagedArrayCollection& OutCollection) = 0;
 		virtual bool ImportGrowthDataFromTexture2D(const FPVImportTexture2DParams& InParams, FManagedArrayCollection& OutCollection) = 0;
+		virtual TSharedRef<FManagedArrayCollection> CreateMockTreeGrowthDataCollection() const = 0;
 	};
 
 	extern PROCEDURALVEGETATION_API IPVEditorInterface* GPVEditorInterface;
diff --git a/Engine/Plugins/Experimental/ProceduralVegetationEditor/Source/ProceduralVegetation/Public/Tests/PVMockTree.h b/Engine/Plugins/Experimental/ProceduralVegetationEditor/Source/ProceduralVegetation/Public/Tests/PVMockTree.h
new file mode 100644
--- /dev/null
+++ b/Engine/Plugins/Experimental/ProceduralVegetationEditor/Source/ProceduralVegetation/Public/Tests/PVMockTree.h
@@ -0,0 +1,128 @@
+// Copyright Epic Games, Inc. All Rights Reserved.
+#pragma once
+
+#if WITH_DEV_AUTOMATION_TESTS
+
+#include "CoreMinimal.h"
+
+struct FPVBranchHierarchyDescription;
+struct FManagedArrayCollection;
+
+namespace PVMockTree
+{
+	/*
+	* Branch Indices:     Branch Numbers:
+	*           *B2		            *BN3
+	*   *B3 |  / 		    *BN4|  /
+	*    \  | / / *B1	     \  | / / *BN2
+	*     \ | \/		      \ | \/ 
+	*      \| /			       \| / 
+	*       |/ 			        |/ 
+	*       |			        |
+	*       | *B0		        | *BN1
+	* 
+	* 
+	* Point Indices:
+	*      5*  11*
+	* 		|   /     *9
+	* *6    |  /     /
+	* \     | *10   /
+	*  \   4*  \   /
+	*   \   |   \ /
+	*    \  |    *8
+	*     \ |   /
+	*      3*  *7
+	*       | /
+	*       |/
+	*       *2
+	*       |
+	*       *1
+	*       |
+	*       *0
+	*/
+
+	const FVector3f Point0  = {  00.0000,  00.0000, 000.0000 };
+	const FVector3f Point1  = {  00.0394, -00.1783, 046.2954 };
+	const FVector3f Point2  = {  00.0788, -00.3566, 092.5909 };
+	const FVector3f Point3  = {  01.1927,  00.9499, 125.4070 };
+	const FVector3f Point4  = {  03.0192,  01.8644, 161.3479 };
+	const FVector3f Point5  = {  04.8458,  02.7790, 197.2887 };
+	const FVector3f Point6  = {  19.7899,  73.3726, 171.2636 };
+	const FVector3f Point7  = { -23.0204, -29.6581, 148.4679 };
+	const FVector3f Point8  = { -24.5249, -31.0789, 156.2391 };
+	const FVector3f Point9  = { -30.1656, -36.4065, 185.3766 };
+	const FVector3f Point10 = { -24.4424, -21.5258, 175.3671 };
+	const FVector3f Point11 = { -23.0204, -23.7218, 201.3998 };
+
+	const float Point0_Radius = 10.f;
+	const float Point1_Radius = 0.f;
+	const float Point2_Radius = 0.f;
+	const float Point3_Radius = 0.f;
+	const float Point4_Radius = 0.f;
+	const float Point5_Radius = 0.f;
+	const float Point6_Radius = 0.f;
+	const float Point7_Radius = 0.f;
+	const float Point8_Radius = 0.f;
+	const float Point9_Radius = 0.f;
+	const float Point10_Radius = 0.f;
+	const float Point11_Radius = 0.f;
+
+	const float Point0To1_Dist   = 46.2958;
+	const float Point1To2_Dist   = 46.2958;
+	const float Point2To3_Dist   = 32.861;
+	const float Point3To4_Dist   = 35.9989;
+	const float Point4To5_Dist   = 35.9989;
+	const float Point3To6_Dist   = 87.7139;
+	const float Point2To7_Dist   = 67.1893;
+	const float Point7To8_Dist   = 08.0419;
+	const float Point8To9_Dist   = 30.1529;
+	const float Point8To10_Dist  = 21.3811;
+	const float Point10To11_Dist = 26.1638;
+
+	const float Point0_LengthFromRoot = 0;
+	const float Point1_LengthFromRoot = Point0To1_Dist;
+	const float Point2_LengthFromRoot = Point1_LengthFromRoot + Point1To2_Dist;
+	const float Point3_LengthFromRoot = Point2_LengthFromRoot + Point2To3_Dist;
+	const float Point4_LengthFromRoot = Point3_LengthFromRoot + Point3To4_Dist;
+	const float Point5_LengthFromRoot = Point4_LengthFromRoot + Point4To5_Dist;
+	const float Point6_LengthFromRoot = Point3_LengthFromRoot + Point3To6_Dist;
+	const float Point7_LengthFromRoot = Point2_LengthFromRoot + Point2To7_Dist;
+	const float Point8_LengthFromRoot = Point7_LengthFromRoot + Point7To8_Dist;
+	const float Point9_LengthFromRoot = Point8_LengthFromRoot + Point8To9_Dist;
+	const float Point10_LengthFromRoot = Point8_LengthFromRoot + Point8To10_Dist;
+	const float Point11_LengthFromRoot = Point10_LengthFromRoot + Point10To11_Dist;
+
+	const TArray<int32> Branch0_PointIndices = { 0, 1, 2, 3, 4, 5 };
+	const TArray<int32> Branch1_PointIndices = { 2, 7, 8, 9 };
+	const TArray<int32> Branch2_PointIndices = { 8, 10, 11 };
+	const TArray<int32> Branch3_PointIndices = { 3, 6 };
+
+	const int32 Branch0_BranchNumber = 1;
+	const int32 Branch1_BranchNumber = 2;
+	const int32 Branch2_BranchNumber = 3;
+	const int32 Branch3_BranchNumber = 4;
+
+	const TArray<FVector3f> Points = { Point0, Point1, Point2, Point3, Point4, Point5, Point6, Point7, Point8, Point9, Point10, Point11 };
+	const TArray<float> PointsRadii = { Point0_Radius, Point1_Radius, Point2_Radius, Point3_Radius, Point4_Radius, Point5_Radius, Point6_Radius, Point7_Radius, Point8_Radius, Point9_Radius, Point10_Radius, Point11_Radius };
+	const TArray<TArray<int32>> BranchPointIndices = { Branch0_PointIndices, Branch1_PointIndices, Branch2_PointIndices, Branch3_PointIndices };
+}
+
+namespace PVMockTreeCollection
+{
+	enum class CompareResult
+	{
+		Success,
+		InvalidCollection,
+		MismatchNrOfPoints,
+		MismatchNrOfBranches,
+		MismatchPointPositions
+	};
+
+	PROCEDURALVEGETATION_API extern const TCHAR* CompareResultToString(CompareResult InValue);
+
+	PROCEDURALVEGETATION_API extern CompareResult Compare(const FManagedArrayCollection& Collection, bool bCheckPointPositions);
+
+	PROCEDURALVEGETATION_API extern TSharedRef<FManagedArrayCollection> CreateCollection();
+}
+
+#endif // WITH_DEV_AUTOMATION_TESTS
\ No newline at end of file
diff --git a/Engine/Plugins/Experimental/ProceduralVegetationEditor/Source/ProceduralVegetation/Public/Tests/PVTestsCommon.h b/Engine/Plugins/Experimental/ProceduralVegetationEditor/Source/ProceduralVegetation/Public/Tests/PVTestsCommon.h
new file mode 100644
--- /dev/null
+++ b/Engine/Plugins/Experimental/ProceduralVegetationEditor/Source/ProceduralVegetation/Public/Tests/PVTestsCommon.h
@@ -0,0 +1,21 @@
+// Copyright Epic Games, Inc. All Rights Reserved.
+#pragma once
+
+#if WITH_DEV_AUTOMATION_TESTS
+
+#include "CoreMinimal.h"
+#include "Tests/PVMockTree.h"
+#include "Tests/PCGTestsCommon.h"
+
+#include "Misc/AutomationTest.h"
+
+DECLARE_LOG_CATEGORY_EXTERN(LogPVTest, Log, All);
+
+#define PV_IMPLEMENT_SIMPLE_AUTOMATION_TEST_PRIVATE(TClass, PrettyName, TFlags) \
+	IMPLEMENT_SIMPLE_AUTOMATION_TEST(TClass, PrettyName, TFlags) \
+	bool TClass::RunTest(const FString& Parameters)
+
+#define PV_SIMPLE_AUTOMATION_TEST(CategoryName, TestName) \
+	PV_IMPLEMENT_SIMPLE_AUTOMATION_TEST_PRIVATE(UE_JOIN(UE_JOIN(FPVAutomationTest_,CategoryName),TestName), "PVE." UE_STRINGIZE(CategoryName) "." UE_STRINGIZE(TestName), PCGTestsCommon::TestFlags)
+
+#endif //WITH_DEV_AUTOMATION_TESTS
\ No newline at end of file
diff --git a/Engine/Plugins/Experimental/ProceduralVegetationEditor/Source/ProceduralVegetationEditor/Private/Importer/PVImportHelpers.cpp b/Engine/Plugins/Experimental/ProceduralVegetationEditor/Source/ProceduralVegetationEditor/Private/Importer/PVImportHelpers.cpp
--- a/Engine/Plugins/Experimental/ProceduralVegetationEditor/Source/ProceduralVegetationEditor/Private/Importer/PVImportHelpers.cpp
+++ b/Engine/Plugins/Experimental/ProceduralVegetationEditor/Source/ProceduralVegetationEditor/Private/Importer/PVImportHelpers.cpp
@@ -1610,23 +1610,6 @@ bool PV::Import::Helper::GenerateGrowthDataFromPlantDescription(FManagedArrayCol
 	return true;
 }
 
-bool PV::Import::Helper::GenerateProceduralVegetationPresetFromPlantDescription(const FPVPlantDescription& InPlantDescription, UProceduralVegetationPreset* InOutProceduralVegetationPreset)
-{
-	FManagedArrayCollection Collection;
-	if (!GenerateGrowthDataFromPlantDescription(Collection, InPlantDescription))
-	{
-		return false;
-	}
-
-	InOutProceduralVegetationPreset->Modify();
-	InOutProceduralVegetationPreset->Variants.Reset();
-	InOutProceduralVegetationPreset->Variants.Add(TEXT("Default"), Collection);
-	InOutProceduralVegetationPreset->FillVariationInfo();
-	InOutProceduralVegetationPreset->PostEditChange();
-
-	return true;
-}
-
 void PV::Import::Helper::DistToBranch(const FVector3f& Point, const TArray<FVector3f>& Points, const TArray<int32>& BranchPointIndices, FVector3f& OutClosestPoint, float& OutDistSquared, int32& OutPointIndex)
 {
 	OutClosestPoint = FVector3f::ZeroVector;
diff --git a/Engine/Plugins/Experimental/ProceduralVegetationEditor/Source/ProceduralVegetationEditor/Private/Importer/PVImportHelpers.h b/Engine/Plugins/Experimental/ProceduralVegetationEditor/Source/ProceduralVegetationEditor/Private/Importer/PVImportHelpers.h
--- a/Engine/Plugins/Experimental/ProceduralVegetationEditor/Source/ProceduralVegetationEditor/Private/Importer/PVImportHelpers.h
+++ b/Engine/Plugins/Experimental/ProceduralVegetationEditor/Source/ProceduralVegetationEditor/Private/Importer/PVImportHelpers.h
@@ -93,8 +93,6 @@ namespace PV::Import::Helper
 
 	bool GenerateGrowthDataFromPlantDescription(FManagedArrayCollection& InOutCollection, const FPVPlantDescription& InPlantDescription);
 
-	bool GenerateProceduralVegetationPresetFromPlantDescription(const FPVPlantDescription& InPlantDescription, UProceduralVegetationPreset* InOutProceduralVegetationPreset);
-
 	void DistToBranch(const FVector3f& Point, const TArray<FVector3f>& Points, const TArray<int32>& BranchPointIndices, FVector3f& OutClosestPoint, float& OutDistSquared, int32& OutPointIndex);
 
 	float DistToBranch(const FVector3f& Point, const TArray<FVector3f>& Points, const TArray<int32>& BranchPointIndices);
diff --git a/Engine/Plugins/Experimental/ProceduralVegetationEditor/Source/ProceduralVegetationEditor/Private/Importer/PVImportTests.cpp b/Engine/Plugins/Experimental/ProceduralVegetationEditor/Source/ProceduralVegetationEditor/Private/Importer/PVImportTests.cpp
deleted file mode 100644
diff --git a/Engine/Plugins/Experimental/ProceduralVegetationEditor/Source/ProceduralVegetationEditor/Private/Importer/PVImporter_MockTree.cpp b/Engine/Plugins/Experimental/ProceduralVegetationEditor/Source/ProceduralVegetationEditor/Private/Importer/PVImporter_MockTree.cpp
deleted file mode 100644
diff --git a/Engine/Plugins/Experimental/ProceduralVegetationEditor/Source/ProceduralVegetationEditor/Private/Importer/PVImporter_MockTree.h b/Engine/Plugins/Experimental/ProceduralVegetationEditor/Source/ProceduralVegetationEditor/Private/Importer/PVImporter_MockTree.h
deleted file mode 100644
diff --git a/Engine/Plugins/Experimental/ProceduralVegetationEditor/Source/ProceduralVegetationEditor/Private/Importer/PVMockTree.cpp b/Engine/Plugins/Experimental/ProceduralVegetationEditor/Source/ProceduralVegetationEditor/Private/Importer/PVMockTree.cpp
deleted file mode 100644
diff --git a/Engine/Plugins/Experimental/ProceduralVegetationEditor/Source/ProceduralVegetationEditor/Private/Importer/PVMockTree.h b/Engine/Plugins/Experimental/ProceduralVegetationEditor/Source/ProceduralVegetationEditor/Private/Importer/PVMockTree.h
deleted file mode 100644
diff --git a/Engine/Plugins/Experimental/ProceduralVegetationEditor/Source/ProceduralVegetationEditor/Private/Importer/PVPlantDescription.cpp b/Engine/Plugins/Experimental/ProceduralVegetationEditor/Source/ProceduralVegetationEditor/Private/Importer/PVPlantDescription.cpp
--- a/Engine/Plugins/Experimental/ProceduralVegetationEditor/Source/ProceduralVegetationEditor/Private/Importer/PVPlantDescription.cpp
+++ b/Engine/Plugins/Experimental/ProceduralVegetationEditor/Source/ProceduralVegetationEditor/Private/Importer/PVPlantDescription.cpp
@@ -2,12 +2,6 @@
 
 #include "PVPlantDescription.h"
 
-#include "Engine/StaticMesh.h"
-#include "PrimitiveSceneProxy.h"
-#include "PrimitiveDrawInterface.h"
-#include "MeshDescription.h"
-#include "MeshElementCollector.h"
-
 FPVBranchHierarchyDescription::EValidateHierarchyResult FPVBranchHierarchyDescription::ValidateHierarchy() const
 {
 	if (!Branches.IsValidIndex(RootBranchIndex))
@@ -93,222 +87,3 @@ FPVBranchHierarchyDescription::EValidateHierarchyResult FPVBranchHierarchyDescri
 
 	return EValidateHierarchyResult::Success;
 }
-
-FPrimitiveSceneProxy* UPVPlantDescriptionVisualizationComponent::CreateSceneProxy()
-{
-	class FPVPlantDescriptionSceneProxy final : public FPrimitiveSceneProxy
-	{
-	private:
-		const FPVPlantDescription PlantDescription;
-		TArray<FMeshDescription> FoliageMeshDescriptions;
-	public:
-		FPVPlantDescriptionSceneProxy(const UPVPlantDescriptionVisualizationComponent* InComponent)
-			: FPrimitiveSceneProxy(InComponent)
-			, PlantDescription(InComponent->PlantDescriptionDataAsset ? InComponent->PlantDescriptionDataAsset->PlantDescription : FPVPlantDescription())
-		{
-			for (TObjectPtr<UStaticMesh> FoliageMesh : PlantDescription.Foliage.FoliageMeshes)
-			{
-				FMeshDescription& MeshDescription = FoliageMeshDescriptions.AddDefaulted_GetRef();
-
-				if (FoliageMesh)
-				{
-					if (FMeshDescription* FoliageMeshDescription = FoliageMesh->GetMeshDescription(0))
-					{
-						MeshDescription = *FoliageMeshDescription;
-					}
-				}
-			}
-		}
-
-		virtual FPrimitiveViewRelevance GetViewRelevance(const FSceneView* View) const override
-		{
-			FPrimitiveViewRelevance ViewRelevance;
-			ViewRelevance.bDrawRelevance = IsShown(View);
-			ViewRelevance.bDynamicRelevance = true;
-			ViewRelevance.bSeparateTranslucency = ViewRelevance.bNormalTranslucency = true;
-			return ViewRelevance;
-		}
-
-		static void DrawDebugCircle(FPrimitiveDrawInterface* PDI, const FVector& Center, const FQuat& Rotation, float Radius, int32 Segments, const FColor& Color, uint8 DepthPriority, float Thickness = 0.f)
-		{
-			Segments = FMath::Max(Segments, 4);
-			const float AngleStep = 2.f * UE_PI / float(Segments);
-
-			float Angle = 0.f;
-			while (Segments--)
-			{
-				const FVector Vertex1 = Center + Radius * Rotation.RotateVector(FVector(0, FMath::Cos(Angle), 0) + FVector(0, 0, FMath::Sin(Angle)));
-				Angle += AngleStep;
-				const FVector Vertex2 = Center + Radius * Rotation.RotateVector(FVector(0, FMath::Cos(Angle), 0) + FVector(0, 0, FMath::Sin(Angle)));
-				PDI->DrawLine(Vertex1, Vertex2, Color, DepthPriority, Thickness);
-			}
-		}
-
-		void DrawBranchHierarchy(FPrimitiveDrawInterface* PDI) const
-		{
-			const static TFunction<void(FPrimitiveDrawInterface*, const FPVBranchHierarchyDescription&, int32, int32)> DrawBranchesRecursive = 
-				[](FPrimitiveDrawInterface* PDI, const FPVBranchHierarchyDescription& BranchHierarchy, int32 BranchIndex, int32 Depth)
-			{
-				const FPVBranchDescription& Branch = BranchHierarchy.Branches[BranchIndex];
-
-				const FLinearColor DrawColor = FLinearColor::MakeRandomSeededColor(BranchIndex);
-
-				for (int32 i = 0; i < Branch.PointIndices.Num(); ++i)
-				{
-					const int32 PointIndex = Branch.PointIndices[i];
-					const FVector3f& Point = BranchHierarchy.Points[PointIndex];
-
-					PDI->DrawPoint(FVector(Point), DrawColor, 10.f, SDPG_Foreground);
-
-					if (i != 0)
-					{
-						if (BranchHierarchy.PointsRadii.IsValidIndex(PointIndex) && BranchHierarchy.PointsRadii[PointIndex] > 0)
-						{
-							const FVector3f ApicalDirection = [&]()
-								{
-									if (i < Branch.PointIndices.Num() - 1)
-									{
-										return (Point - BranchHierarchy.Points[Branch.PointIndices[i + 1]]).GetSafeNormal();
-									}
-									else if (i > 0)
-									{
-										return (BranchHierarchy.Points[Branch.PointIndices[i - 1]] - Point).GetSafeNormal();
-									}
-
-									return FVector3f::UpVector;
-								}();
-
-
-							DrawDebugCircle(PDI, FVector(Point), FRotationMatrix::MakeFromX(FVector(ApicalDirection)).ToQuat(), BranchHierarchy.PointsRadii[PointIndex], 8, FColor::Red, SDPG_Foreground, 2);
-						}
-					}
-
-					if (i != Branch.PointIndices.Num() - 1)
-					{
-						const int32 PointIndex2 = Branch.PointIndices[i + 1];
-						const FVector3f& Point2 = BranchHierarchy.Points[PointIndex2];
-
-						PDI->DrawLine(FVector(Point), FVector(Point2), DrawColor, SDPG_Foreground, 1.f);
-					}
-				}
-
-				for (int32 ChildIndex : Branch.ChildBranchIndices)
-				{
-					DrawBranchesRecursive(PDI, BranchHierarchy, ChildIndex, Depth + 1);
-				}
-			};
-			
-			for (const FPVBranchHierarchyDescription& BranchHierarchy : PlantDescription.BranchHierarchies)
-			{
-				if (BranchHierarchy.ValidateHierarchy() != FPVBranchHierarchyDescription::EValidateHierarchyResult::Success)
-				{
-					return;
-				}
-
-				DrawBranchesRecursive(PDI, BranchHierarchy, BranchHierarchy.RootBranchIndex, 0);
-			}
-		}
-		
-		void DrawFoliage(FPrimitiveDrawInterface* PDI) const
-		{
-			for (const FPVFoliagePointsDescription& FoliagePointsDescription : PlantDescription.Foliage.FoliagePoints)
-			{
-				for (int32 i = 0; i < FoliagePointsDescription.Positions.Num(); ++i)
-				{
-					const FVector3f Position = FoliagePointsDescription.Positions[i];
-					const FQuat4f Orientation = FoliagePointsDescription.Orientations[i];
-					const FVector3f Scale = FoliagePointsDescription.Scales[i];
-					const FMeshDescription& MeshDescription = FoliageMeshDescriptions[FoliagePointsDescription.FoliageIndices[i]];
-					
-					for (FTriangleID TriangleID : MeshDescription.Triangles().GetElementIDs())
-					{
-						const TConstArrayView<FVertexInstanceID> VertexInstances = MeshDescription.GetTriangleVertexInstances(TriangleID);
-
-						const FVertexID VertexID0 = MeshDescription.GetVertexInstanceVertex(VertexInstances[0]);
-						const FVertexID VertexID1 = MeshDescription.GetVertexInstanceVertex(VertexInstances[1]);
-						const FVertexID VertexID2 = MeshDescription.GetVertexInstanceVertex(VertexInstances[2]);
-
-						FVector3f Vertex0Position = MeshDescription.GetVertexPosition(VertexID0);
-						FVector3f Vertex1Position = MeshDescription.GetVertexPosition(VertexID1);
-						FVector3f Vertex2Position = MeshDescription.GetVertexPosition(VertexID2);
-
-						Vertex0Position *= Scale;
-						Vertex0Position = Orientation.RotateVector(Vertex0Position);
-						Vertex0Position += Position;
-
-						Vertex1Position *= Scale;
-						Vertex1Position = Orientation.RotateVector(Vertex1Position);
-						Vertex1Position += Position;
-
-						Vertex2Position *= Scale;
-						Vertex2Position = Orientation.RotateVector(Vertex2Position);
-						Vertex2Position += Position;
-
-						PDI->DrawLine(FVector(Vertex0Position), FVector(Vertex1Position), FLinearColor(FColor::Orange), SDPG_Foreground, 1);
-						PDI->DrawLine(FVector(Vertex1Position), FVector(Vertex2Position), FLinearColor(FColor::Orange), SDPG_Foreground, 1);
-						PDI->DrawLine(FVector(Vertex2Position), FVector(Vertex0Position), FLinearColor(FColor::Orange), SDPG_Foreground, 1);
-					}
-				}
-			}
-		}
-
-		virtual void GetDynamicMeshElements(
-			const TArray<const FSceneView*>& Views,
-			const FSceneViewFamily& ViewFamily,
-			uint32 VisibilityMap,
-			FMeshElementCollector& Collector
-		) const override
-		{
-			for (int32 ViewIndex = 0; ViewIndex < Views.Num(); ViewIndex++)
-			{
-				if (VisibilityMap & (1 << ViewIndex))
-				{
-					const FSceneView* View = Views[ViewIndex];
-
-#if UE_ENABLE_DEBUG_DRAWING
-					FPrimitiveDrawInterface* PDI = Collector.GetDebugPDI(ViewIndex);
-#else	
-					FPrimitiveDrawInterface* PDI = Collector.GetPDI(ViewIndex);
-#endif
-
-					DrawBranchHierarchy(PDI);
-					DrawFoliage(PDI);
-				}
-			}
-		}
-
-		virtual SIZE_T GetTypeHash() const override
-		{
-			static size_t UniquePointer;
-			return reinterpret_cast<size_t>(&UniquePointer);
-		}
-
-		virtual uint32 GetMemoryFootprint(void) const override
-		{
-			return sizeof * this + GetAllocatedSize();
-		}
-
-		virtual bool CanBeOccluded() const override { return false; }
-	};
-
-	FPVPlantDescriptionSceneProxy* Proxy = new FPVPlantDescriptionSceneProxy(this);
-	return Proxy;
-}
-
-FBoxSphereBounds UPVPlantDescriptionVisualizationComponent::CalcBounds(const FTransform& LocalToWorld) const
-{
-	FBox BBox = FBox(ForceInit);
-
-	if (IsValid(PlantDescriptionDataAsset))
-	{
-		for (const FPVBranchHierarchyDescription& BranchHierarchy : PlantDescriptionDataAsset->PlantDescription.BranchHierarchies)
-		{
-			for (const FVector3f& Point : BranchHierarchy.Points)
-			{
-				BBox += FVector(Point);
-			}
-		}
-	}
-
-	return FBoxSphereBounds(BBox).TransformBy(LocalToWorld);
-}
diff --git a/Engine/Plugins/Experimental/ProceduralVegetationEditor/Source/ProceduralVegetationEditor/Private/Importer/PVPlantDescription.h b/Engine/Plugins/Experimental/ProceduralVegetationEditor/Source/ProceduralVegetationEditor/Private/Importer/PVPlantDescription.h
--- a/Engine/Plugins/Experimental/ProceduralVegetationEditor/Source/ProceduralVegetationEditor/Private/Importer/PVPlantDescription.h
+++ b/Engine/Plugins/Experimental/ProceduralVegetationEditor/Source/ProceduralVegetationEditor/Private/Importer/PVPlantDescription.h
@@ -2,11 +2,11 @@
 #pragma once
 
 #include "CoreMinimal.h"
-#include "Engine/DataAsset.h"
-#include "Components/PrimitiveComponent.h"
 
 #include "PVPlantDescription.generated.h"
 
+class UStaticMesh;
+
 USTRUCT()
 struct FPVBranchDescription
 {
@@ -144,29 +144,3 @@ struct FPVPlantDescription
 
 	int32 GetNumFoliagePoints() const { return Foliage.GetNumFoliagePoints(); }
 };
-
-// Internal, for debugging only
-UCLASS(BlueprintType)
-class PROCEDURALVEGETATIONEDITOR_API UPVPlantDescriptionDataAsset : public UDataAsset
-{
-	GENERATED_BODY()
-public:
-	UPROPERTY(VisibleAnywhere, Category = "Branch Hierarchy")
-	FPVPlantDescription PlantDescription;
-};
-
-// Internal, for debugging only
-UCLASS(Blueprintable, BlueprintType, meta = (BlueprintSpawnableComponent), MinimalAPI)
-class UPVPlantDescriptionVisualizationComponent : public UPrimitiveComponent
-{
-	GENERATED_BODY()
-public:
-
-	UPROPERTY(EditAnywhere, Category = "Branch Hierarchy Asset")
-	TObjectPtr<UPVPlantDescriptionDataAsset> PlantDescriptionDataAsset;
-
-public:
-
-	virtual FPrimitiveSceneProxy* CreateSceneProxy() override;
-	virtual FBoxSphereBounds CalcBounds(const FTransform& LocalToWorld) const override;
-};
\ No newline at end of file
diff --git a/Engine/Plugins/Experimental/ProceduralVegetationEditor/Source/ProceduralVegetationEditor/Private/PVEditorInterfaceImpl.cpp b/Engine/Plugins/Experimental/ProceduralVegetationEditor/Source/ProceduralVegetationEditor/Private/PVEditorInterfaceImpl.cpp
--- a/Engine/Plugins/Experimental/ProceduralVegetationEditor/Source/ProceduralVegetationEditor/Private/PVEditorInterfaceImpl.cpp
+++ b/Engine/Plugins/Experimental/ProceduralVegetationEditor/Source/ProceduralVegetationEditor/Private/PVEditorInterfaceImpl.cpp
@@ -2,13 +2,17 @@
 
 #include "PVEditorInterfaceImpl.h"
 
+#include "Importer/PVImportHelpers.h"
+#include "Importer/PVPlantDescription.h"
 #include "Importer/PVImporter_StaticMesh.h"
 #include "Importer/PVImporter_Texture2D.h"
 #include "PVImportStaticMeshParams.h"
 #include "PVImportTexture2DParams.h"
+#include "Tests/PVMockTree.h"
 
 #include "Engine/StaticMesh.h"
 #include "Engine/Texture2D.h"
+#include "GeometryCollection/ManagedArrayCollection.h"
 
 bool PV::FPVEditorInterfaceImpl::ImportGrowthDataFromStaticMesh(const FPVImportStaticMeshParams& InParams, FManagedArrayCollection& OutCollection)
 {
@@ -20,4 +24,16 @@ bool PV::FPVEditorInterfaceImpl::ImportGrowthDataFromTexture2D(const FPVImportTe
 {
 	const PV::Texture2DImport::EImportResult ImportResult = PV::Texture2DImport::ImportGrowthDataFromTexture2D(InParams, OutCollection);
 	return ImportResult == PV::Texture2DImport::EImportResult::Success;
+}
+
+TSharedRef<FManagedArrayCollection> PV::FPVEditorInterfaceImpl::CreateMockTreeGrowthDataCollection() const
+{
+	FPVPlantDescription PlantDescription;
+	PlantDescription.BranchHierarchies.Add(PV::Import::Helper::CreateBranchHierarchyFromPoints(PVMockTree::Points, PVMockTree::PointsRadii, PVMockTree::BranchPointIndices));
+
+	TSharedRef<FManagedArrayCollection> OutCollection = MakeShared<FManagedArrayCollection>();
+	const bool bResult = PV::Import::Helper::GenerateGrowthDataFromPlantDescription(*OutCollection, PlantDescription);
+	check(bResult != false); // Mock tree is staticly defined to always be valid, so it should always return valid growth data
+
+	return OutCollection;
 }
\ No newline at end of file
diff --git a/Engine/Plugins/Experimental/ProceduralVegetationEditor/Source/ProceduralVegetationEditor/Private/PVEditorInterfaceImpl.h b/Engine/Plugins/Experimental/ProceduralVegetationEditor/Source/ProceduralVegetationEditor/Private/PVEditorInterfaceImpl.h
--- a/Engine/Plugins/Experimental/ProceduralVegetationEditor/Source/ProceduralVegetationEditor/Private/PVEditorInterfaceImpl.h
+++ b/Engine/Plugins/Experimental/ProceduralVegetationEditor/Source/ProceduralVegetationEditor/Private/PVEditorInterfaceImpl.h
@@ -16,6 +16,7 @@ namespace PV
 	public:
 		bool ImportGrowthDataFromStaticMesh(const FPVImportStaticMeshParams& InParams, FManagedArrayCollection& OutCollection) override;
 		bool ImportGrowthDataFromTexture2D(const FPVImportTexture2DParams& InParams, FManagedArrayCollection& OutCollection) override;
+		TSharedRef<FManagedArrayCollection> CreateMockTreeGrowthDataCollection() const override;
 	};
 };
 #endif
diff --git a/Engine/Plugins/Experimental/ProceduralVegetationEditor/Source/ProceduralVegetationEditor/Private/Tests/Importer/PVImportTests.cpp b/Engine/Plugins/Experimental/ProceduralVegetationEditor/Source/ProceduralVegetationEditor/Private/Tests/Importer/PVImportTests.cpp
new file mode 100644
--- /dev/null
+++ b/Engine/Plugins/Experimental/ProceduralVegetationEditor/Source/ProceduralVegetationEditor/Private/Tests/Importer/PVImportTests.cpp
@@ -0,0 +1,698 @@
+// Copyright Epic Games, Inc. All Rights Reserved.
+
+#if WITH_DEV_AUTOMATION_TESTS
+
+#include "CoreTypes.h"
+#include "Misc/AutomationTest.h"
+#include "HAL/PreprocessorHelpers.h"
+#include "GeometryCollection/ManagedArrayCollection.h"
+
+#include "Tests/PVTestsCommon.h"
+#include "Tests/PVMockTree.h"
+#include "Importer/PVImportHelpers.h"
+
+#include "Facades/PVAttributesNames.h"
+#include "Facades/PVPointFacade.h"
+#include "Facades/PVBranchFacade.h"
+#include "Facades/PVPlantFacade.h"
+
+#define PV_IMPORTER_TEST(TestName) PV_SIMPLE_AUTOMATION_TEST(Importer, TestName)
+
+namespace PVMockTree
+{
+	FPVBranchHierarchyDescription CreateBranchHierarchyDescription()
+	{
+		return PV::Import::Helper::CreateBranchHierarchyFromPoints(PVMockTree::Points, PVMockTree::PointsRadii, PVMockTree::BranchPointIndices);
+	}
+}
+
+PV_IMPORTER_TEST(FillGrowthDataFromBranchHierarchyDescription)
+{
+	const FPVBranchHierarchyDescription BranchHierarchy = PVMockTree::CreateBranchHierarchyDescription();
+
+	const int32 NumBranchPoints = BranchHierarchy.GetNumBranchPoints();
+	const int32 NumFoliagePoints = 0;
+	const int32 NumBranches = BranchHierarchy.GetNumBranches();
+	const int32 NumFoliageMeshes = 0;
+
+	FManagedArrayCollection ExpectedCollection;
+	PV::Import::Helper::CreateEmptyGrowthData(ExpectedCollection, NumBranchPoints, NumFoliagePoints, NumBranches, NumFoliageMeshes);
+
+	PV::Facades::FPointFacade ExpectedPointFacade(ExpectedCollection);
+	PV::Facades::FBranchFacade ExpectedBranchFacade(ExpectedCollection);
+	
+	// Init ExpectedPointFacade
+	{
+		TManagedArray<FVector3f>& ExpectedPositions = ExpectedPointFacade.ModifyPositions();
+
+		ExpectedPositions[0] = PVMockTree::Point0;
+		ExpectedPositions[1] = PVMockTree::Point1;
+		ExpectedPositions[2] = PVMockTree::Point2;
+		ExpectedPositions[3] = PVMockTree::Point3;
+		ExpectedPositions[4] = PVMockTree::Point4;
+		ExpectedPositions[5] = PVMockTree::Point5;
+		ExpectedPositions[6] = PVMockTree::Point6;
+		ExpectedPositions[7] = PVMockTree::Point7;
+		ExpectedPositions[8] = PVMockTree::Point8;
+		ExpectedPositions[9] = PVMockTree::Point9;
+		ExpectedPositions[10] = PVMockTree::Point10;
+		ExpectedPositions[11] = PVMockTree::Point11;
+	}
+
+	// Init ExpectedBranchFacade
+	{
+		ExpectedBranchFacade.SetPoints(0, PVMockTree::Branch0_PointIndices);
+		ExpectedBranchFacade.SetPoints(1, PVMockTree::Branch1_PointIndices);
+		ExpectedBranchFacade.SetPoints(2, PVMockTree::Branch2_PointIndices);
+		ExpectedBranchFacade.SetPoints(3, PVMockTree::Branch3_PointIndices);
+
+		ExpectedBranchFacade.SetParentBranchNumber(0, 0);
+		ExpectedBranchFacade.SetParentBranchNumber(1, PVMockTree::Branch0_BranchNumber);
+		ExpectedBranchFacade.SetParentBranchNumber(2, PVMockTree::Branch1_BranchNumber);
+		ExpectedBranchFacade.SetParentBranchNumber(3, PVMockTree::Branch0_BranchNumber);
+
+		ExpectedBranchFacade.SetParents(0, { });
+		ExpectedBranchFacade.SetParents(1, { PVMockTree::Branch0_BranchNumber });
+		ExpectedBranchFacade.SetParents(2, { PVMockTree::Branch0_BranchNumber, PVMockTree::Branch1_BranchNumber });
+		ExpectedBranchFacade.SetParents(3, { PVMockTree::Branch0_BranchNumber });
+
+		ExpectedBranchFacade.SetChildren(0, { PVMockTree::Branch1_BranchNumber, PVMockTree::Branch3_BranchNumber });
+		ExpectedBranchFacade.SetChildren(1, { PVMockTree::Branch2_BranchNumber });
+		ExpectedBranchFacade.SetChildren(2, { });
+		ExpectedBranchFacade.SetChildren(3, { });
+	}
+
+	FManagedArrayCollection ActualCollection;
+	PV::Import::Helper::CreateEmptyGrowthData(ActualCollection, NumBranchPoints, NumFoliagePoints, NumBranches, NumFoliageMeshes);
+
+	UTEST_EQUAL("CorrectNrOfPointsInGroup", ActualCollection.NumElements(PV::GroupNames::PointGroup), NumBranchPoints);
+	UTEST_EQUAL("CorrectNrOfBranchesInGroup", ActualCollection.NumElements(PV::GroupNames::BranchGroup), NumBranches);
+
+	PV::Import::Helper::FillGrowthDataFromBranchHierarchyDescription(ActualCollection, { BranchHierarchy });
+
+	const PV::Facades::FPointFacade ActualPointFacade(ActualCollection);
+	const PV::Facades::FBranchFacade ActualBranchFacade(ActualCollection);
+
+	// Validate that the points are correctly indexed and in the correct location
+	UTEST_TRUE("PointPosition0", ActualPointFacade.GetPosition(0).Equals(ExpectedPointFacade.GetPosition(0)));
+	UTEST_TRUE("PointPosition1", ActualPointFacade.GetPosition(1).Equals(ExpectedPointFacade.GetPosition(1)));
+	UTEST_TRUE("PointPosition2", ActualPointFacade.GetPosition(2).Equals(ExpectedPointFacade.GetPosition(2)));
+	UTEST_TRUE("PointPosition3", ActualPointFacade.GetPosition(3).Equals(ExpectedPointFacade.GetPosition(3)));
+	UTEST_TRUE("PointPosition4", ActualPointFacade.GetPosition(4).Equals(ExpectedPointFacade.GetPosition(4)));
+	UTEST_TRUE("PointPosition5", ActualPointFacade.GetPosition(5).Equals(ExpectedPointFacade.GetPosition(5)));
+	UTEST_TRUE("PointPosition6", ActualPointFacade.GetPosition(6).Equals(ExpectedPointFacade.GetPosition(6)));
+	UTEST_TRUE("PointPosition7", ActualPointFacade.GetPosition(7).Equals(ExpectedPointFacade.GetPosition(7)));
+	UTEST_TRUE("PointPosition8", ActualPointFacade.GetPosition(8).Equals(ExpectedPointFacade.GetPosition(8)));
+	UTEST_TRUE("PointPosition9", ActualPointFacade.GetPosition(9).Equals(ExpectedPointFacade.GetPosition(9)));
+	UTEST_TRUE("PointPosition10", ActualPointFacade.GetPosition(10).Equals(ExpectedPointFacade.GetPosition(10)));
+	UTEST_TRUE("PointPosition11", ActualPointFacade.GetPosition(11).Equals(ExpectedPointFacade.GetPosition(11)));
+
+	// Validate that the branch point indices are correct
+	UTEST_EQUAL("BranchPoints0", ActualBranchFacade.GetPoints(0), ExpectedBranchFacade.GetPoints(0));
+	UTEST_EQUAL("BranchPoints1", ActualBranchFacade.GetPoints(1), ExpectedBranchFacade.GetPoints(1));
+	UTEST_EQUAL("BranchPoints2", ActualBranchFacade.GetPoints(2), ExpectedBranchFacade.GetPoints(2));
+	UTEST_EQUAL("BranchPoints3", ActualBranchFacade.GetPoints(3), ExpectedBranchFacade.GetPoints(3));
+
+	// Validate that the branch parents are correct
+	UTEST_EQUAL("ParentBranchNumber0", ActualBranchFacade.GetParentBranchNumber(0), ExpectedBranchFacade.GetParentBranchNumber(0));
+	UTEST_EQUAL("ParentBranchNumber1", ActualBranchFacade.GetParentBranchNumber(1), ExpectedBranchFacade.GetParentBranchNumber(1));
+	UTEST_EQUAL("ParentBranchNumber2", ActualBranchFacade.GetParentBranchNumber(2), ExpectedBranchFacade.GetParentBranchNumber(2));
+	UTEST_EQUAL("ParentBranchNumber3", ActualBranchFacade.GetParentBranchNumber(3), ExpectedBranchFacade.GetParentBranchNumber(3));
+	UTEST_EQUAL("BranchParents0", ActualBranchFacade.GetParents(0), ExpectedBranchFacade.GetParents(0));
+	UTEST_EQUAL("BranchParents1", ActualBranchFacade.GetParents(1), ExpectedBranchFacade.GetParents(1));
+	UTEST_EQUAL("BranchParents2", ActualBranchFacade.GetParents(2), ExpectedBranchFacade.GetParents(2));
+	UTEST_EQUAL("BranchParents3", ActualBranchFacade.GetParents(3), ExpectedBranchFacade.GetParents(3));
+
+	return true;
+}
+
+PV_IMPORTER_TEST(ComputeLengthFromRoot)
+{
+	const FPVBranchHierarchyDescription BranchHierarchy = PVMockTree::CreateBranchHierarchyDescription();
+
+	const int32 NumBranchPoints = BranchHierarchy.GetNumBranchPoints();
+	const int32 NumFoliagePoints = 0;
+	const int32 NumBranches = BranchHierarchy.GetNumBranches();
+	const int32 NumFoliageMeshes = 0;
+
+	FManagedArrayCollection Collection;
+	PV::Import::Helper::CreateEmptyGrowthData(Collection, NumBranchPoints, NumFoliagePoints, NumBranches, NumFoliageMeshes);
+	PV::Import::Helper::FillGrowthDataFromBranchHierarchyDescription(Collection, { BranchHierarchy });
+
+	PV::Facades::FPointFacade PointFacade(Collection);
+	const PV::Facades::FBranchFacade BranchFacade(Collection);
+	const PV::Facades::FPlantFacade PlantFacade(Collection);
+	PV::Import::Helper::ComputeLengthFromRoot(PointFacade, BranchFacade, PlantFacade);
+
+	UTEST_NEARLY_EQUAL("Point0_LengthFromRoot", PointFacade.GetLengthFromRoot(0), PVMockTree::Point0_LengthFromRoot, 0.01f);
+	UTEST_NEARLY_EQUAL("Point1_LengthFromRoot", PointFacade.GetLengthFromRoot(1), PVMockTree::Point1_LengthFromRoot, 0.01f);
+	UTEST_NEARLY_EQUAL("Point2_LengthFromRoot", PointFacade.GetLengthFromRoot(2), PVMockTree::Point2_LengthFromRoot, 0.01f);
+	UTEST_NEARLY_EQUAL("Point3_LengthFromRoot", PointFacade.GetLengthFromRoot(3), PVMockTree::Point3_LengthFromRoot, 0.01f);
+	UTEST_NEARLY_EQUAL("Point4_LengthFromRoot", PointFacade.GetLengthFromRoot(4), PVMockTree::Point4_LengthFromRoot, 0.01f);
+	UTEST_NEARLY_EQUAL("Point5_LengthFromRoot", PointFacade.GetLengthFromRoot(5), PVMockTree::Point5_LengthFromRoot, 0.01f);
+	UTEST_NEARLY_EQUAL("Point6_LengthFromRoot", PointFacade.GetLengthFromRoot(6), PVMockTree::Point6_LengthFromRoot, 0.01f);
+	UTEST_NEARLY_EQUAL("Point7_LengthFromRoot", PointFacade.GetLengthFromRoot(7), PVMockTree::Point7_LengthFromRoot, 0.01f);
+	UTEST_NEARLY_EQUAL("Point8_LengthFromRoot", PointFacade.GetLengthFromRoot(8), PVMockTree::Point8_LengthFromRoot, 0.01f);
+	UTEST_NEARLY_EQUAL("Point9_LengthFromRoot", PointFacade.GetLengthFromRoot(9), PVMockTree::Point9_LengthFromRoot, 0.01f);
+	UTEST_NEARLY_EQUAL("Point10_LengthFromRoot", PointFacade.GetLengthFromRoot(10), PVMockTree::Point10_LengthFromRoot, 0.01f);
+	UTEST_NEARLY_EQUAL("Point11_LengthFromRoot", PointFacade.GetLengthFromRoot(11), PVMockTree::Point11_LengthFromRoot, 0.01f);
+
+	return true;
+}
+
+PV_IMPORTER_TEST(ComputeBudDevelopment)
+{
+	struct FBudDevelopmentDataLayout
+	{
+		int32 Generatation;
+		int32 Age;
+		int32 Branch_Age;
+		int32 AgeSenescence;
+		int32 Light_Senescence;
+		int32 Relative_Age;
+	};
+
+	const FBudDevelopmentDataLayout Point0_ExpectedBudDevelopment = {
+		.Generatation = 1, 
+		.Age = 1,
+		.Branch_Age = 1,
+		.AgeSenescence = 0,
+		.Light_Senescence = 0,
+		.Relative_Age = 1
+	};
+	const FBudDevelopmentDataLayout Point1_ExpectedBudDevelopment = {
+		.Generatation = 1,
+		.Age = 2,
+		.Branch_Age = 1,
+		.AgeSenescence = 0,
+		.Light_Senescence = 0,
+		.Relative_Age = 2
+	};
+	const FBudDevelopmentDataLayout Point2_ExpectedBudDevelopment = {
+		.Generatation = 1,
+		.Age = 3,
+		.Branch_Age = 1,
+		.AgeSenescence = 0,
+		.Light_Senescence = 0,
+		.Relative_Age = 3
+	};
+	const FBudDevelopmentDataLayout Point3_ExpectedBudDevelopment = {
+		.Generatation = 1,
+		.Age = 4,
+		.Branch_Age = 1,
+		.AgeSenescence = 0,
+		.Light_Senescence = 0,
+		.Relative_Age = 4
+	};
+	const FBudDevelopmentDataLayout Point4_ExpectedBudDevelopment = {
+		.Generatation = 1,
+		.Age = 5,
+		.Branch_Age = 1,
+		.AgeSenescence = 0,
+		.Light_Senescence = 0,
+		.Relative_Age = 5
+	};
+	const FBudDevelopmentDataLayout Point5_ExpectedBudDevelopment = {
+		.Generatation = 1,
+		.Age = 6,
+		.Branch_Age = 1,
+		.AgeSenescence = 0,
+		.Light_Senescence = 0,
+		.Relative_Age = 6
+	};
+	const FBudDevelopmentDataLayout Point6_ExpectedBudDevelopment = {
+		.Generatation = 2,
+		.Age = 5,
+		.Branch_Age = 4,
+		.AgeSenescence = 0,
+		.Light_Senescence = 0,
+		.Relative_Age = 5
+	};
+	const FBudDevelopmentDataLayout Point7_ExpectedBudDevelopment = {
+		.Generatation = 2,
+		.Age = 4,
+		.Branch_Age = 3,
+		.AgeSenescence = 0,
+		.Light_Senescence = 0,
+		.Relative_Age = 4
+	};
+	const FBudDevelopmentDataLayout Point8_ExpectedBudDevelopment = {
+		.Generatation = 2,
+		.Age = 5,
+		.Branch_Age = 3,
+		.AgeSenescence = 0,
+		.Light_Senescence = 0,
+		.Relative_Age = 5
+	};
+	const FBudDevelopmentDataLayout Point9_ExpectedBudDevelopment = {
+		.Generatation = 2,
+		.Age = 6,
+		.Branch_Age = 3,
+		.AgeSenescence = 0,
+		.Light_Senescence = 0,
+		.Relative_Age = 6
+	};
+	const FBudDevelopmentDataLayout Point10_ExpectedBudDevelopment = {
+		.Generatation = 3,
+		.Age = 6,
+		.Branch_Age = 5,
+		.AgeSenescence = 0,
+		.Light_Senescence = 0,
+		.Relative_Age = 6
+	};
+	const FBudDevelopmentDataLayout Point11_ExpectedBudDevelopment = {
+		.Generatation = 3,
+		.Age = 7,
+		.Branch_Age = 5,
+		.AgeSenescence = 0,
+		.Light_Senescence = 0,
+		.Relative_Age = 7
+	};
+
+	const int32 ExpectedNumBudDevelopmentValues = 6;
+
+	const FPVBranchHierarchyDescription BranchHierarchy = PVMockTree::CreateBranchHierarchyDescription();
+	const int32 NumBranchPoints = BranchHierarchy.GetNumBranchPoints();
+	const int32 NumFoliagePoints = 0;
+	const int32 NumBranches = BranchHierarchy.GetNumBranches();
+	const int32 NumFoliageMeshes = 0;
+
+	FManagedArrayCollection Collection;
+	PV::Import::Helper::CreateEmptyGrowthData(Collection, NumBranchPoints, NumFoliagePoints, NumBranches, NumFoliageMeshes);
+	PV::Import::Helper::FillGrowthDataFromBranchHierarchyDescription(Collection, { BranchHierarchy });
+
+	PV::Facades::FPointFacade PointFacade(Collection);
+	const PV::Facades::FBranchFacade BranchFacade(Collection);
+	const PV::Facades::FPlantFacade PlantFacade(Collection);
+	PV::Import::Helper::ComputeBudDevelopment(PointFacade, BranchFacade, PlantFacade);
+
+	UTEST_TRUE("Point0_BudDevelopment_Num", PointFacade.GetBudDevelopment(0).Num() == ExpectedNumBudDevelopmentValues);
+	UTEST_TRUE("Point1_BudDevelopment_Num", PointFacade.GetBudDevelopment(1).Num() == ExpectedNumBudDevelopmentValues);
+	UTEST_TRUE("Point2_BudDevelopment_Num", PointFacade.GetBudDevelopment(2).Num() == ExpectedNumBudDevelopmentValues);
+	UTEST_TRUE("Point3_BudDevelopment_Num", PointFacade.GetBudDevelopment(3).Num() == ExpectedNumBudDevelopmentValues);
+	UTEST_TRUE("Point4_BudDevelopment_Num", PointFacade.GetBudDevelopment(4).Num() == ExpectedNumBudDevelopmentValues);
+	UTEST_TRUE("Point5_BudDevelopment_Num", PointFacade.GetBudDevelopment(5).Num() == ExpectedNumBudDevelopmentValues);
+	UTEST_TRUE("Point6_BudDevelopment_Num", PointFacade.GetBudDevelopment(6).Num() == ExpectedNumBudDevelopmentValues);
+	UTEST_TRUE("Point7_BudDevelopment_Num", PointFacade.GetBudDevelopment(7).Num() == ExpectedNumBudDevelopmentValues);
+	UTEST_TRUE("Point8_BudDevelopment_Num", PointFacade.GetBudDevelopment(8).Num() == ExpectedNumBudDevelopmentValues);
+	UTEST_TRUE("Point9_BudDevelopment_Num", PointFacade.GetBudDevelopment(9).Num() == ExpectedNumBudDevelopmentValues);
+	UTEST_TRUE("Point10_BudDevelopment_Num", PointFacade.GetBudDevelopment(10).Num() == ExpectedNumBudDevelopmentValues);
+	UTEST_TRUE("Point11_BudDevelopment_Num", PointFacade.GetBudDevelopment(11).Num() == ExpectedNumBudDevelopmentValues);
+
+	UTEST_TRUE(
+		"Point0_BudDevelopment", 
+		FMemory::Memcmp((void*)PointFacade.GetBudDevelopment(0).GetData(), (void*)&Point0_ExpectedBudDevelopment, ExpectedNumBudDevelopmentValues * sizeof(int32)) == 0
+	);
+	UTEST_TRUE(
+		"Point1_BudDevelopment",
+		FMemory::Memcmp((void*)PointFacade.GetBudDevelopment(1).GetData(), (void*)&Point1_ExpectedBudDevelopment, ExpectedNumBudDevelopmentValues * sizeof(int32)) == 0
+	);
+	UTEST_TRUE(
+		"Point2_BudDevelopment",
+		FMemory::Memcmp((void*)PointFacade.GetBudDevelopment(2).GetData(), (void*)&Point2_ExpectedBudDevelopment, ExpectedNumBudDevelopmentValues * sizeof(int32)) == 0
+	);
+	UTEST_TRUE(
+		"Point3_BudDevelopment",
+		FMemory::Memcmp((void*)PointFacade.GetBudDevelopment(3).GetData(), (void*)&Point3_ExpectedBudDevelopment, ExpectedNumBudDevelopmentValues * sizeof(int32)) == 0
+	);
+	UTEST_TRUE(
+		"Point4_BudDevelopment",
+		FMemory::Memcmp((void*)PointFacade.GetBudDevelopment(4).GetData(), (void*)&Point4_ExpectedBudDevelopment, ExpectedNumBudDevelopmentValues * sizeof(int32)) == 0
+	);
+	UTEST_TRUE(
+		"Point5_BudDevelopment",
+		FMemory::Memcmp((void*)PointFacade.GetBudDevelopment(5).GetData(), (void*)&Point5_ExpectedBudDevelopment, ExpectedNumBudDevelopmentValues * sizeof(int32)) == 0
+	);
+	UTEST_TRUE(
+		"Point6_BudDevelopment",
+		FMemory::Memcmp((void*)PointFacade.GetBudDevelopment(6).GetData(), (void*)&Point6_ExpectedBudDevelopment, ExpectedNumBudDevelopmentValues * sizeof(int32)) == 0
+	);
+	UTEST_TRUE(
+		"Point7_BudDevelopment",
+		FMemory::Memcmp((void*)PointFacade.GetBudDevelopment(7).GetData(), (void*)&Point7_ExpectedBudDevelopment, ExpectedNumBudDevelopmentValues * sizeof(int32)) == 0
+	);
+	UTEST_TRUE(
+		"Point8_BudDevelopment",
+		FMemory::Memcmp((void*)PointFacade.GetBudDevelopment(8).GetData(), (void*)&Point8_ExpectedBudDevelopment, ExpectedNumBudDevelopmentValues * sizeof(int32)) == 0
+	);
+	UTEST_TRUE(
+		"Point9_BudDevelopment",
+		FMemory::Memcmp((void*)PointFacade.GetBudDevelopment(9).GetData(), (void*)&Point9_ExpectedBudDevelopment, ExpectedNumBudDevelopmentValues * sizeof(int32)) == 0
+	);
+	UTEST_TRUE(
+		"Point10_BudDevelopment",
+		FMemory::Memcmp((void*)PointFacade.GetBudDevelopment(10).GetData(), (void*)&Point10_ExpectedBudDevelopment, ExpectedNumBudDevelopmentValues * sizeof(int32)) == 0
+	);
+
+	UTEST_TRUE(
+		"Point11_BudDevelopment",
+		FMemory::Memcmp((void*)PointFacade.GetBudDevelopment(11).GetData(), (void*)&Point11_ExpectedBudDevelopment, ExpectedNumBudDevelopmentValues * sizeof(int32)) == 0
+	);
+
+	return true;
+}
+
+PV_IMPORTER_TEST(ComputeBudDirections)
+{
+	const FVector3f ExpectedApical0 = (PVMockTree::Point1 - PVMockTree::Point0).GetSafeNormal();
+	const FVector3f ExpectedApical1 = (PVMockTree::Point2 - PVMockTree::Point1).GetSafeNormal();
+	const FVector3f ExpectedApical2 = (PVMockTree::Point3 - PVMockTree::Point2).GetSafeNormal();
+	const FVector3f ExpectedApical3 = (PVMockTree::Point4 - PVMockTree::Point3).GetSafeNormal();
+	const FVector3f ExpectedApical4 = (PVMockTree::Point5 - PVMockTree::Point4).GetSafeNormal();
+	const FVector3f ExpectedApical5 = (PVMockTree::Point5 - PVMockTree::Point4).GetSafeNormal();
+	const FVector3f ExpectedApical6 = (PVMockTree::Point6 - PVMockTree::Point3).GetSafeNormal();
+	const FVector3f ExpectedApical7 = (PVMockTree::Point8 - PVMockTree::Point7).GetSafeNormal();
+	const FVector3f ExpectedApical8 = (PVMockTree::Point9 - PVMockTree::Point8).GetSafeNormal();
+	const FVector3f ExpectedApical9 = (PVMockTree::Point9 - PVMockTree::Point8).GetSafeNormal();
+	const FVector3f ExpectedApical10 = (PVMockTree::Point11 - PVMockTree::Point10).GetSafeNormal();
+	const FVector3f ExpectedApical11 = (PVMockTree::Point11 - PVMockTree::Point10).GetSafeNormal();
+
+	const int32 ExpectedNumBudDirectionValues = 6;
+
+	const FPVBranchHierarchyDescription BranchHierarchy = PVMockTree::CreateBranchHierarchyDescription();
+	const int32 NumBranchPoints = BranchHierarchy.GetNumBranchPoints();
+	const int32 NumFoliagePoints = 0;
+	const int32 NumBranches = BranchHierarchy.GetNumBranches();
+	const int32 NumFoliageMeshes = 0;
+
+	FManagedArrayCollection Collection;
+	PV::Import::Helper::CreateEmptyGrowthData(Collection, NumBranchPoints, NumFoliagePoints, NumBranches, NumFoliageMeshes);
+	PV::Import::Helper::FillGrowthDataFromBranchHierarchyDescription(Collection, { BranchHierarchy });
+
+	PV::Facades::FPointFacade PointFacade(Collection);
+	const PV::Facades::FBranchFacade BranchFacade(Collection);
+	const PV::Facades::FPlantFacade PlantFacade(Collection);
+	PV::Import::Helper::ComputeBudDirections(PointFacade, BranchFacade, PlantFacade);
+
+	UTEST_TRUE("Point0_BudDirections_Num", PointFacade.GetBudDirection(0).Num() == ExpectedNumBudDirectionValues);
+	UTEST_TRUE("Point1_BudDirections_Num", PointFacade.GetBudDirection(1).Num() == ExpectedNumBudDirectionValues);
+	UTEST_TRUE("Point2_BudDirections_Num", PointFacade.GetBudDirection(2).Num() == ExpectedNumBudDirectionValues);
+	UTEST_TRUE("Point3_BudDirections_Num", PointFacade.GetBudDirection(3).Num() == ExpectedNumBudDirectionValues);
+	UTEST_TRUE("Point4_BudDirections_Num", PointFacade.GetBudDirection(4).Num() == ExpectedNumBudDirectionValues);
+	UTEST_TRUE("Point5_BudDirections_Num", PointFacade.GetBudDirection(5).Num() == ExpectedNumBudDirectionValues);
+	UTEST_TRUE("Point6_BudDirections_Num", PointFacade.GetBudDirection(6).Num() == ExpectedNumBudDirectionValues);
+	UTEST_TRUE("Point7_BudDirections_Num", PointFacade.GetBudDirection(7).Num() == ExpectedNumBudDirectionValues);
+	UTEST_TRUE("Point8_BudDirections_Num", PointFacade.GetBudDirection(8).Num() == ExpectedNumBudDirectionValues);
+	UTEST_TRUE("Point9_BudDirections_Num", PointFacade.GetBudDirection(9).Num() == ExpectedNumBudDirectionValues);
+	UTEST_TRUE("Point10_BudDirections_Num", PointFacade.GetBudDirection(10).Num() == ExpectedNumBudDirectionValues);
+	UTEST_TRUE("Point11_BudDirections_Num", PointFacade.GetBudDirection(11).Num() == ExpectedNumBudDirectionValues);
+
+	UTEST_NEARLY_EQUAL("Point0_Apical", FVector(PointFacade.GetBudDirection(0)[0]), FVector(ExpectedApical0), UE_KINDA_SMALL_NUMBER);
+	UTEST_NEARLY_EQUAL("Point1_Apical", FVector(PointFacade.GetBudDirection(1)[0]), FVector(ExpectedApical1), UE_KINDA_SMALL_NUMBER);
+	UTEST_NEARLY_EQUAL("Point2_Apical", FVector(PointFacade.GetBudDirection(2)[0]), FVector(ExpectedApical2), UE_KINDA_SMALL_NUMBER);
+	UTEST_NEARLY_EQUAL("Point3_Apical", FVector(PointFacade.GetBudDirection(3)[0]), FVector(ExpectedApical3), UE_KINDA_SMALL_NUMBER);
+	UTEST_NEARLY_EQUAL("Point4_Apical", FVector(PointFacade.GetBudDirection(4)[0]), FVector(ExpectedApical4), UE_KINDA_SMALL_NUMBER);
+	UTEST_NEARLY_EQUAL("Point5_Apical", FVector(PointFacade.GetBudDirection(5)[0]), FVector(ExpectedApical5), UE_KINDA_SMALL_NUMBER);
+	UTEST_NEARLY_EQUAL("Point6_Apical", FVector(PointFacade.GetBudDirection(6)[0]), FVector(ExpectedApical6), UE_KINDA_SMALL_NUMBER);
+	UTEST_NEARLY_EQUAL("Point7_Apical", FVector(PointFacade.GetBudDirection(7)[0]), FVector(ExpectedApical7), UE_KINDA_SMALL_NUMBER);
+	UTEST_NEARLY_EQUAL("Point8_Apical", FVector(PointFacade.GetBudDirection(8)[0]), FVector(ExpectedApical8), UE_KINDA_SMALL_NUMBER);
+	UTEST_NEARLY_EQUAL("Point9_Apical", FVector(PointFacade.GetBudDirection(9)[0]), FVector(ExpectedApical9), UE_KINDA_SMALL_NUMBER);
+	UTEST_NEARLY_EQUAL("Point10_Apical", FVector(PointFacade.GetBudDirection(10)[0]), FVector(ExpectedApical10), UE_KINDA_SMALL_NUMBER);
+	UTEST_NEARLY_EQUAL("Point11_Apical", FVector(PointFacade.GetBudDirection(11)[0]), FVector(ExpectedApical11), UE_KINDA_SMALL_NUMBER);
+
+	// TODO: Test UpVector and Axillary
+
+	return true;
+}
+
+PV_IMPORTER_TEST(SplitFusedPoints)
+{
+	const FPVBranchHierarchyDescription BranchHierarchy = PVMockTree::CreateBranchHierarchyDescription();
+	const int32 NumBranchPoints = BranchHierarchy.GetNumBranchPoints();
+	const int32 NumFoliagePoints = 0;
+	const int32 NumBranches = BranchHierarchy.GetNumBranches();
+	const int32 NumFoliageMeshes = 0;
+
+	FManagedArrayCollection Collection;
+	PV::Import::Helper::CreateEmptyGrowthData(Collection, NumBranchPoints, NumFoliagePoints, NumBranches, NumFoliageMeshes);
+	PV::Import::Helper::FillGrowthDataFromBranchHierarchyDescription(Collection, { BranchHierarchy });
+
+	PV::Facades::FPointFacade PointFacade(Collection);
+	const PV::Facades::FBranchFacade BranchFacade(Collection);
+	const PV::Facades::FPlantFacade PlantFacade(Collection);
+	PV::Import::Helper::ComputeBudDirections(PointFacade, BranchFacade, PlantFacade);
+	PV::Import::Helper::ComputeLengthFromRoot(PointFacade, BranchFacade, PlantFacade);
+	PV::Import::Helper::ComputeBudDevelopment(PointFacade, BranchFacade, PlantFacade);
+	PV::Import::Helper::ComputeBudDirections(PointFacade, BranchFacade, PlantFacade);
+
+	const bool bSplitFusedPointsResult = PV::Import::Helper::SplitFusedPoints(Collection);
+	UTEST_TRUE("SplitFusedPointsSucceeded", bSplitFusedPointsResult);
+
+	const int32 Branch0_OldRootPointIndex = 0;
+	const int32 Branch1_OldRootPointIndex = 2;
+	const int32 Branch2_OldRootPointIndex = 8;
+	const int32 Branch3_OldRootPointIndex = 3;
+
+	const int32 Branch0_NewRootPointIndex = BranchFacade.GetPoints(0)[0];
+	const int32 Branch1_NewRootPointIndex = BranchFacade.GetPoints(1)[0];
+	const int32 Branch2_NewRootPointIndex = BranchFacade.GetPoints(2)[0];
+	const int32 Branch3_NewRootPointIndex = BranchFacade.GetPoints(3)[0];
+
+	UTEST_TRUE("Branch0_HasSamePointIndex",  Branch0_NewRootPointIndex == Branch0_OldRootPointIndex);
+	UTEST_TRUE("Branch1_HasUniqePointIndex", Branch1_NewRootPointIndex != Branch1_OldRootPointIndex && Branch1_NewRootPointIndex > 11);
+	UTEST_TRUE("Branch2_HasUniqePointIndex", Branch2_NewRootPointIndex != Branch2_OldRootPointIndex && Branch2_NewRootPointIndex > 11);
+	UTEST_TRUE("Branch3_HasUniqePointIndex", Branch3_NewRootPointIndex != Branch3_OldRootPointIndex && Branch3_NewRootPointIndex > 11);
+
+	UTEST_EQUAL("Branch0_NewPointPositionIsSame", PointFacade.GetPosition(Branch0_NewRootPointIndex), PointFacade.GetPosition(Branch0_OldRootPointIndex));
+	UTEST_EQUAL("Branch1_NewPointPositionIsSame", PointFacade.GetPosition(Branch1_NewRootPointIndex), PointFacade.GetPosition(Branch1_OldRootPointIndex));
+	UTEST_EQUAL("Branch2_NewPointPositionIsSame", PointFacade.GetPosition(Branch2_NewRootPointIndex), PointFacade.GetPosition(Branch2_OldRootPointIndex));
+	UTEST_EQUAL("Branch3_NewPointPositionIsSame", PointFacade.GetPosition(Branch3_NewRootPointIndex), PointFacade.GetPosition(Branch3_OldRootPointIndex));
+
+	UTEST_EQUAL("Branch0_NewPointLengthFromRootIsSame", PointFacade.GetLengthFromRoot(Branch0_NewRootPointIndex), PointFacade.GetLengthFromRoot(Branch0_OldRootPointIndex));
+	UTEST_EQUAL("Branch1_NewPointLengthFromRootIsSame", PointFacade.GetLengthFromRoot(Branch1_NewRootPointIndex), PointFacade.GetLengthFromRoot(Branch1_OldRootPointIndex));
+	UTEST_EQUAL("Branch2_NewPointLengthFromRootIsSame", PointFacade.GetLengthFromRoot(Branch2_NewRootPointIndex), PointFacade.GetLengthFromRoot(Branch2_OldRootPointIndex));
+	UTEST_EQUAL("Branch3_NewPointLengthFromRootIsSame", PointFacade.GetLengthFromRoot(Branch3_NewRootPointIndex), PointFacade.GetLengthFromRoot(Branch3_OldRootPointIndex));
+
+	UTEST_EQUAL("Branch0_NewPointLengthFromSeedIsSame", PointFacade.GetLengthFromSeed(Branch0_NewRootPointIndex), PointFacade.GetLengthFromSeed(Branch0_OldRootPointIndex));
+	UTEST_EQUAL("Branch1_NewPointLengthFromSeedIsSame", PointFacade.GetLengthFromSeed(Branch1_NewRootPointIndex), PointFacade.GetLengthFromSeed(Branch1_OldRootPointIndex));
+	UTEST_EQUAL("Branch2_NewPointLengthFromSeedIsSame", PointFacade.GetLengthFromSeed(Branch2_NewRootPointIndex), PointFacade.GetLengthFromSeed(Branch2_OldRootPointIndex));
+	UTEST_EQUAL("Branch3_NewPointLengthFromSeedIsSame", PointFacade.GetLengthFromSeed(Branch3_NewRootPointIndex), PointFacade.GetLengthFromSeed(Branch3_OldRootPointIndex));
+
+	return true;
+}
+
+PV_IMPORTER_TEST(FindRootBranch)
+{
+	const TArray<TArray<int32>>& Branches = PVMockTree::BranchPointIndices;
+	const int32 RootBranchIndex = PV::Import::Helper::FindRootBranch(PVMockTree::Points, Branches);
+
+	UTEST_NOT_EQUAL("ValidRootBranch", RootBranchIndex, int32(INDEX_NONE));
+
+	const TArray<int32> RootBranchIndices = Branches[RootBranchIndex];
+	const FVector3f RootVertex = PVMockTree::Points[RootBranchIndices[0]];
+
+	UTEST_NEARLY_EQUAL("CorrectRootBranchLocation", FVector(RootVertex), FVector::ZeroVector, UE_KINDA_SMALL_NUMBER);
+
+	return true;
+}
+
+PV_IMPORTER_TEST(EstimateBranchHierarchy)
+{
+	// Test standard MockTree
+	{
+		TArray<TArray<int32>> Branches = {
+			PVMockTree::Branch0_PointIndices,
+			PVMockTree::Branch1_PointIndices,
+			PVMockTree::Branch2_PointIndices,
+			PVMockTree::Branch3_PointIndices
+		};
+
+		TArray<int32> ExpectedBranchHierarchy = { -1, 0, 1, 0 };
+
+		// Naive testing of hierarchy
+		{
+			const int32 RootBranchIndex = PV::Import::Helper::FindRootBranch(PVMockTree::Points, Branches);
+			const TArray<int32> BranchHierarchy = PV::Import::Helper::EstimateBranchHierarchy(PVMockTree::Points, Branches, RootBranchIndex);
+
+			UTEST_EQUAL("CorrectNrOfBranches", BranchHierarchy.Num(), ExpectedBranchHierarchy.Num());
+			for (int32 i = 0; i < Branches.Num(); ++i)
+			{
+				UTEST_EQUAL(FString::Printf(TEXT("Branch%d_CorrectParent"), i), BranchHierarchy[i], ExpectedBranchHierarchy[i]);
+			}
+		}
+
+		// Test with branch list reversed
+		{
+			Algo::Reverse(Branches);
+			Algo::Reverse(ExpectedBranchHierarchy);
+			for (int32& Hierarchy : ExpectedBranchHierarchy)
+			{
+				if (Hierarchy != -1)
+				{
+					Hierarchy = Branches.Num() - 1 - Hierarchy;
+				}
+			}
+
+			const int32 RootBranchIndex = PV::Import::Helper::FindRootBranch(PVMockTree::Points, Branches);
+			const TArray<int32> BranchHierarchy = PV::Import::Helper::EstimateBranchHierarchy(PVMockTree::Points, Branches, RootBranchIndex);
+			
+			UTEST_EQUAL("CorrectNrOfBranches_Reverse", BranchHierarchy.Num(), ExpectedBranchHierarchy.Num());
+			for (int32 i = 0; i < Branches.Num(); ++i)
+			{
+				UTEST_EQUAL(FString::Printf(TEXT("Branch%d_CorrectParent_Reverse"), i), BranchHierarchy[i], ExpectedBranchHierarchy[i]);
+			}
+		}
+	}
+
+	// Add additional branches to MockTree and ensure hierarchy stays correct
+	{
+		// Add a few branches starting from the same point as another as we want to make sure the algorithm can handle multiple 
+		// branches being at the same distance from the start point of another.
+		TArray<FVector3f> PointsCpy = PVMockTree::Points;
+		const int32 Point2_1_Index = PointsCpy.Add(PVMockTree::Point2 + FVector3f(-10, 0, 0));
+		const int32 Point2_2_Index = PointsCpy.Add(PVMockTree::Point2 + FVector3f(10, 0, 0));
+		const int32 Point2_3_Index = PointsCpy.Add(PVMockTree::Point2 + FVector3f(-10, 10, 0));
+		const int32 Point2_4_Index = PointsCpy.Add(PVMockTree::Point2 + FVector3f(10, 10, 0));
+		const int32 Point2_5_Index = PointsCpy.Add(PVMockTree::Point2 + FVector3f(-10, 10, -10));
+		const int32 Point2_6_Index = PointsCpy.Add(PVMockTree::Point2 + FVector3f(10, 8, 1));
+
+		const int32 Point8_1_Index = PointsCpy.Add(PVMockTree::Point8 + FVector3f(10, 0, 0));
+		const int32 Point8_2_Index = PointsCpy.Add(PVMockTree::Point8 + FVector3f(-10, 0, 0));
+		const int32 Point8_3_Index = PointsCpy.Add(PVMockTree::Point8 + FVector3f(10, 10, 0));
+		const int32 Point8_4_Index = PointsCpy.Add(PVMockTree::Point8 + FVector3f(-10, 10, 0));
+		const int32 Point8_5_Index = PointsCpy.Add(PVMockTree::Point8 + FVector3f(10, 10, 10));
+		const int32 Point8_6_Index = PointsCpy.Add(PVMockTree::Point8 + FVector3f(-10, 10, 10));
+
+		const int32 Point3_0_Index = PointsCpy.Add(PVMockTree::Point3 + FVector3f(10, 0, 0));
+		const int32 Point3_1_Index = PointsCpy.Add(PVMockTree::Point3 + FVector3f(-10, 0, 0));
+
+		TArray<TArray<int32>> Branches = {
+			PVMockTree::Branch0_PointIndices,
+			PVMockTree::Branch1_PointIndices,
+			PVMockTree::Branch2_PointIndices,
+			PVMockTree::Branch3_PointIndices,
+			{ 2, Point2_1_Index },
+			{ 2, Point2_2_Index },
+			{ 2, Point2_3_Index },
+			{ 2, Point2_4_Index },
+			{ 2, Point2_5_Index },
+			{ 2, Point2_6_Index },
+
+			{ 8, Point8_1_Index },
+			{ 8, Point8_2_Index },
+			{ 8, Point8_3_Index },
+			{ 8, Point8_4_Index },
+			{ 8, Point8_5_Index },
+			{ 8, Point8_6_Index },
+
+			{ 3, Point3_0_Index },
+			{ 3, Point3_1_Index },
+		};
+
+		TArray<int32> ExpectedBranchHierarchy = { 
+			-1, 0, 1, 0, 
+			0, 0, 0, 0, 0, 0,
+			1, 1, 1, 1, 1, 1,
+			0, 0
+		};
+
+		// Naive testing of hierarchy
+		{
+			const int32 RootBranchIndex = PV::Import::Helper::FindRootBranch(PointsCpy, Branches);
+			const TArray<int32> BranchHierarchy = PV::Import::Helper::EstimateBranchHierarchy(PointsCpy, Branches, RootBranchIndex);
+
+			UTEST_EQUAL("CorrectNrOfBranches", BranchHierarchy.Num(), ExpectedBranchHierarchy.Num());
+			for (int32 i = 0; i < Branches.Num(); ++i)
+			{
+				UTEST_EQUAL(FString::Printf(TEXT("Branch%d_CorrectParent"), i), BranchHierarchy[i], ExpectedBranchHierarchy[i]);
+			}
+		}
+
+		// Test with branch list reversed
+		{
+			Algo::Reverse(Branches);
+			Algo::Reverse(ExpectedBranchHierarchy);
+			for (int32& Hierarchy : ExpectedBranchHierarchy)
+			{
+				if (Hierarchy != -1)
+				{
+					Hierarchy = Branches.Num() - 1 - Hierarchy;
+				}
+			}
+
+			const int32 RootBranchIndex = PV::Import::Helper::FindRootBranch(PointsCpy, Branches);
+			const TArray<int32> BranchHierarchy = PV::Import::Helper::EstimateBranchHierarchy(PointsCpy, Branches, RootBranchIndex);
+
+			UTEST_EQUAL("CorrectNrOfBranches_Reverse", BranchHierarchy.Num(), ExpectedBranchHierarchy.Num());
+			for (int32 i = 0; i < Branches.Num(); ++i)
+			{
+				UTEST_EQUAL(FString::Printf(TEXT("Branch%d_CorrectParent_Reverse"), i), BranchHierarchy[i], ExpectedBranchHierarchy[i]);
+			}
+		}
+	}
+
+	return true;
+}
+
+PV_IMPORTER_TEST(CreateBranchHierarchyFromEdges)
+{
+	FPVBranchHierarchyDescription ExpectedBranchHierarchy;
+	ExpectedBranchHierarchy.RootBranchIndex = 0;
+	ExpectedBranchHierarchy.Points = PVMockTree::Points;
+
+	FPVBranchDescription& B0 = ExpectedBranchHierarchy.Branches.AddDefaulted_GetRef();
+	B0.BranchIndex = 0;
+	B0.PointIndices = PVMockTree::Branch0_PointIndices;
+	B0.ParentBranchIndex = -1;
+	B0.ChildBranchIndices = { 1, 3 };
+
+	FPVBranchDescription& B1 = ExpectedBranchHierarchy.Branches.AddDefaulted_GetRef();
+	B1.BranchIndex = 1;
+	B1.PointIndices = PVMockTree::Branch1_PointIndices;
+	B1.ParentBranchIndex = 0;
+	B1.ChildBranchIndices = { 2 };
+
+	FPVBranchDescription& B2 = ExpectedBranchHierarchy.Branches.AddDefaulted_GetRef();
+	B2.BranchIndex = 2;
+	B2.PointIndices = PVMockTree::Branch2_PointIndices;
+	B2.ParentBranchIndex = 1;
+	B2.ChildBranchIndices = { };
+
+	FPVBranchDescription& B3 = ExpectedBranchHierarchy.Branches.AddDefaulted_GetRef();
+	B3.BranchIndex = 3;
+	B3.PointIndices = PVMockTree::Branch3_PointIndices;
+	B3.ParentBranchIndex = 0;
+	B3.ChildBranchIndices = { };
+
+	const FPVBranchHierarchyDescription BranchHierarchy = PVMockTree::CreateBranchHierarchyDescription();
+
+	UTEST_EQUAL("CorrectRootIndex", BranchHierarchy.RootBranchIndex, ExpectedBranchHierarchy.RootBranchIndex);
+	UTEST_EQUAL("CorrectNrOfBranches", BranchHierarchy.Branches.Num(), ExpectedBranchHierarchy.Branches.Num());
+
+	for (int32 i = 0; i < BranchHierarchy.Branches.Num(); ++i)
+	{
+		const FPVBranchDescription& Branch = BranchHierarchy.Branches[i];
+		const FPVBranchDescription* ExpectedBranchPtr = ExpectedBranchHierarchy.Branches.FindByPredicate([&](const FPVBranchDescription& InBranch) { return InBranch.BranchIndex == Branch.BranchIndex; });
+
+		UTEST_NOT_NULL("ContainsBranchIndex", ExpectedBranchPtr);
+
+		const FPVBranchDescription& ExpectedBranch = *ExpectedBranchPtr;
+
+		const FString BranchName = FString::Printf(TEXT("Branch%d: "), i);
+
+		UTEST_EQUAL(BranchName + "_CorrectBranchIndex", Branch.BranchIndex, ExpectedBranch.BranchIndex);
+		UTEST_EQUAL(BranchName + "_CorrectParent", Branch.ParentBranchIndex, ExpectedBranch.ParentBranchIndex);
+		UTEST_EQUAL(BranchName + "_CorrectNrOfChildren", Branch.ChildBranchIndices.Num(), ExpectedBranch.ChildBranchIndices.Num());
+		UTEST_EQUAL(BranchName + "_CorrectNrOfPoints", Branch.PointIndices.Num(), ExpectedBranch.PointIndices.Num());
+
+		for (int32 j = 0; j < Branch.PointIndices.Num(); ++j)
+		{
+			const int32 PointIndex = Branch.PointIndices[j];
+			const int32 ExpectedPointIndex = ExpectedBranch.PointIndices[j];
+			UTEST_EQUAL(BranchName + FString::Printf(TEXT("_CorrectPointIndex %d"), j), PointIndex, ExpectedPointIndex);
+
+			const FVector Point = FVector(BranchHierarchy.Points[PointIndex]);
+			const FVector ExpectedPoint = FVector(ExpectedBranchHierarchy.Points[ExpectedPointIndex]);
+			UTEST_NEARLY_EQUAL(BranchName + FString::Printf(TEXT("_CorrectPoint %d"), j), Point, ExpectedPoint, UE_KINDA_SMALL_NUMBER);
+		}
+
+		for (int32 j = 0; j < Branch.ChildBranchIndices.Num(); ++j)
+		{
+			UTEST_TRUE(BranchName + FString::Printf(TEXT("_CorrectChildIndex %d"), j), ExpectedBranch.ChildBranchIndices.Contains(Branch.ChildBranchIndices[j]));
+		}
+	}
+
+	return true;
+}
+
+#endif
\ No newline at end of file
