[PVE] Expose parameter panel and enabled basic parameter and math nodes.

#Synced-CL 49920928

diff --git a/Engine/Plugins/Experimental/ProceduralVegetationEditor/Source/ProceduralVegetationEditor/Private/PVEditor.cpp b/Engine/Plugins/Experimental/ProceduralVegetationEditor/Source/ProceduralVegetationEditor/Private/PVEditor.cpp
--- a/Engine/Plugins/Experimental/ProceduralVegetationEditor/Source/ProceduralVegetationEditor/Private/PVEditor.cpp
+++ b/Engine/Plugins/Experimental/ProceduralVegetationEditor/Source/ProceduralVegetationEditor/Private/PVEditor.cpp
@@ -245,6 +245,7 @@ TSharedRef<FTabManager::FLayout> FPVEditor::GetDefaultLayout() const
 					FTabManager::NewStack()
 					->SetSizeCoefficient(0.6f)
 					->AddTab(GetPanelID(EPCGEditorPanel::GraphEditor), ETabState::OpenedTab)
+					->AddTab(GetPanelID(EPCGEditorPanel::UserParams), ETabState::ClosedTab)
 					->SetHideTabWell(true)
 				)
 				->Split
diff --git a/Engine/Plugins/Experimental/ProceduralVegetationEditor/Source/ProceduralVegetationEditor/Private/PVEditorSchema.cpp b/Engine/Plugins/Experimental/ProceduralVegetationEditor/Source/ProceduralVegetationEditor/Private/PVEditorSchema.cpp
--- a/Engine/Plugins/Experimental/ProceduralVegetationEditor/Source/ProceduralVegetationEditor/Private/PVEditorSchema.cpp
+++ b/Engine/Plugins/Experimental/ProceduralVegetationEditor/Source/ProceduralVegetationEditor/Private/PVEditorSchema.cpp
@@ -1,6 +1,9 @@
 ﻿// Copyright Epic Games, Inc. All Rights Reserved.
 
 #include "PVEditorSchema.h"
+#include "PCGSubgraph.h"
+#include "Elements/PCGCreateAttribute.h"
+#include "Elements/Metadata/PCGMetadataMathsOpElement.h"
 #include "Nodes/PVBaseSettings.h"
 
 EPCGElementType UPVEditorSchema::GetElementTypeFiltering() const
@@ -15,7 +18,9 @@ const FPCGGraphEditorCustomization& UPVEditorSchema::GetGraphEditorCustomization
 		.bFilterSettings = true,
 		.FilteredSettingsTypes =
 	{
-			UPVBaseSettings::StaticClass()
+			UPVBaseSettings::StaticClass(),
+			UPCGCreateAttributeSetSettings::StaticClass(),
+			UPCGMetadataMathsSettings::StaticClass()
 		}
 	};
 
