Make it easier to visualize different stages of 2D importer

- Add debug visualization to the 2D importer by default
- Fix a bug where the material for the texture debug quad was garbage collected

#Synced-CL 49960140

diff --git a/Engine/Plugins/Experimental/ProceduralVegetationEditor/Source/ProceduralVegetation/Private/Nodes/PVTexture2DImporterSettings.cpp b/Engine/Plugins/Experimental/ProceduralVegetationEditor/Source/ProceduralVegetation/Private/Nodes/PVTexture2DImporterSettings.cpp
--- a/Engine/Plugins/Experimental/ProceduralVegetationEditor/Source/ProceduralVegetation/Private/Nodes/PVTexture2DImporterSettings.cpp
+++ b/Engine/Plugins/Experimental/ProceduralVegetationEditor/Source/ProceduralVegetation/Private/Nodes/PVTexture2DImporterSettings.cpp
@@ -1,15 +1,82 @@
 // Copyright Epic Games, Inc. All Rights Reserved.
 
 #include "PVTexture2DImporterSettings.h"
+
 #include "IPVEditorInterface.h"
-#include "PCGContext.h"
 #include "ProceduralVegetationModule.h"
 #include "DataTypes/PVGrowthData.h"
 #include "ProceduralVegetationPreset.h"
 #include "PVImportTexture2DParams.h"
 
+#include "PCGContext.h"
+
+#include "Engine/Texture2D.h"
+#include "Materials/MaterialInstanceDynamic.h"
+#include "UObject/ConstructorHelpers.h"
+
 #define LOCTEXT_NAMESPACE "PVTexture2DImporterSettings"
 
+UPVTexture2DImporterSettings::UPVTexture2DImporterSettings()
+{
+#if WITH_EDITORONLY_DATA
+	static ConstructorHelpers::FObjectFinder<UMaterialInterface> TextureVisualizationMaterialRef(TEXT("/Script/Engine.Material'/Engine/Engine_MI_Shaders/M_Shader_SimpleTranslucent.M_Shader_SimpleTranslucent'"));
+	TextureVisualizationMaterial = TextureVisualizationMaterialRef.Object;
+
+	FPVVisualizationSettings& LineSegmentVisualization = DebugVisualizationSettings.VisualizationSettings.AddDefaulted_GetRef();
+	LineSegmentVisualization.DebugType = EPVDebugType::Custom;
+	LineSegmentVisualization.AttributeToFilter = TEXT("LineSegment");
+	LineSegmentVisualization.CustomGroupToFilter = TEXT("LineSegments");
+	LineSegmentVisualization.VisualizationMode = EPVDebugValueVisualizationMode::Curve;
+
+	FPVVisualizationSettings& PerimiterCurveVisualization = DebugVisualizationSettings.VisualizationSettings.AddDefaulted_GetRef();
+	PerimiterCurveVisualization.DebugType = EPVDebugType::Custom;
+	PerimiterCurveVisualization.bShowAnchorPoints = false;
+	PerimiterCurveVisualization.AttributeToFilter = TEXT("PerimiterCurve");
+	PerimiterCurveVisualization.CustomGroupToFilter = TEXT("PerimiterCurves");
+	PerimiterCurveVisualization.VisualizationMode = EPVDebugValueVisualizationMode::Curve;
+
+	FPVVisualizationSettings& EndPointVisualization = DebugVisualizationSettings.VisualizationSettings.AddDefaulted_GetRef();
+	EndPointVisualization.DebugType = EPVDebugType::Custom;
+	EndPointVisualization.bShowAnchorPoints = false;
+	EndPointVisualization.bDrawPointAsMesh = false;
+	EndPointVisualization.bUsePivotAsPosition = true;
+	EndPointVisualization.AttributeToFilter = TEXT("EndPoint");
+	EndPointVisualization.CustomGroupToFilter = TEXT("VertexAttributes");
+	EndPointVisualization.CustomPivotPositionAttributeName = "Vertex";
+	EndPointVisualization.VisualizationMode = EPVDebugValueVisualizationMode::Point;
+	EndPointVisualization.Color = FColor(39, 0, 255);
+	EndPointVisualization.GizmoScale = 8.000000;
+
+	FPVVisualizationSettings& BranchCurvesVisualization = DebugVisualizationSettings.VisualizationSettings.AddDefaulted_GetRef();
+	BranchCurvesVisualization.DebugType = EPVDebugType::Custom;
+	BranchCurvesVisualization.bShowAnchorPoints = false;
+	BranchCurvesVisualization.bDrawPointAsMesh = false;
+	BranchCurvesVisualization.AttributeToFilter = TEXT("Curve");
+	BranchCurvesVisualization.CustomGroupToFilter = TEXT("BranchCurves");
+	BranchCurvesVisualization.VisualizationMode = EPVDebugValueVisualizationMode::Curve;
+
+	FPVVisualizationSettings& RootPointsVisualization = DebugVisualizationSettings.VisualizationSettings.AddDefaulted_GetRef();
+	RootPointsVisualization.DebugType = EPVDebugType::Custom;
+	RootPointsVisualization.bShowAnchorPoints = false;
+	RootPointsVisualization.bDrawPointAsMesh = false;
+	RootPointsVisualization.AttributeToFilter = TEXT("Point");
+	RootPointsVisualization.CustomGroupToFilter = TEXT("InputRootPoints");
+	RootPointsVisualization.VisualizationMode = EPVDebugValueVisualizationMode::Point;
+	RootPointsVisualization.GizmoScale = 10.000000;
+
+	FPVVisualizationSettings& PixelSamplingVisualization = DebugVisualizationSettings.VisualizationSettings.AddDefaulted_GetRef();
+	PixelSamplingVisualization.DebugType = EPVDebugType::Custom;
+	PixelSamplingVisualization.bShowAnchorPoints = false;
+	PixelSamplingVisualization.bDrawPointAsMesh = false;
+	PixelSamplingVisualization.bUsePivotAsPosition = true;
+	PixelSamplingVisualization.AttributeToFilter = TEXT("Value");
+	PixelSamplingVisualization.CustomGroupToFilter = TEXT("PixelDebug");
+	PixelSamplingVisualization.VisualizationMode = EPVDebugValueVisualizationMode::Point;
+	PixelSamplingVisualization.Color = FColor(0, 0, 255);
+	PixelSamplingVisualization.GizmoScale = 3.000000;
+#endif // WITH_EDITORONLY_DATA
+}
+
 #if WITH_EDITOR
 FText UPVTexture2DImporterSettings::GetDefaultNodeTitle() const
 { 
@@ -57,15 +124,32 @@ bool FPVTexture2DImporterElement::ExecuteInternal(FPCGContext* InContext) const
 	{
 		return true;
 	}
-	
-	FManagedArrayCollection Collection;
-	if (!PV::GPVEditorInterface->ImportGrowthDataFromTexture2D(Settings->Params, Collection))
+
+	UMaterialInstanceDynamic* TextureVisualizationMaterialInstance = nullptr;
+
+	if (Settings->TextureVisualizationMaterial)
+	{
+		TextureVisualizationMaterialInstance = UMaterialInstanceDynamic::Create(Settings->TextureVisualizationMaterial, nullptr);
+		TextureVisualizationMaterialInstance->SetFlags(RF_Transient);
+		
+		const FTaskTagScope Scope(ETaskTag::EParallelGameThread); // Needed for SetTextureParameterValue
+		TextureVisualizationMaterialInstance->SetTextureParameterValue(TEXT("Diffuse"), Settings->Params.Texture2DAsset);
+		TextureVisualizationMaterialInstance->SetTextureParameterValue(TEXT("OpacityMask"), Settings->Params.Texture2DAsset);
+	}
+
+	FPVImportTexture2DOutput ImportTexture2DOutput;
+	if (!PV::GPVEditorInterface->ImportGrowthDataFromTexture2D(Settings->Params, TextureVisualizationMaterialInstance, ImportTexture2DOutput))
 	{
 		return false;
 	}
 
 	UPVGrowthData* OutVariantData = FPCGContext::NewObject_AnyThread<UPVGrowthData>(InContext);
-	OutVariantData->Initialize(MoveTemp(Collection));
+	OutVariantData->Initialize(MoveTemp(ImportTexture2DOutput.OutputCollection));
+
+	if (TextureVisualizationMaterialInstance)
+	{
+		OutVariantData->AddManagedResource(TextureVisualizationMaterialInstance);
+	}
 
 	FPCGTaggedData& CollectionOutput = InContext->OutputData.TaggedData.Emplace_GetRef();
 	CollectionOutput.Data = OutVariantData;
diff --git a/Engine/Plugins/Experimental/ProceduralVegetationEditor/Source/ProceduralVegetation/Private/Nodes/PVTexture2DImporterSettings.h b/Engine/Plugins/Experimental/ProceduralVegetationEditor/Source/ProceduralVegetation/Private/Nodes/PVTexture2DImporterSettings.h
--- a/Engine/Plugins/Experimental/ProceduralVegetationEditor/Source/ProceduralVegetation/Private/Nodes/PVTexture2DImporterSettings.h
+++ b/Engine/Plugins/Experimental/ProceduralVegetationEditor/Source/ProceduralVegetation/Private/Nodes/PVTexture2DImporterSettings.h
@@ -7,6 +7,8 @@
 
 #include "PVTexture2DImporterSettings.generated.h"
 
+class UMaterialInterface;
+
 UCLASS(BlueprintType, ClassGroup = (Procedural))
 class UPVTexture2DImporterSettings : public UPVBaseSettings
 {
@@ -14,6 +16,8 @@ class UPVTexture2DImporterSettings : public UPVBaseSettings
 
 public:
 	
+	UPVTexture2DImporterSettings();
+	
 #if WITH_EDITOR
 	virtual FName GetDefaultNodeName() const override { return FName(TEXT("ProceduralVegetationTexture2DImporter")); }
 	virtual FText GetDefaultNodeTitle() const override;
@@ -32,6 +36,10 @@ public:
 	UPROPERTY(EditAnywhere, Category="Importer Settings", meta = (ShowOnlyInnerProperties))
 	FPVImportTexture2DParams Params;
 
+#if WITH_EDITORONLY_DATA
+	UPROPERTY(EditAnywhere, Category = DebugSettings, meta = (EditCondition = "bDebug", EditConditionHides))
+	TObjectPtr<UMaterialInterface> TextureVisualizationMaterial;
+#endif // WITH_EDITOR_ONLY_DATA
 };
 
 class FPVTexture2DImporterElement : public FPVBaseElement
diff --git a/Engine/Plugins/Experimental/ProceduralVegetationEditor/Source/ProceduralVegetation/Public/DataTypes/PVData.h b/Engine/Plugins/Experimental/ProceduralVegetationEditor/Source/ProceduralVegetation/Public/DataTypes/PVData.h
--- a/Engine/Plugins/Experimental/ProceduralVegetationEditor/Source/ProceduralVegetation/Public/DataTypes/PVData.h
+++ b/Engine/Plugins/Experimental/ProceduralVegetationEditor/Source/ProceduralVegetation/Public/DataTypes/PVData.h
@@ -190,6 +190,14 @@ public:
 
 	const FManagedArrayCollection& GetCollection() const { return Collection; }
 	FManagedArrayCollection& GetCollection() { return Collection; }
+
+	void AddManagedResource(UObject* InResource) 
+	{ 
+		if (InResource)
+		{
+			ManagedResources.AddUnique(InResource);
+		}
+	}
 	
 #if WITH_EDITORONLY_DATA
 	void SetDebugSettings(const FPVDebugSettings& InSettings) const { DebugSettings = InSettings; }
@@ -201,6 +209,9 @@ protected:
 
 	FManagedArrayCollection Collection;
 
+	UPROPERTY(Transient)
+	TArray<TObjectPtr<UObject>> ManagedResources;
+
 #if WITH_EDITOR
 	mutable FPVDebugSettings DebugSettings;
 #endif
diff --git a/Engine/Plugins/Experimental/ProceduralVegetationEditor/Source/ProceduralVegetation/Public/IPVEditorInterface.h b/Engine/Plugins/Experimental/ProceduralVegetationEditor/Source/ProceduralVegetation/Public/IPVEditorInterface.h
--- a/Engine/Plugins/Experimental/ProceduralVegetationEditor/Source/ProceduralVegetation/Public/IPVEditorInterface.h
+++ b/Engine/Plugins/Experimental/ProceduralVegetationEditor/Source/ProceduralVegetation/Public/IPVEditorInterface.h
@@ -8,6 +8,7 @@
 struct FManagedArrayCollection;
 struct FPVImportStaticMeshParams;
 class UStaticMesh;
+class UMaterialInterface;
 
 #if WITH_EDITOR
 namespace PV
@@ -18,7 +19,7 @@ namespace PV
 		virtual ~IPVEditorInterface() = default;
 
 		virtual bool ImportGrowthDataFromStaticMesh(const FPVImportStaticMeshParams& InParams, FManagedArrayCollection& OutCollection) = 0;
-		virtual bool ImportGrowthDataFromTexture2D(const FPVImportTexture2DParams& InParams, FManagedArrayCollection& OutCollection) = 0;
+		virtual bool ImportGrowthDataFromTexture2D(const FPVImportTexture2DParams& InParams, const UMaterialInterface* TextureVisualizationMaterial, FPVImportTexture2DOutput& Output) = 0;
 		virtual TSharedRef<FManagedArrayCollection> CreateMockTreeGrowthDataCollection() const = 0;
 	};
 
diff --git a/Engine/Plugins/Experimental/ProceduralVegetationEditor/Source/ProceduralVegetation/Public/Nodes/PVBaseSettings.h b/Engine/Plugins/Experimental/ProceduralVegetationEditor/Source/ProceduralVegetation/Public/Nodes/PVBaseSettings.h
--- a/Engine/Plugins/Experimental/ProceduralVegetationEditor/Source/ProceduralVegetation/Public/Nodes/PVBaseSettings.h
+++ b/Engine/Plugins/Experimental/ProceduralVegetationEditor/Source/ProceduralVegetation/Public/Nodes/PVBaseSettings.h
@@ -39,12 +39,16 @@ protected:
 	virtual FPCGDataTypeIdentifier GetInputPinTypeIdentifier() const;
 	virtual FPCGDataTypeIdentifier GetOutputPinTypeIdentifier() const;
 
-private:
-	
+protected:
+
 #if WITH_EDITORONLY_DATA
 	UPROPERTY(EditAnywhere, Category = DebugSettings, meta=(EditCondition = "bDebug", EditConditionHides))
 	FPVDebugSettings DebugVisualizationSettings;
+#endif
+
+private:
 
+#if WITH_EDITORONLY_DATA
 	UPROPERTY()
 	TArray<EPVRenderType> CurrentRenderType;
 #endif
diff --git a/Engine/Plugins/Experimental/ProceduralVegetationEditor/Source/ProceduralVegetation/Public/PVImportTexture2DParams.h b/Engine/Plugins/Experimental/ProceduralVegetationEditor/Source/ProceduralVegetation/Public/PVImportTexture2DParams.h
--- a/Engine/Plugins/Experimental/ProceduralVegetationEditor/Source/ProceduralVegetation/Public/PVImportTexture2DParams.h
+++ b/Engine/Plugins/Experimental/ProceduralVegetationEditor/Source/ProceduralVegetation/Public/PVImportTexture2DParams.h
@@ -2,7 +2,7 @@
 #pragma once
 
 #include "Engine/EngineTypes.h"
-
+#include "GeometryCollection/ManagedArrayCollection.h"
 #include "PVImportTexture2DParams.generated.h"
 
 class UTexture2D;
@@ -83,3 +83,8 @@ struct FPVImportTexture2DParams
 	UPROPERTY(EditAnywhere, Category = "Inspect")
 	EPVImportTexture2DDebugState DebugState = EPVImportTexture2DDebugState::None;
 };
+
+struct FPVImportTexture2DOutput
+{
+	FManagedArrayCollection OutputCollection;
+};
\ No newline at end of file
diff --git a/Engine/Plugins/Experimental/ProceduralVegetationEditor/Source/ProceduralVegetationEditor/Private/Importer/PVImporter_Texture2D.cpp b/Engine/Plugins/Experimental/ProceduralVegetationEditor/Source/ProceduralVegetationEditor/Private/Importer/PVImporter_Texture2D.cpp
--- a/Engine/Plugins/Experimental/ProceduralVegetationEditor/Source/ProceduralVegetationEditor/Private/Importer/PVImporter_Texture2D.cpp
+++ b/Engine/Plugins/Experimental/ProceduralVegetationEditor/Source/ProceduralVegetationEditor/Private/Importer/PVImporter_Texture2D.cpp
@@ -10,7 +10,7 @@
 #include "GeometryCollection/ManagedArrayCollection.h"
 
 #include "Engine/Texture2D.h"
-#include "Materials/MaterialInstanceDynamic.h"
+
 #include "Materials/MaterialInterface.h"
 
 #include "DynamicMesh/DynamicMesh3.h"
@@ -36,8 +36,8 @@ namespace Internal
 	const static FName PixelsDebugGroupName = TEXT("PixelDebug");
 	const static FName PixelPositionsAttributeName = TEXT("Position");
 	const static FName PixelValueAttributeName = TEXT("Value");
-	const static FName CurveGroupName = TEXT("Curves");
-	const static FName CurveAttributeName = TEXT("Curve");
+	const static FName PerimiterCurveGroupName = TEXT("PerimiterCurves");
+	const static FName PerimiterCurveAttributeName = TEXT("PerimiterCurve");
 	const static FName InputRootPointsGroupName = TEXT("InputRootPoints");
 	const static FName RootPointsGroupName = TEXT("RootPoints");
 	const static FName RootPointsAttributeName = TEXT("Point");
@@ -97,23 +97,23 @@ namespace Internal
 		}
 	}
 
-	static void AddDebugCurvesToCollection(FManagedArrayCollection& InOutCollection, const TArray<TArray<FVector2f>>& InCurves, float Scale)
+	static void AddDebugPerimiterCurvesToCollection(FManagedArrayCollection& InOutCollection, const TArray<TArray<FVector2f>>& InPerimiterCurves, float Scale)
 	{
-		InOutCollection.AddGroup(CurveGroupName);
-		InOutCollection.AddElements(InCurves.Num(), CurveGroupName);
+		InOutCollection.AddGroup(PerimiterCurveGroupName);
+		InOutCollection.AddElements(InPerimiterCurves.Num(), PerimiterCurveGroupName);
     	
-		auto& CurveAttribute = InOutCollection.AddAttribute<TArray<FVector3f>>(CurveAttributeName, CurveGroupName);
-		for (int32 i = 0; i < InCurves.Num(); i++)
+		auto& PerimiterCurveAttribute = InOutCollection.AddAttribute<TArray<FVector3f>>(PerimiterCurveAttributeName, PerimiterCurveGroupName);
+		for (int32 i = 0; i < InPerimiterCurves.Num(); i++)
 		{
-			const TArray<FVector2f>& Curve = InCurves[i];
+			const TArray<FVector2f>& Curve = InPerimiterCurves[i];
 			for (const FVector2f& CurvePos : Curve)
 			{
-				CurveAttribute[i].Add(ToVector3(CurvePos, Scale));
+				PerimiterCurveAttribute[i].Add(ToVector3(CurvePos, Scale));
 			}
 		}
 	}
 
-	static void AddDebugTextureQuadToCollection(FManagedArrayCollection& InOutCollection, UTexture2D& Texture2D, float Scale)
+	static void AddDebugTextureQuadToCollection(FManagedArrayCollection& InOutCollection, UTexture2D& Texture2D, const UMaterialInterface* TextureVisualizationMaterial, float Scale)
 	{
 		const TArray<FVector> Points = {
 			FVector(0, 0, 1),
@@ -147,21 +147,9 @@ namespace Internal
 			1
 		);
 
-		const FSoftObjectPath ParentMaterialPath = FSoftObjectPath(TEXT("/Script/Engine.Material'/Engine/Engine_MI_Shaders/M_Shader_SimpleTranslucent.M_Shader_SimpleTranslucent'"));
-		const TSoftObjectPtr<UMaterialInterface> ParentMaterialRef(ParentMaterialPath);
-		UMaterialInterface* ParentMaterial = ParentMaterialRef.LoadSynchronous();
-		if (ParentMaterial)
-		{
-			UMaterialInstanceDynamic* MaterialInstance = UMaterialInstanceDynamic::Create(ParentMaterial, nullptr);
-			MaterialInstance->SetFlags(RF_Transient);
-			FTaskTagScope Scope(ETaskTag::EParallelGameThread); // Needed for SetTextureParameterValue
-			MaterialInstance->SetTextureParameterValue(TEXT("Diffuse"), &Texture2D);
-			MaterialInstance->SetTextureParameterValue(TEXT("OpacityMask"), &Texture2D);
-
-			const static FName MaterialPathAttributeName("MaterialPath");
-			auto& MaterialArray = MeshGeomCollection->AddAttribute<FString>(MaterialPathAttributeName, FGeometryCollection::MaterialGroup);
-			MaterialArray[0] = MaterialInstance->GetPathName();
-		}
+		const static FName MaterialPathAttributeName("MaterialPath");
+		auto& MaterialArray = MeshGeomCollection->AddAttribute<FString>(MaterialPathAttributeName, FGeometryCollection::MaterialGroup);
+		MaterialArray[0] = TextureVisualizationMaterial ? TextureVisualizationMaterial->GetPathName() : TEXT("None");
 
 		MeshGeomCollection->CopyTo(&InOutCollection);
 	}
@@ -846,10 +834,12 @@ namespace Internal
 	}
 };
 
-PV::Texture2DImport::EImportResult PV::Texture2DImport::ImportGrowthDataFromTexture2D(const FPVImportTexture2DParams& InParams, FManagedArrayCollection& OutCollection)
+PV::Texture2DImport::EImportResult PV::Texture2DImport::ImportGrowthDataFromTexture2D(const FPVImportTexture2DParams& InParams, const UMaterialInterface* TextureVisualizationMaterial, FPVImportTexture2DOutput& Output)
 {
 	using namespace PV::DynamicMeshHelper;
 	using namespace Internal;
+
+	FManagedArrayCollection& OutCollection = Output.OutputCollection;
 	
 	if (!InParams.Texture2DAsset || !InParams.Texture2DAsset->Source.IsValid())
 	{
@@ -858,7 +848,7 @@ PV::Texture2DImport::EImportResult PV::Texture2DImport::ImportGrowthDataFromText
 
 	if (InParams.DebugState == EPVImportTexture2DDebugState::Texture)
 	{
-		AddDebugTextureQuadToCollection(OutCollection, *InParams.Texture2DAsset, InParams.Scale);
+		AddDebugTextureQuadToCollection(OutCollection, *InParams.Texture2DAsset, TextureVisualizationMaterial, InParams.Scale);
 		return EImportResult::Success;
 	}
 
@@ -879,7 +869,7 @@ PV::Texture2DImport::EImportResult PV::Texture2DImport::ImportGrowthDataFromText
 	if (InParams.DebugState == EPVImportTexture2DDebugState::SampleImage)
 	{
 		AddDebugPixelValuesToCollection(OutCollection, BlackAndWhiteImageData, SampleSize, InParams.Scale);
-		AddDebugTextureQuadToCollection(OutCollection, *InParams.Texture2DAsset, InParams.Scale);
+		AddDebugTextureQuadToCollection(OutCollection, *InParams.Texture2DAsset, TextureVisualizationMaterial, InParams.Scale);
 		return EImportResult::Success;
 	}
 
@@ -887,32 +877,32 @@ PV::Texture2DImport::EImportResult PV::Texture2DImport::ImportGrowthDataFromText
 	TraceImage(BlackAndWhiteImageData, SampleSize, Curves);
 	if (InParams.DebugState == EPVImportTexture2DDebugState::TraceImage)
 	{
-		AddDebugCurvesToCollection(OutCollection, Curves, InParams.Scale);
-		AddDebugTextureQuadToCollection(OutCollection, *InParams.Texture2DAsset, InParams.Scale);
+		AddDebugPerimiterCurvesToCollection(OutCollection, Curves, InParams.Scale);
+		AddDebugTextureQuadToCollection(OutCollection, *InParams.Texture2DAsset, TextureVisualizationMaterial, InParams.Scale);
 		return EImportResult::Success;
 	}
 	
 	RemoveSmallCurves(Curves, InParams.MinBoundsArea);
 	if (InParams.DebugState == EPVImportTexture2DDebugState::RemoveSmallCurves)
 	{
-		AddDebugCurvesToCollection(OutCollection, Curves, InParams.Scale);
-		AddDebugTextureQuadToCollection(OutCollection, *InParams.Texture2DAsset, InParams.Scale);
+		AddDebugPerimiterCurvesToCollection(OutCollection, Curves, InParams.Scale);
+		AddDebugTextureQuadToCollection(OutCollection, *InParams.Texture2DAsset, TextureVisualizationMaterial, InParams.Scale);
 		return EImportResult::Success;
 	}
 	
 	SmoothCurves(Curves, InParams.SmoothingIterations);
 	if (InParams.DebugState == EPVImportTexture2DDebugState::SmoothCurves)
 	{
-		AddDebugCurvesToCollection(OutCollection, Curves, InParams.Scale);
-		AddDebugTextureQuadToCollection(OutCollection, *InParams.Texture2DAsset, InParams.Scale);
+		AddDebugPerimiterCurvesToCollection(OutCollection, Curves, InParams.Scale);
+		AddDebugTextureQuadToCollection(OutCollection, *InParams.Texture2DAsset, TextureVisualizationMaterial, InParams.Scale);
 		return EImportResult::Success;
 	}
 	
 	SimplifyCurves(Curves, InParams.SimplificationEpsilon);
 	if (InParams.DebugState == EPVImportTexture2DDebugState::SimplifyCurves)
 	{
-		AddDebugCurvesToCollection(OutCollection, Curves, InParams.Scale);
-		AddDebugTextureQuadToCollection(OutCollection, *InParams.Texture2DAsset, InParams.Scale);
+		AddDebugPerimiterCurvesToCollection(OutCollection, Curves, InParams.Scale);
+		AddDebugTextureQuadToCollection(OutCollection, *InParams.Texture2DAsset, TextureVisualizationMaterial, InParams.Scale);
 		return EImportResult::Success;
 	}
 	
@@ -960,7 +950,7 @@ PV::Texture2DImport::EImportResult PV::Texture2DImport::ImportGrowthDataFromText
 	);
 	if (InParams.DebugState == EPVImportTexture2DDebugState::ApplyRootPoints)
 	{
-		AddDebugTextureQuadToCollection(OutCollection, *InParams.Texture2DAsset, InParams.Scale);
+		AddDebugTextureQuadToCollection(OutCollection, *InParams.Texture2DAsset, TextureVisualizationMaterial, InParams.Scale);
 		AddDynamicMeshToCollection(DynamicMesh, OutCollection);
 		AddDebugRootPointsToCollection(OutCollection, InParams.RootPoints, InParams.Scale);
 		EndPointAttributeDefinition->CopyAttributeToCollection(DynamicMesh, OutCollection);
diff --git a/Engine/Plugins/Experimental/ProceduralVegetationEditor/Source/ProceduralVegetationEditor/Private/Importer/PVImporter_Texture2D.h b/Engine/Plugins/Experimental/ProceduralVegetationEditor/Source/ProceduralVegetationEditor/Private/Importer/PVImporter_Texture2D.h
--- a/Engine/Plugins/Experimental/ProceduralVegetationEditor/Source/ProceduralVegetationEditor/Private/Importer/PVImporter_Texture2D.h
+++ b/Engine/Plugins/Experimental/ProceduralVegetationEditor/Source/ProceduralVegetationEditor/Private/Importer/PVImporter_Texture2D.h
@@ -3,8 +3,8 @@
 
 #include "CoreMinimal.h"
 #include "PVImportTexture2DParams.h"
-#include "GeometryCollection/ManagedArrayCollection.h"
 
+class UMaterialInterface;
 class UTexture2D;
 
 namespace PV::Texture2DImport
@@ -18,5 +18,5 @@ namespace PV::Texture2DImport
 		FailedToGenerateGrowthData
 	};
 
-	EImportResult ImportGrowthDataFromTexture2D(const FPVImportTexture2DParams& InParams, FManagedArrayCollection& OutCollection);
+	EImportResult ImportGrowthDataFromTexture2D(const FPVImportTexture2DParams& InParams, const UMaterialInterface* TextureVisualizationMaterial, FPVImportTexture2DOutput& Output);
 };
diff --git a/Engine/Plugins/Experimental/ProceduralVegetationEditor/Source/ProceduralVegetationEditor/Private/PVEditorInterfaceImpl.cpp b/Engine/Plugins/Experimental/ProceduralVegetationEditor/Source/ProceduralVegetationEditor/Private/PVEditorInterfaceImpl.cpp
--- a/Engine/Plugins/Experimental/ProceduralVegetationEditor/Source/ProceduralVegetationEditor/Private/PVEditorInterfaceImpl.cpp
+++ b/Engine/Plugins/Experimental/ProceduralVegetationEditor/Source/ProceduralVegetationEditor/Private/PVEditorInterfaceImpl.cpp
@@ -20,9 +20,9 @@ bool PV::FPVEditorInterfaceImpl::ImportGrowthDataFromStaticMesh(const FPVImportS
 	return ImportResult == PV::StaticMeshImport::EImportResult::Success;
 }
 
-bool PV::FPVEditorInterfaceImpl::ImportGrowthDataFromTexture2D(const FPVImportTexture2DParams& InParams, FManagedArrayCollection& OutCollection)
+bool PV::FPVEditorInterfaceImpl::ImportGrowthDataFromTexture2D(const FPVImportTexture2DParams& InParams, const UMaterialInterface* TextureVisualizationMaterial, FPVImportTexture2DOutput &Output)
 {
-	const PV::Texture2DImport::EImportResult ImportResult = PV::Texture2DImport::ImportGrowthDataFromTexture2D(InParams, OutCollection);
+	const PV::Texture2DImport::EImportResult ImportResult = PV::Texture2DImport::ImportGrowthDataFromTexture2D(InParams, TextureVisualizationMaterial, Output);
 	return ImportResult == PV::Texture2DImport::EImportResult::Success;
 }
 
diff --git a/Engine/Plugins/Experimental/ProceduralVegetationEditor/Source/ProceduralVegetationEditor/Private/PVEditorInterfaceImpl.h b/Engine/Plugins/Experimental/ProceduralVegetationEditor/Source/ProceduralVegetationEditor/Private/PVEditorInterfaceImpl.h
--- a/Engine/Plugins/Experimental/ProceduralVegetationEditor/Source/ProceduralVegetationEditor/Private/PVEditorInterfaceImpl.h
+++ b/Engine/Plugins/Experimental/ProceduralVegetationEditor/Source/ProceduralVegetationEditor/Private/PVEditorInterfaceImpl.h
@@ -6,8 +6,6 @@
 
 #include "GeometryCollection/ManagedArrayCollection.h"
 
-struct FPVImportStaticMeshParams;
-
 #if WITH_EDITOR
 namespace PV
 {
@@ -15,7 +13,7 @@ namespace PV
 	{
 	public:
 		bool ImportGrowthDataFromStaticMesh(const FPVImportStaticMeshParams& InParams, FManagedArrayCollection& OutCollection) override;
-		bool ImportGrowthDataFromTexture2D(const FPVImportTexture2DParams& InParams, FManagedArrayCollection& OutCollection) override;
+		bool ImportGrowthDataFromTexture2D(const FPVImportTexture2DParams& InParams, const UMaterialInterface* TextureVisualizationMaterial, FPVImportTexture2DOutput& Output) override;
 		TSharedRef<FManagedArrayCollection> CreateMockTreeGrowthDataCollection() const override;
 	};
 };
