[PCG][PVE] Fix missing debug (CollectionSpreadSheetWidget) tab

#rnx

#Synced-CL 49959840

diff --git a/Engine/Plugins/Experimental/ProceduralVegetationEditor/Source/ProceduralVegetationEditor/Private/PVEditor.cpp b/Engine/Plugins/Experimental/ProceduralVegetationEditor/Source/ProceduralVegetationEditor/Private/PVEditor.cpp
--- a/Engine/Plugins/Experimental/ProceduralVegetationEditor/Source/ProceduralVegetationEditor/Private/PVEditor.cpp
+++ b/Engine/Plugins/Experimental/ProceduralVegetationEditor/Source/ProceduralVegetationEditor/Private/PVEditor.cpp
@@ -8,6 +8,7 @@
 #include "PackagesDialog.h"
 #include "PCGDefaultExecutionSource.h"
 #include "PCGEditorCommands.h"
+#include "PCGEditorTabFactories.h"
 #include "ProceduralVegetation.h"
 #include "ProceduralVegetationEditorModule.h"
 #include "PVEditorCommands.h"
@@ -180,25 +181,16 @@ void FPVEditor::OnClose()
 	}
 }
 
-void FPVEditor::RegisterTabSpawners(const TSharedRef<class FTabManager>& InTabManager)
+void FPVEditor::RegisterExtraTabFactories(FWorkflowAllowedTabSet& TabSet)
 {
-	FPCGEditor::RegisterTabSpawners(InTabManager);
-
 	if (CollectionSpreadSheetWidget.IsValid())
 	{
-		InTabManager->RegisterTabSpawner(CollectionSpreadSheetTabId, FOnSpawnTab::CreateSP(this, &FPVEditor::SpawnTab_CollectionSpreadSheet))
-			.SetDisplayName(LOCTEXT("CollectionSpreadsheetTab", "Collection Spreadsheet"))
-			.SetGroup(WorkspaceMenuCategory.ToSharedRef());
+		FPCGGenericTabFactoryParams Params(CollectionSpreadSheetTabId, SharedThis<FPCGEditor>(this), [this]() { return CollectionSpreadSheetWidget.ToSharedRef(); });
+		Params.Label = LOCTEXT("CollectionSpreadsheetTab", "Collection Spreadsheet");
+		TabSet.RegisterFactory(MakeShared<FPCGGenericTabFactory>(Params));
 	}
 }
 
-void FPVEditor::UnregisterTabSpawners(const TSharedRef<class FTabManager>& InTabManager)
-{
-	FPCGEditor::UnregisterTabSpawners(InTabManager);
-
-	InTabManager->UnregisterTabSpawner(CollectionSpreadSheetTabId);
-}
-
 IPCGBaseSubsystem* FPVEditor::GetSubsystem() const
 {
 	return UPCGEngineSubsystem::Get();
@@ -934,25 +926,10 @@ TSharedRef<SCollectionSpreadSheetWidget> FPVEditor::CreateCollectionSpreadSheetW
 	return SNew(SCollectionSpreadSheetWidget);
 }
 
-TSharedRef<SDockTab> FPVEditor::SpawnTab_CollectionSpreadSheet(const FSpawnTabArgs& SpawnTabArgs)
-{
-	check(SpawnTabArgs.GetTabId() == FPVEditor::CollectionSpreadSheetTabId);
-
-	TSharedRef<SDockTab> DockableTab = SNew(SDockTab)
-		[
-			CollectionSpreadSheetWidget.ToSharedRef()
-		];
-
-	CollectionSpreadSheetTab = DockableTab;
-
-	return DockableTab;
-}
-
 void FPVEditor::UpdateCollectionSpreadSheet(const UPCGNode* InSelectedNode)
 {
-	const bool bIsCollectionSpreadSheetVisible = CollectionSpreadSheetTab.IsValid() &&
-		CollectionSpreadSheetTab.Pin()->GetVisibility().IsVisible() &&
-		CollectionSpreadSheetWidget;
+	TSharedPtr<SDockTab> CollectionSpreadSheetTab = GetTabManager() ? GetTabManager()->FindExistingLiveTab(CollectionSpreadSheetTabId) : nullptr;
+	const bool bIsCollectionSpreadSheetVisible = CollectionSpreadSheetWidget.IsValid() && CollectionSpreadSheetTab.IsValid() && CollectionSpreadSheetTab->GetVisibility().IsVisible();
 
 	if (InSelectedNode)
 	{
diff --git a/Engine/Plugins/Experimental/ProceduralVegetationEditor/Source/ProceduralVegetationEditor/Private/PVEditor.h b/Engine/Plugins/Experimental/ProceduralVegetationEditor/Source/ProceduralVegetationEditor/Private/PVEditor.h
--- a/Engine/Plugins/Experimental/ProceduralVegetationEditor/Source/ProceduralVegetationEditor/Private/PVEditor.h
+++ b/Engine/Plugins/Experimental/ProceduralVegetationEditor/Source/ProceduralVegetationEditor/Private/PVEditor.h
@@ -30,8 +30,6 @@ public:
 	virtual FLinearColor GetWorldCentricTabColorScale() const override;
 	virtual FString GetWorldCentricTabPrefix() const override;
 	virtual void OnClose() override;
-	virtual void RegisterTabSpawners(const TSharedRef<class FTabManager>& InTabManager) override;
-	virtual void UnregisterTabSpawners(const TSharedRef<class FTabManager>& InTabManager) override;
 	// ~End FAssetEditorToolkit interface
 
 	virtual IPCGBaseSubsystem* GetSubsystem() const override;
@@ -40,6 +38,7 @@ protected:
 	// ~Begin FPCGEditor interface
 	virtual TAttribute<FGraphAppearanceInfo> GetAppearanceInfo() const override;
 	virtual TSharedRef<FTabManager::FLayout> GetDefaultLayout() const override;
+	virtual void RegisterExtraTabFactories(FWorkflowAllowedTabSet& TabSet) override;
 	virtual bool IsPanelAvailable(const FName PanelID) const override;
 	virtual void RegisterToolbarInternal(FToolBarBuilder& ToolbarBuilder) const override;
 	virtual void BindCommands() override;
@@ -63,14 +62,10 @@ private:
 	TObjectPtr<UPCGEditorGraphNodeBase> NodeBeingInspected = nullptr;
 	
 	TSharedPtr<SCollectionSpreadSheetWidget> CollectionSpreadSheetWidget;
-
-	TWeakPtr<SDockTab> CollectionSpreadSheetTab;
 	
 	FDateTime SessionStartTime;
 	
 	TSharedRef<SCollectionSpreadSheetWidget> CreateCollectionSpreadSheetWidget();
-
-	TSharedRef<SDockTab> SpawnTab_CollectionSpreadSheet(const FSpawnTabArgs& SpawnTabArgs);
 	
 	void UpdateCollectionSpreadSheet(const UPCGNode* InSelectedNode);
 
diff --git a/Engine/Plugins/PCG/Source/PCGEditor/Private/PCGEditorDefaultMode.cpp b/Engine/Plugins/PCG/Source/PCGEditor/Private/PCGEditorDefaultMode.cpp
--- a/Engine/Plugins/PCG/Source/PCGEditor/Private/PCGEditorDefaultMode.cpp
+++ b/Engine/Plugins/PCG/Source/PCGEditor/Private/PCGEditorDefaultMode.cpp
@@ -137,6 +137,8 @@ FPCGEditorDefaultMode::FPCGEditorDefaultMode(TSharedPtr<FPCGEditor> InPCGEditor)
 		TabFactories.RegisterFactory(Info.CreateFactory.Execute());
 	}
 
+	InPCGEditor->RegisterExtraTabFactories(ExtraTabFactories);
+
 	ToolbarExtender = MakeShared<FExtender>();
 	ToolbarExtender->AddToolBarExtension(
 		"Asset",
@@ -156,6 +158,7 @@ void FPCGEditorDefaultMode::RegisterTabFactories(TSharedPtr<FTabManager> InTabMa
 	if (TSharedPtr<FPCGEditor> PCGEditor = StaticCastSharedPtr<FPCGEditor>(GetHost()))
 	{
 		PCGEditor->PushTabFactories(TabFactories);
+		PCGEditor->PushTabFactories(ExtraTabFactories);
 	}
 }
 
diff --git a/Engine/Plugins/PCG/Source/PCGEditor/Private/PCGEditorDefaultMode.h b/Engine/Plugins/PCG/Source/PCGEditor/Private/PCGEditorDefaultMode.h
--- a/Engine/Plugins/PCG/Source/PCGEditor/Private/PCGEditorDefaultMode.h
+++ b/Engine/Plugins/PCG/Source/PCGEditor/Private/PCGEditorDefaultMode.h
@@ -29,4 +29,5 @@ public:
 
 private:
 	FWorkflowAllowedTabSet TabFactories;
+	FWorkflowAllowedTabSet ExtraTabFactories;
 };
\ No newline at end of file
diff --git a/Engine/Plugins/PCG/Source/PCGEditor/Public/PCGEditor.h b/Engine/Plugins/PCG/Source/PCGEditor/Public/PCGEditor.h
--- a/Engine/Plugins/PCG/Source/PCGEditor/Public/PCGEditor.h
+++ b/Engine/Plugins/PCG/Source/PCGEditor/Public/PCGEditor.h
@@ -555,6 +555,7 @@ private:
 	void UpdateDefaultExecutionSource();
 
 	PCGEDITOR_API virtual TSharedRef<FTabManager::FLayout> GetDefaultLayout() const;
+	virtual void RegisterExtraTabFactories(FWorkflowAllowedTabSet& TabSet) {}
 
 	TSharedPtr<SGraphEditor> GraphEditorWidget;
 	TSharedPtr<SPCGEditorGraphDetailsView> PropertyDetailsWidgets[PCGEditorTabs::NumPropertyDetailTabs];
diff --git a/Engine/Plugins/PCG/Source/PCGEditor/Public/PCGEditorTabFactories.h b/Engine/Plugins/PCG/Source/PCGEditor/Public/PCGEditorTabFactories.h
--- a/Engine/Plugins/PCG/Source/PCGEditor/Public/PCGEditorTabFactories.h
+++ b/Engine/Plugins/PCG/Source/PCGEditor/Public/PCGEditorTabFactories.h
@@ -35,7 +35,7 @@ struct FPCGGenericTabFactoryParams
 struct FPCGGenericTabFactory : public FWorkflowTabFactory
 {
 public:
-	FPCGGenericTabFactory(const FPCGGenericTabFactoryParams& InParams);
+	PCGEDITOR_API FPCGGenericTabFactory(const FPCGGenericTabFactoryParams& InParams);
 
 	virtual TSharedRef<SWidget> CreateTabBody(const FWorkflowTabSpawnInfo& Info) const override;
 
