[PVE] Add end-caps to each branch during meshing

This CL adds support for generating tapered end-caps at the tip of branches during mesh construction when end-caps are enabled. The implementation creates a short multi-step cap that smoothly reduces scale to zero, ensuring correct topology, UV continuity, and normal generation at branch ends.

Overview:
- Add end-caps option added to the mesh builder settings
- Generates a 3 step tapered cap aligned to the branch tangent
- Exponentially reduces scale and step size to close the mesh with a spherical cap
- Ensures correct UV and displacement continuity at branch tips
- Creates unique vertices per triangle when collapsing to zero scale

#Synced-CL 49988250
#Synced-CL 49988250

diff --git a/Engine/Plugins/Experimental/ProceduralVegetationEditor/Source/ProceduralVegetation/Private/Implementations/PVMeshBuilder.cpp b/Engine/Plugins/Experimental/ProceduralVegetationEditor/Source/ProceduralVegetation/Private/Implementations/PVMeshBuilder.cpp
--- a/Engine/Plugins/Experimental/ProceduralVegetationEditor/Source/ProceduralVegetation/Private/Implementations/PVMeshBuilder.cpp
+++ b/Engine/Plugins/Experimental/ProceduralVegetationEditor/Source/ProceduralVegetation/Private/Implementations/PVMeshBuilder.cpp
@@ -773,6 +773,7 @@ void FPVMeshBuilder::GenerateBranchMeshData(const bool bPrimitiveIsTrunk, const
 		return FMath::IsNearlyZero(InDenominator) ? InNumerator : (InNumerator / InDenominator);
 	};
 	
+	FVector3f TangentVector = FVector3f::UpVector;
 	for (int32 Index = 0; Index < PrimitivePoints.Num(); ++Index)
 	{
 		const int32& PointIndex = PrimitivePoints[Index];
@@ -795,8 +796,7 @@ void FPVMeshBuilder::GenerateBranchMeshData(const bool bPrimitiveIsTrunk, const
 		const float PointScale = PointScales[PointIndex];
 		
 		Displacement_UV_V += FVector3f::Dist(PointPosition, PreviousPosition) * MaxScaleRatio * (MaxPointScale / PointScale);
-
-		FVector3f TangentVector = FVector3f::UpVector;
+		
 		if (PrimitivePoints.IsValidIndex(Index - 1) && PrimitivePoints.IsValidIndex(Index + 1))
 		[[likely]]
 		{
@@ -916,6 +916,139 @@ void FPVMeshBuilder::GenerateBranchMeshData(const bool bPrimitiveIsTrunk, const
 		PreviousIndices = MoveTemp(CurrentIndices);
 		CurrentIndices.Empty();
 	}
+	
+	// Add End-cap
+	if (PreviousIndices.Num() >= 3 && MeshBuilderParams.AddEndCaps)
+	{
+		const int32 LastPointIndex = PrimitivePoints.Last();
+		float CurrentLengthFromRoot = PointFacade.GetLengthFromRoot(LastPointIndex);
+		float CurrentScale = PointScales[LastPointIndex];
+		const int32 CapDivisions = TargetMeshDivisions[LastPointIndex];
+		float StepSize = CurrentScale * 0.5f;
+		
+		float CurrentRawTextureV = PointFacade.GetTextureCoordV(LastPointIndex);
+		float LastTextureCoordUOffset = PointFacade.GetTextureCoordUOffset(LastPointIndex);
+		FVector2f LastURange = PointFacade.GetURange(LastPointIndex);
+		FVector2f OldRange = FVector2f(0, 1.0f);
+		
+		// URatio matches the main branch's computation
+		float Interval = LastURange.Y - LastURange.X;
+		float URatio = FMath::IsNearlyZero(Interval) ? 1.0f : 1.0f / Interval;
+		
+		const float UVOGenerationOffset = GenerationNumber - 1;
+		
+		for (int32 CapStep = 1; CapStep <= 3; ++CapStep)
+		{
+			FVector3f CurrentPosition = PreviousPosition + (TangentVector * StepSize);
+			CurrentLengthFromRoot += StepSize;
+			
+			// Keeping displacement V and texture V aligned
+			const float TextureScaleV = FMath::Max(CurrentScale, UE_KINDA_SMALL_NUMBER);
+			const float CapStrideV = StepSize * MaxScaleRatio * (MaxPointScale / TextureScaleV);
+			Displacement_UV_V += CapStrideV;
+			CurrentRawTextureV += CapStrideV;
+			float FinalTextureV = CurrentRawTextureV / URatio;
+			
+			StepSize *= 0.25f;
+			
+			if (CapStep == 3)
+			{
+				CurrentScale = 0.0f;
+			}
+			else
+			{
+				CurrentScale *= 0.5f; 
+			}
+
+			FVector2f UV1 = FVector2f(0, 1.0 - Displacement_UV_V);
+			FVector2f UV2 = FVector2f(CurrentLengthFromRoot * 0.01, SafeDivide(CurrentScale, MaxPointScale));
+			
+			if (CurrentScale <= UE_KINDA_SMALL_NUMBER)
+			[[unlikely]]
+			{
+				// A unique vertex for each tri at the tip.
+				const int32 NumOfLastRingVertices = PreviousIndices.Num();
+				
+				for (int32 Division = 0; Division < NumOfLastRingVertices; ++Division)
+				{
+					const float U_Division = static_cast<float>(Division) / (NumOfLastRingVertices - 1);
+					float UV_U = (1 - U_Division) + LastTextureCoordUOffset;
+					UV_U = FMath::GetMappedRangeValueUnclamped(
+							OldRange,
+							LastURange,
+							FMath::Wrap(UV_U, OldRange.X, OldRange.Y))
+						+ UVOGenerationOffset;
+					
+					UV1.X = 1.0f - U_Division;
+				
+					CurrentIndices.Add(OutLocalMeshData.Vertices.Emplace(
+						CurrentPosition,
+						TangentVector,
+						FVector2f(UV_U, FinalTextureV),
+						UV1,
+						UV2,
+						static_cast<float>(LastPointIndex)
+					));
+				}
+			}
+			else
+			[[likely]]
+			{
+				const float RotationDelta = UE_TWO_PI / CapDivisions;
+				
+				for (int32 Division = 0; Division < CapDivisions; ++Division)
+				{
+					float UV_U = (1 - (static_cast<float>(Division) / CapDivisions)) + LastTextureCoordUOffset;
+					UV_U = FMath::GetMappedRangeValueUnclamped(
+						OldRange,
+						LastURange,
+						FMath::Wrap(UV_U, OldRange.X, OldRange.Y));
+					UV_U += UVOGenerationOffset;
+					
+					float Displacement_UV_U = (static_cast<float>(Division) / CapDivisions);
+					UV1.X = 1 - Displacement_UV_U;
+					
+					const FVector3f Direction = FQuat4f(TangentVector, RotationDelta * Division).RotateVector(PreviousUpVector).GetSafeNormal();
+					
+					FVector3f RadialPoint = CurrentPosition + (Direction * CurrentScale);
+					const FVector3f Normal = FVector3f((RadialPoint - CurrentPosition).GetUnsafeNormal());
+
+					CurrentIndices.Add(OutLocalMeshData.Vertices.Emplace(
+						RadialPoint, 
+						Normal, 
+						FVector2f(UV_U, FinalTextureV), 
+						UV1,
+						UV2,
+						static_cast<float>(LastPointIndex)
+					));
+				}
+				
+				const auto FirstVertex = OutLocalMeshData.Vertices[CurrentIndices[0]];
+				float UV_U = FMath::GetMappedRangeValueUnclamped(
+					OldRange, LastURange, FMath::Wrap(LastTextureCoordUOffset, OldRange.X, OldRange.Y)
+				);
+				UV_U += UVOGenerationOffset;
+				
+				UV1.X = 0;
+
+				CurrentIndices.Add(OutLocalMeshData.Vertices.Emplace(
+					FirstVertex.Position,
+					FirstVertex.Normal,
+					FVector2f(UV_U, FinalTextureV),
+					UV1,
+					UV2,
+					static_cast<float>(LastPointIndex)
+				));
+			}
+			
+			PreviousPosition = CurrentPosition;
+
+			TriangulateRings(PreviousIndices, CurrentIndices, PolyGroupIndex, OutLocalMeshData);
+
+			PreviousIndices = MoveTemp(CurrentIndices);
+			CurrentIndices.Empty();
+		}
+	}
 }
 
 void FPVMeshBuilder::UpdateFoliagePivotPoints(const TSet<int32>& PointsToRemove, FManagedArrayCollection& OutCollection)
diff --git a/Engine/Plugins/Experimental/ProceduralVegetationEditor/Source/ProceduralVegetation/Private/Implementations/PVMeshBuilder.h b/Engine/Plugins/Experimental/ProceduralVegetationEditor/Source/ProceduralVegetation/Private/Implementations/PVMeshBuilder.h
--- a/Engine/Plugins/Experimental/ProceduralVegetationEditor/Source/ProceduralVegetation/Private/Implementations/PVMeshBuilder.h
+++ b/Engine/Plugins/Experimental/ProceduralVegetationEditor/Source/ProceduralVegetation/Private/Implementations/PVMeshBuilder.h
@@ -63,6 +63,12 @@ struct FPVMeshBuilderParams
 		))
 	int32 MaxMeshDivisions = 12;
 
+	UPROPERTY(EditAnywhere, Category="Mesh Settings", DisplayName="Add End Caps",
+		meta=(Tooltip=
+			"Toggles end caps at the end of each branch.\n\nWhen checked a spherical end cap will be added to the tip of each branch."
+		))
+	bool AddEndCaps = true;
+
 	UPROPERTY(EditAnywhere, Category="Mesh Settings|Advanced", DisplayName="Reduction Accuracy",
 		meta=(ClampMin=1, ClampMax=15, UIMin=1, UIMax=15, Tooltip=
 			"Accuracy level for mesh reduction.\n\nControls how closely the simplified mesh follows the original. It specifies number of iterations that are used to evaluate the curve reduction path. The more number of iterations the more accurate the reduction will be at the cost of processing power."
