[PVE] Fix issue where static analysis is generating a warning regarding the denominator

#Synced-CL 49992026

diff --git a/Engine/Plugins/Experimental/ProceduralVegetationEditor/Source/ProceduralVegetation/Private/Implementations/PVMeshBuilder.cpp b/Engine/Plugins/Experimental/ProceduralVegetationEditor/Source/ProceduralVegetation/Private/Implementations/PVMeshBuilder.cpp
--- a/Engine/Plugins/Experimental/ProceduralVegetationEditor/Source/ProceduralVegetation/Private/Implementations/PVMeshBuilder.cpp
+++ b/Engine/Plugins/Experimental/ProceduralVegetationEditor/Source/ProceduralVegetation/Private/Implementations/PVMeshBuilder.cpp
@@ -918,7 +918,8 @@ void FPVMeshBuilder::GenerateBranchMeshData(const bool bPrimitiveIsTrunk, const
 	}
 	
 	// Add End-cap
-	if (PreviousIndices.Num() >= 3 && MeshBuilderParams.AddEndCaps)
+	const int32 NumOfLastRingVertices = PreviousIndices.Num();
+	if (NumOfLastRingVertices >= 3 && MeshBuilderParams.AddEndCaps)
 	{
 		const int32 LastPointIndex = PrimitivePoints.Last();
 		float CurrentLengthFromRoot = PointFacade.GetLengthFromRoot(LastPointIndex);
@@ -967,11 +968,11 @@ void FPVMeshBuilder::GenerateBranchMeshData(const bool bPrimitiveIsTrunk, const
 			[[unlikely]]
 			{
 				// A unique vertex for each tri at the tip.
-				const int32 NumOfLastRingVertices = PreviousIndices.Num();
-				
+				check(NumOfLastRingVertices >= 2);
+				const float U_DivisionDenominator = static_cast<float>(NumOfLastRingVertices - 1);
 				for (int32 Division = 0; Division < NumOfLastRingVertices; ++Division)
 				{
-					const float U_Division = static_cast<float>(Division) / (NumOfLastRingVertices - 1);
+					const float U_Division = static_cast<float>(Division) / U_DivisionDenominator;
 					float UV_U = (1 - U_Division) + LastTextureCoordUOffset;
 					UV_U = FMath::GetMappedRangeValueUnclamped(
 							OldRange,
