Add custom visualization to PVE 2D importer node

- Add custom visualization collections to UPVData
- Output custom visualization collection from PVImporter_Texture2D
- Add custom debug visualization settings in PVImportTexture2DSettings
- Add helpers to PVFacadeHelpers.h to allow for easy iteration of plant points (also improve interface by replacing boolean return type with EForEachResult enum)

#Synced-CL 50040315

diff --git a/Engine/Plugins/Experimental/ProceduralVegetationEditor/Source/ProceduralVegetation/Private/DataTypes/PVData.cpp b/Engine/Plugins/Experimental/ProceduralVegetationEditor/Source/ProceduralVegetation/Private/DataTypes/PVData.cpp
--- a/Engine/Plugins/Experimental/ProceduralVegetationEditor/Source/ProceduralVegetation/Private/DataTypes/PVData.cpp
+++ b/Engine/Plugins/Experimental/ProceduralVegetationEditor/Source/ProceduralVegetation/Private/DataTypes/PVData.cpp
@@ -75,6 +75,13 @@ UPCGSpatialData* UPVData::CopyInternal(FPCGContext* Context) const
 	return NewProceduralVegetationData;
 }
 
+#if WITH_EDITORONLY_DATA
+void UPVData::AddVisualizationCollection(FManagedArrayCollection InCollection, TArray<FPVVisualizationSettings> InVisualizationSettings)
+{
+	VisualizationCollections.Emplace(MoveTemp(InCollection), MoveTemp(InVisualizationSettings));
+}
+#endif
+
 const UPCGPointArrayData* UPVData::ToPointArrayData(FPCGContext* Context, const FBox& InBounds) const
 {
 	return CastChecked<UPCGPointArrayData>(ToBasePointData(Context, InBounds, UPCGPointArrayData::StaticClass()));
diff --git a/Engine/Plugins/Experimental/ProceduralVegetationEditor/Source/ProceduralVegetation/Private/Implementations/PVSlope.cpp b/Engine/Plugins/Experimental/ProceduralVegetationEditor/Source/ProceduralVegetation/Private/Implementations/PVSlope.cpp
--- a/Engine/Plugins/Experimental/ProceduralVegetationEditor/Source/ProceduralVegetation/Private/Implementations/PVSlope.cpp
+++ b/Engine/Plugins/Experimental/ProceduralVegetationEditor/Source/ProceduralVegetation/Private/Implementations/PVSlope.cpp
@@ -39,7 +39,7 @@ namespace PVSlope
 			BranchFacade,
 			PointFacade,
 			BranchIndex,
-			[&](int32 BranchIndex, int32 BranchPointIndex, int32 PointIndex, const TArray<int32>& ChildPointIndices)->bool
+			[&](int32 BranchIndex, int32 BranchPointIndex, int32 PointIndex, const TArray<int32>& ChildPointIndices)->PV::FacadeHelper::EForEachResult
 			{
 				const FVector3f& PointPosition = PointPositions[PointIndex];
 				const float PointDistSquared = (PointPosition - RootPosition).SizeSquared();
@@ -47,7 +47,7 @@ namespace PVSlope
 				{
 					BranchLengthSquared = PointDistSquared;
 				}
-				return true;
+				return PV::FacadeHelper::EForEachResult::Continue;
 			}
 		);
 
diff --git a/Engine/Plugins/Experimental/ProceduralVegetationEditor/Source/ProceduralVegetation/Private/Nodes/PVTexture2DImporterSettings.cpp b/Engine/Plugins/Experimental/ProceduralVegetationEditor/Source/ProceduralVegetation/Private/Nodes/PVTexture2DImporterSettings.cpp
--- a/Engine/Plugins/Experimental/ProceduralVegetationEditor/Source/ProceduralVegetation/Private/Nodes/PVTexture2DImporterSettings.cpp
+++ b/Engine/Plugins/Experimental/ProceduralVegetationEditor/Source/ProceduralVegetation/Private/Nodes/PVTexture2DImporterSettings.cpp
@@ -28,12 +28,12 @@ UPVTexture2DImporterSettings::UPVTexture2DImporterSettings()
 	LineSegmentVisualization.CustomGroupToFilter = TEXT("LineSegments");
 	LineSegmentVisualization.VisualizationMode = EPVDebugValueVisualizationMode::Curve;
 
-	FPVVisualizationSettings& PerimiterCurveVisualization = DebugVisualizationSettings.VisualizationSettings.AddDefaulted_GetRef();
-	PerimiterCurveVisualization.DebugType = EPVDebugType::Custom;
-	PerimiterCurveVisualization.bShowAnchorPoints = false;
-	PerimiterCurveVisualization.AttributeToFilter = TEXT("PerimiterCurve");
-	PerimiterCurveVisualization.CustomGroupToFilter = TEXT("PerimiterCurves");
-	PerimiterCurveVisualization.VisualizationMode = EPVDebugValueVisualizationMode::Curve;
+	FPVVisualizationSettings& PerimeterCurveVisualization = DebugVisualizationSettings.VisualizationSettings.AddDefaulted_GetRef();
+	PerimeterCurveVisualization.DebugType = EPVDebugType::Custom;
+	PerimeterCurveVisualization.bShowAnchorPoints = false;
+	PerimeterCurveVisualization.AttributeToFilter = TEXT("PerimeterCurve");
+	PerimeterCurveVisualization.CustomGroupToFilter = TEXT("PerimeterCurves");
+	PerimeterCurveVisualization.VisualizationMode = EPVDebugValueVisualizationMode::Curve;
 
 	FPVVisualizationSettings& EndPointVisualization = DebugVisualizationSettings.VisualizationSettings.AddDefaulted_GetRef();
 	EndPointVisualization.DebugType = EPVDebugType::Custom;
@@ -145,6 +145,40 @@ bool FPVTexture2DImporterElement::ExecuteInternal(FPCGContext* InContext) const
 
 	UPVGrowthData* OutVariantData = FPCGContext::NewObject_AnyThread<UPVGrowthData>(InContext);
 	OutVariantData->Initialize(MoveTemp(ImportTexture2DOutput.OutputCollection));
+	
+	{
+		TArray<FPVVisualizationSettings> VisualizationSettings;
+
+		FPVVisualizationSettings& PerimeterCurveVisualization = VisualizationSettings.AddDefaulted_GetRef();
+		PerimeterCurveVisualization.DebugType = EPVDebugType::Custom;
+		PerimeterCurveVisualization.bShowAnchorPoints = false;
+		PerimeterCurveVisualization.AttributeToFilter = TEXT("PerimeterCurve");
+		PerimeterCurveVisualization.CustomGroupToFilter = TEXT("PerimeterCurves");
+		PerimeterCurveVisualization.VisualizationMode = EPVDebugValueVisualizationMode::Curve;
+
+		FPVVisualizationSettings& EndPointVisualization = VisualizationSettings.AddDefaulted_GetRef();
+		EndPointVisualization.DebugType = EPVDebugType::Custom;
+		EndPointVisualization.bShowAnchorPoints = false;
+		EndPointVisualization.bDrawPointAsMesh = false;
+		EndPointVisualization.bUsePivotAsPosition = false;
+		EndPointVisualization.AttributeToFilter = TEXT("Vertex");
+		EndPointVisualization.CustomGroupToFilter = TEXT("EndPoints");
+		EndPointVisualization.VisualizationMode = EPVDebugValueVisualizationMode::Point;
+		EndPointVisualization.Color = FColor::Red;
+		EndPointVisualization.GizmoScale = 6.5f;
+
+		FPVVisualizationSettings& RootPointsVisualization = VisualizationSettings.AddDefaulted_GetRef();
+		RootPointsVisualization.DebugType = EPVDebugType::Custom;
+		RootPointsVisualization.bShowAnchorPoints = false;
+		RootPointsVisualization.bDrawPointAsMesh = false;
+		RootPointsVisualization.AttributeToFilter = TEXT("Point");
+		RootPointsVisualization.CustomGroupToFilter = TEXT("InputRootPoints");
+		RootPointsVisualization.VisualizationMode = EPVDebugValueVisualizationMode::Point;
+		RootPointsVisualization.Color = FColor::Blue;
+		RootPointsVisualization.GizmoScale = 10;
+
+		OutVariantData->AddVisualizationCollection(MoveTemp(ImportTexture2DOutput.VisualizationCollection), MoveTemp(VisualizationSettings));
+	}
 
 	if (TextureVisualizationMaterialInstance)
 	{
diff --git a/Engine/Plugins/Experimental/ProceduralVegetationEditor/Source/ProceduralVegetation/Public/DataTypes/PVData.h b/Engine/Plugins/Experimental/ProceduralVegetationEditor/Source/ProceduralVegetation/Public/DataTypes/PVData.h
--- a/Engine/Plugins/Experimental/ProceduralVegetationEditor/Source/ProceduralVegetation/Public/DataTypes/PVData.h
+++ b/Engine/Plugins/Experimental/ProceduralVegetationEditor/Source/ProceduralVegetation/Public/DataTypes/PVData.h
@@ -167,6 +167,12 @@ struct FPVDebugSettings
 	bool bAutoFocusLoopDebug = false;
 };
 
+struct FPVVisualizationCollection
+{
+	FManagedArrayCollection Collection;
+	TArray<FPVVisualizationSettings> VisualizationSettings;
+};
+
 UCLASS(BlueprintType, ClassGroup = (Procedural))
 class PROCEDURALVEGETATION_API UPVData : public UPCGSpatialData
 {
@@ -198,23 +204,16 @@ public:
 			ManagedResources.AddUnique(InResource);
 		}
 	}
-	
+
 #if WITH_EDITORONLY_DATA
+	void AddVisualizationCollection(FManagedArrayCollection InCollection, TArray<FPVVisualizationSettings> InVisualizationSettings = {});
 	void SetDebugSettings(const FPVDebugSettings& InSettings) const { DebugSettings = InSettings; }
 	void SetDebugSettings(FPVDebugSettings&& InSettings) const { DebugSettings = MoveTemp(InSettings); }
 	const FPVDebugSettings& GetDebugSettings() const { return DebugSettings; }
+	const TArray<FPVVisualizationCollection>& GetVisualizationCollections() const { return VisualizationCollections; }
 #endif
 	
 protected:
-
-	FManagedArrayCollection Collection;
-
-	UPROPERTY(Transient)
-	TArray<TObjectPtr<UObject>> ManagedResources;
-
-#if WITH_EDITOR
-	mutable FPVDebugSettings DebugSettings;
-#endif
 	
 	// ~Begin UPCGData interface
 	virtual bool SupportsFullDataCrc() const override { return true; }
@@ -228,8 +227,21 @@ public:
 	//~End UPCGSpatialData interface
 	
 private:
+
 	const UPCGBasePointData* ToBasePointData(FPCGContext* Context, const FBox& InBounds, TSubclassOf<UPCGBasePointData> PointDataClass) const;
+
 protected:
 	
+	FManagedArrayCollection Collection;
+
+	UPROPERTY(Transient)
+	TArray<TObjectPtr<UObject>> ManagedResources;
+
+#if WITH_EDITOR
+	mutable FPVDebugSettings DebugSettings;
+
+	TArray<FPVVisualizationCollection> VisualizationCollections;
+#endif
+
 	mutable FBox CachedBounds = FBox(EForceInit::ForceInit);
 };
diff --git a/Engine/Plugins/Experimental/ProceduralVegetationEditor/Source/ProceduralVegetation/Public/Helpers/PVFacadeHelper.h b/Engine/Plugins/Experimental/ProceduralVegetationEditor/Source/ProceduralVegetation/Public/Helpers/PVFacadeHelper.h
--- a/Engine/Plugins/Experimental/ProceduralVegetationEditor/Source/ProceduralVegetation/Public/Helpers/PVFacadeHelper.h
+++ b/Engine/Plugins/Experimental/ProceduralVegetationEditor/Source/ProceduralVegetation/Public/Helpers/PVFacadeHelper.h
@@ -12,20 +12,26 @@ namespace PV::Facades
 
 namespace PV::FacadeHelper
 {
+	enum class EForEachResult
+	{
+		Break    = 0,
+		Continue = 1
+	};
+
 	template <typename StateType>
 	void WalkBranches(
 		const PV::Facades::FPlantFacade& PlantFacade, 
 		const PV::Facades::FBranchFacade& BranchFacade, 
 		const StateType& RootBranchState, 
-		const TFunction<bool(int32, StateType&)>& Body
+		const TFunction<EForEachResult(int32, StateType&)>& Body
 	)
 	{
-		const TFunction<bool(int32, const StateType&)> WalkBranches_Recursive = [&](int32 BranchIndex, const StateType& State) -> bool
+		const TFunction<EForEachResult(int32, const StateType&)> WalkBranches_Recursive = [&](int32 BranchIndex, const StateType& State) -> EForEachResult
 		{
 			StateType BranchState = State;
-			if (!Body(BranchIndex, BranchState))
+			if (Body(BranchIndex, BranchState) == EForEachResult::Break)
 			{
-				return false;
+				return EForEachResult::Break;
 			}
 
 			const TArray<int32>& BranchChildren = BranchFacade.GetChildren(BranchIndex);
@@ -34,18 +40,18 @@ namespace PV::FacadeHelper
 				const int32 ChildBranchIndex = BranchFacade.GetBranchNumbers().Find(BranchChild);
 				check(ChildBranchIndex != INDEX_NONE);
 
-				if (!WalkBranches_Recursive(ChildBranchIndex, BranchState))
+				if (WalkBranches_Recursive(ChildBranchIndex, BranchState) == EForEachResult::Break)
 				{
-					return false;
+					return EForEachResult::Break;
 				}
 			}
 
-			return true;
+			return EForEachResult::Continue;
 		};
 
 		for (const int32 TrunkIndex : PlantFacade.GetTrunkIndices())
 		{
-			if (!WalkBranches_Recursive(TrunkIndex, RootBranchState))
+			if (WalkBranches_Recursive(TrunkIndex, RootBranchState) == EForEachResult::Break)
 			{
 				return;
 			}
@@ -55,34 +61,34 @@ namespace PV::FacadeHelper
 	inline void WalkBranches(
 		const PV::Facades::FPlantFacade& PlantFacade, 
 		const PV::Facades::FBranchFacade& BranchFacade, 
-		const TFunction<bool(int32 BranchIndex, int32 Depth)>& Body
+		const TFunction<EForEachResult(int32 BranchIndex, int32 Depth)>& Body
 	)
 	{
 		int32 Depth = 0;
-		WalkBranches<int32>(PlantFacade, BranchFacade, Depth, [&Body](int32 BranchIndex, int32& Depth) 
+		WalkBranches<int32>(PlantFacade, BranchFacade, Depth, [&Body](int32 BranchIndex, int32& Depth) -> EForEachResult
 		{
-			if (!Body(BranchIndex, Depth))
+			if (Body(BranchIndex, Depth) == EForEachResult::Break)
 			{
-				return false;
+				return EForEachResult::Break;
 			}
 			Depth += 1;
-			return true;
+			return EForEachResult::Continue;
 		});
 	}
 
 	inline void WalkBranches(
 		const PV::Facades::FPlantFacade& PlantFacade, 
 		const PV::Facades::FBranchFacade& BranchFacade, 
-		const TFunction<bool(int32 BranchIndex)>& Body
+		const TFunction<EForEachResult(int32 BranchIndex)>& Body
 	)
 	{
-		WalkBranches<int32>(PlantFacade, BranchFacade, INDEX_NONE, [&](int32 BranchIndex, int32&) 
+		WalkBranches<int32>(PlantFacade, BranchFacade, INDEX_NONE, [&](int32 BranchIndex, int32&) -> EForEachResult
 		{
-			if (!Body(BranchIndex))
+			if (Body(BranchIndex) == EForEachResult::Break)
 			{
-				return false;
+				return EForEachResult::Break;
 			}
-			return true;
+			return EForEachResult::Continue;
 		});
 	}
 
@@ -119,15 +125,15 @@ namespace PV::FacadeHelper
 	);
 
 	template <typename StateType, EWalkBranchMode WalkBranchMode = EWalkBranchMode::SplitPoints>
-	bool WalkBranchPoints(
+	EForEachResult WalkBranchPoints(
 		const PV::Facades::FBranchFacade& BranchFacade,
 		const PV::Facades::FPointFacade& PointFacade,
 		int32 BranchIndex,
 		const StateType& RootPointState,
-		const TFunction<bool(int32 BranchIndex, int32 BranchPointIndex, int32 PointIndex, const TArray<int32>& ChildPointIndices, StateType& PointState)>& Body
+		const TFunction<EForEachResult(int32 BranchIndex, int32 BranchPointIndex, int32 PointIndex, const TArray<int32>& ChildPointIndices, StateType& PointState)>& Body
 	)
 	{
-		const TFunction<bool(int32, bool, const StateType&)> WalkPoints = [&](int32 BranchIndex, bool bSkipFirstBranchPoint, const StateType& ParentState) -> bool
+		const TFunction<EForEachResult(int32, bool, const StateType&)> WalkPoints = [&](int32 BranchIndex, bool bSkipFirstBranchPoint, const StateType& ParentState) -> EForEachResult
 		{
 			StateType PointState = ParentState;
 			const TArray<int32>& CurrentBranchPoints = BranchFacade.GetPoints(BranchIndex);
@@ -138,32 +144,32 @@ namespace PV::FacadeHelper
 				const int32 PointIndex = CurrentBranchPoints[Index];
 				const TArray<int32> BranchPointChildren = GetBranchPointChildPoints<WalkBranchMode>(BranchFacade, PointFacade, BranchIndex, PointIndex);
 
-				if (!Body(BranchIndex, Index, PointIndex, BranchPointChildren, PointState))
+				if (Body(BranchIndex, Index, PointIndex, BranchPointChildren, PointState) == EForEachResult::Break)
 				{
-					return false;
+					return EForEachResult::Break;
 				}
 
 				for (const int32 ChildPointIndex : BranchPointChildren)
 				{
-					if (!WalkPoints(ChildPointIndex, WalkBranchMode == EWalkBranchMode::FusedPoints, PointState))
+					if (WalkPoints(ChildPointIndex, WalkBranchMode == EWalkBranchMode::FusedPoints, PointState) == EForEachResult::Break)
 					{
-						return false;
+						return EForEachResult::Break;
 					}
 				}
 			}
 
-			return true;
+			return EForEachResult::Continue;
 		};
 
 		return WalkPoints(BranchIndex, false, RootPointState);
 	}
 
 	template <EWalkBranchMode WalkBranchMode = EWalkBranchMode::SplitPoints>
-	bool WalkBranchPoints(
+	EForEachResult WalkBranchPoints(
 		const PV::Facades::FBranchFacade& BranchFacade,
 		const PV::Facades::FPointFacade& PointFacade,
 		int32 BranchIndex,
-		const TFunction<bool(int32 BranchIndex, int32 BranchPointIndex, int32 PointIndex, const TArray<int32>& ChildPointIndices)>& Body
+		const TFunction<EForEachResult(int32 BranchIndex, int32 BranchPointIndex, int32 PointIndex, const TArray<int32>& ChildPointIndices)>& Body
 	)
 	{
 		return WalkBranchPoints<int32, WalkBranchMode>(
@@ -171,7 +177,7 @@ namespace PV::FacadeHelper
 			PointFacade,
 			BranchIndex,
 			0,
-			[&](int32 BranchIndex, int32 BranchPointIndex, int32 PointIndex, const TArray<int32>& ChildPointIndices, int32&)
+			[&](int32 BranchIndex, int32 BranchPointIndex, int32 PointIndex, const TArray<int32>& ChildPointIndices, int32&) -> EForEachResult
 			{
 				return Body(BranchIndex, BranchPointIndex, PointIndex, ChildPointIndices);
 			}
@@ -184,12 +190,12 @@ namespace PV::FacadeHelper
 		const PV::Facades::FBranchFacade& BranchFacade,
 		const PV::Facades::FPointFacade& PointFacade,
 		const StateType& RootPointState,
-		const TFunction<bool(int32 BranchIndex, int32 BranchPointIndex, int32 PointIndex, const TArray<int32>& ChildPointIndices, StateType& PointState)>& Body
+		const TFunction<EForEachResult(int32 BranchIndex, int32 BranchPointIndex, int32 PointIndex, const TArray<int32>& ChildPointIndices, StateType& PointState)>& Body
 	)
 	{
 		for (const int32 TrunkIndex : PlantFacade.GetTrunkIndices())
 		{
-			if (!WalkBranchPoints<StateType, WalkBranchMode>(BranchFacade, PointFacade, TrunkIndex, RootPointState, Body))
+			if (WalkBranchPoints<StateType, WalkBranchMode>(BranchFacade, PointFacade, TrunkIndex, RootPointState, Body) == EForEachResult::Break)
 			{
 				return;
 			}
@@ -201,7 +207,7 @@ namespace PV::FacadeHelper
 		const PV::Facades::FPlantFacade& PlantFacade,
 		const PV::Facades::FBranchFacade& BranchFacade,
 		const PV::Facades::FPointFacade& PointFacade,
-		const TFunction<bool(int32 BranchIndex, int32 BranchPointIndex, int32 PointIndex, const TArray<int32>& ChildPointIndices)>& Body
+		const TFunction<EForEachResult(int32 BranchIndex, int32 BranchPointIndex, int32 PointIndex, const TArray<int32>& ChildPointIndices)>& Body
 	)
 	{
 		WalkPlantPoints<int32, WalkBranchMode>(PlantFacade, BranchFacade, PointFacade, INDEX_NONE, [&](int32 BranchIndex, int32 BranchPointIndex, int32 PointIndex, const TArray<int32>& ChildPointIndices, int32&)
@@ -209,4 +215,62 @@ namespace PV::FacadeHelper
 			return Body(BranchIndex, BranchPointIndex, PointIndex, ChildPointIndices);
 		});
 	}
+
+	inline void ForEachPlant(
+		const PV::Facades::FPlantFacade& PlantFacade,
+		const TFunction<EForEachResult(int32 PlantNumber)>& Body
+	)
+	{
+		for (int32 PlantNumber : PlantFacade.GetPlantNumbers())
+		{
+			if (Body(PlantNumber) == EForEachResult::Break)
+			{
+				return;
+			}
+		}
+	}
+
+	inline void ForEachPlantPoint(
+		int32 PlantNumber,
+		const PV::Facades::FPlantFacade& PlantFacade,
+		const PV::Facades::FBranchFacade& BranchFacade,
+		const TFunction<EForEachResult(int32 BranchIndex, int32 PointIndex)>& Body
+	)
+	{
+		const TArray<int32> BranchIndices = PlantFacade.GetBranchIndices(PlantNumber);
+		for (int32 BranchIndex : BranchIndices)
+		{
+			for (int32 PointIndex : BranchFacade.GetPoints(BranchIndex))
+			{
+				if (Body(BranchIndex, PointIndex) == EForEachResult::Break)
+				{
+					return;
+				}
+			}
+		}
+	}
+
+	inline void ForEachPoint(
+		const PV::Facades::FPlantFacade& PlantFacade,
+		const PV::Facades::FBranchFacade& BranchFacade,
+		const TFunction<EForEachResult(int32 PlantNumber, int32 BranchIndex, int32 PointIndex)>& Body
+	)
+	{
+		ForEachPlant(PlantFacade, [&](int32 PlantNumber)->EForEachResult
+		{
+			const TArray<int32> BranchIndices = PlantFacade.GetBranchIndices(PlantNumber);
+			for (int32 BranchIndex : BranchIndices)
+			{
+				for (int32 PointIndex : BranchFacade.GetPoints(BranchIndex))
+				{
+					if (Body(PlantNumber, BranchIndex, PointIndex) == EForEachResult::Break)
+					{
+						return EForEachResult::Break;
+					}
+				}
+			}
+
+			return EForEachResult::Continue;
+		});
+	}
 }
diff --git a/Engine/Plugins/Experimental/ProceduralVegetationEditor/Source/ProceduralVegetation/Public/PVImportTexture2DParams.h b/Engine/Plugins/Experimental/ProceduralVegetationEditor/Source/ProceduralVegetation/Public/PVImportTexture2DParams.h
--- a/Engine/Plugins/Experimental/ProceduralVegetationEditor/Source/ProceduralVegetation/Public/PVImportTexture2DParams.h
+++ b/Engine/Plugins/Experimental/ProceduralVegetationEditor/Source/ProceduralVegetation/Public/PVImportTexture2DParams.h
@@ -80,11 +80,12 @@ struct FPVImportTexture2DParams
 	UPROPERTY(EditAnywhere, Category = "Branch Geneneration|Tip detection", meta = (ClampMin=0, ClampMax=0.01, UIMin=0, UIMax=0.01))
 	float MaxTipAngleSearchDist = 0.005;
 
-	UPROPERTY(EditAnywhere, Category = "Inspect")
+	UPROPERTY(EditAnywhere, Category = "DebugSettings")
 	EPVImportTexture2DDebugState DebugState = EPVImportTexture2DDebugState::None;
 };
 
 struct FPVImportTexture2DOutput
 {
+	FManagedArrayCollection VisualizationCollection;
 	FManagedArrayCollection OutputCollection;
 };
\ No newline at end of file
diff --git a/Engine/Plugins/Experimental/ProceduralVegetationEditor/Source/ProceduralVegetationEditor/Private/Importer/PVImportHelpers.cpp b/Engine/Plugins/Experimental/ProceduralVegetationEditor/Source/ProceduralVegetationEditor/Private/Importer/PVImportHelpers.cpp
--- a/Engine/Plugins/Experimental/ProceduralVegetationEditor/Source/ProceduralVegetationEditor/Private/Importer/PVImportHelpers.cpp
+++ b/Engine/Plugins/Experimental/ProceduralVegetationEditor/Source/ProceduralVegetationEditor/Private/Importer/PVImportHelpers.cpp
@@ -481,12 +481,12 @@ const auto PV::DynamicMeshHelper::RootPointAttributeDefinition       = MakeUniqu
 const auto PV::DynamicMeshHelper::LengthFromRootAttributeDefinition  = MakeUnique<PV::TDynamicMeshVertexAttributeDefinition<float>>(TEXT("LengthFromRoot"));
 const auto PV::DynamicMeshHelper::NextBranchPointAttributeDefinition = MakeUnique<PV::TDynamicMeshVertexAttributeDefinition<int32>>(TEXT("NextBranchPoint"));
 
-void PV::DynamicMeshHelper::AddDynamicMeshToCollection(const UE::Geometry::FDynamicMesh3& DynamicMesh, FManagedArrayCollection& InOutCollection)
+void PV::DynamicMeshHelper::AddDynamicMeshToCollection(const UE::Geometry::FDynamicMesh3& DynamicMesh, FManagedArrayCollection& InOutCollection, const FTransform& Offset)
 {
 	FGeometryCollection MeshGeomCollection;
 
 	UE::Geometry::FGeometryCollectionToDynamicMeshes::FToCollectionOptions Options2;
-	UE::Geometry::FGeometryCollectionToDynamicMeshes::AppendMeshToCollection(MeshGeomCollection, DynamicMesh, FTransform::Identity, Options2);
+	UE::Geometry::FGeometryCollectionToDynamicMeshes::AppendMeshToCollection(MeshGeomCollection, DynamicMesh, Offset, Options2);
 
 	AppendGeometryFromManagedArrayCollection(MeshGeomCollection, InOutCollection);
 	SetManagedArrayCollectionGeometry(MeshGeomCollection, InOutCollection);
@@ -957,7 +957,7 @@ void PV::Import::Helper::FillGrowthDataFromFoliageDescription(
 
 void PV::Import::Helper::ComputeLengthFromRoot(PV::Facades::FPointFacade& PointFacade, const PV::Facades::FBranchFacade& BranchFacade, const PV::Facades::FPlantFacade& PlantFacade)
 {
-	PV::FacadeHelper::WalkBranches(PlantFacade, BranchFacade, [&](int32 BranchIndex)->bool
+	PV::FacadeHelper::WalkBranches(PlantFacade, BranchFacade, [&](int32 BranchIndex)->PV::FacadeHelper::EForEachResult
 	{
 		TManagedArray<float>& LengthFromRoots = PointFacade.ModifyLengthFromRoots();
 		TManagedArray<float>& LengthFromSeeds = PointFacade.ModifyLengthFromSeeds();
@@ -1035,7 +1035,7 @@ void PV::Import::Helper::ComputeLengthFromRoot(PV::Facades::FPointFacade& PointF
 			LengthFromSeeds[PointIndex] = LengthFromRoots[PointIndex];
 		}
 
-		return true;
+		return PV::FacadeHelper::EForEachResult::Continue;
 	});
 }
 
@@ -1056,7 +1056,7 @@ void PV::Import::Helper::ComputeBudDevelopment(PV::Facades::FPointFacade& PointF
 		BranchFacade, 
 		PointFacade,
 		0, 
-		[&](int32 BranchIndex, int32 BranchPointIndex, int32 PointIndex, const TArray<int32>& ChildPointIndices, int32& NumWalkedPoints)->bool
+		[&](int32 BranchIndex, int32 BranchPointIndex, int32 PointIndex, const TArray<int32>& ChildPointIndices, int32& NumWalkedPoints)->PV::FacadeHelper::EForEachResult
 		{
 			NumWalkedPoints += 1;
 
@@ -1077,7 +1077,7 @@ void PV::Import::Helper::ComputeBudDevelopment(PV::Facades::FPointFacade& PointF
 
 			PointFacade.SetBudDevelopment(PointIndex, MoveTemp(BudDevelopment));
 
-			return true;
+			return PV::FacadeHelper::EForEachResult::Continue;
 		}
 	);
 }
@@ -1103,7 +1103,7 @@ void PV::Import::Helper::ComputeBudDirections(PV::Facades::FPointFacade& PointFa
 		BranchFacade,
 		PointFacade,
 		INDEX_NONE,
-		[&](int32 BranchIndex, int32 BranchPointIndex, int32 PointIndex, const TArray<int32>& ChildPointIndices, int32& PrevPointIndex)->bool
+		[&](int32 BranchIndex, int32 BranchPointIndex, int32 PointIndex, const TArray<int32>& ChildPointIndices, int32& PrevPointIndex)->PV::FacadeHelper::EForEachResult
 		{
 			const TArray<int32>& BranchPoints = BranchFacade.GetPoints(BranchIndex);
 
@@ -1141,14 +1141,14 @@ void PV::Import::Helper::ComputeBudDirections(PV::Facades::FPointFacade& PointFa
 
 			PrevPointIndex = PointIndex;
 
-			return true;
+			return PV::FacadeHelper::EForEachResult::Continue;
 		}
 	);
 }
 
 void PV::Import::Helper::InterpolatePointsScales(PV::Facades::FPointFacade& PointFacade, const PV::Facades::FBranchFacade& BranchFacade, const PV::Facades::FPlantFacade& PlantFacade)
 {
-	PV::FacadeHelper::WalkBranches(PlantFacade, BranchFacade, [&](int32 BranchIndex)
+	PV::FacadeHelper::WalkBranches(PlantFacade, BranchFacade, [&](int32 BranchIndex)->PV::FacadeHelper::EForEachResult
 	{
 		const TArray<int32>& BranchPoints = BranchFacade.GetPoints(BranchIndex);
 
@@ -1203,7 +1203,7 @@ void PV::Import::Helper::InterpolatePointsScales(PV::Facades::FPointFacade& Poin
 			i = NextValidScale_Index + 1;
 		}
 
-		return true;
+		return PV::FacadeHelper::EForEachResult::Continue;
 	});
 }
 
diff --git a/Engine/Plugins/Experimental/ProceduralVegetationEditor/Source/ProceduralVegetationEditor/Private/Importer/PVImportHelpers.h b/Engine/Plugins/Experimental/ProceduralVegetationEditor/Source/ProceduralVegetationEditor/Private/Importer/PVImportHelpers.h
--- a/Engine/Plugins/Experimental/ProceduralVegetationEditor/Source/ProceduralVegetationEditor/Private/Importer/PVImportHelpers.h
+++ b/Engine/Plugins/Experimental/ProceduralVegetationEditor/Source/ProceduralVegetationEditor/Private/Importer/PVImportHelpers.h
@@ -40,7 +40,7 @@ namespace PV::DynamicMeshHelper
 	extern const TUniquePtr<TDynamicMeshVertexAttributeDefinition<float>> LengthFromRootAttributeDefinition;
 	extern const TUniquePtr<TDynamicMeshVertexAttributeDefinition<int32>> NextBranchPointAttributeDefinition;
 
-	void AddDynamicMeshToCollection(const UE::Geometry::FDynamicMesh3& DynamicMesh, FManagedArrayCollection& InOutCollection);
+	void AddDynamicMeshToCollection(const UE::Geometry::FDynamicMesh3& DynamicMesh, FManagedArrayCollection& InOutCollection, const FTransform& Offset = FTransform::Identity);
 	
 	void AddBranchHierarchiesToCollection(const TArray<FPVBranchHierarchyDescription>& BranchHierarchyDescriptions, FManagedArrayCollection& OutCollection);
 
diff --git a/Engine/Plugins/Experimental/ProceduralVegetationEditor/Source/ProceduralVegetationEditor/Private/Importer/PVImporter_Texture2D.cpp b/Engine/Plugins/Experimental/ProceduralVegetationEditor/Source/ProceduralVegetationEditor/Private/Importer/PVImporter_Texture2D.cpp
--- a/Engine/Plugins/Experimental/ProceduralVegetationEditor/Source/ProceduralVegetationEditor/Private/Importer/PVImporter_Texture2D.cpp
+++ b/Engine/Plugins/Experimental/ProceduralVegetationEditor/Source/ProceduralVegetationEditor/Private/Importer/PVImporter_Texture2D.cpp
@@ -8,6 +8,7 @@
 #include "IntVectorTypes.h"
 #include "PCGSettings.h"
 #include "GeometryCollection/ManagedArrayCollection.h"
+#include "Helpers/PVFacadeHelper.h"
 
 #include "Engine/Texture2D.h"
 
@@ -36,13 +37,13 @@ namespace Internal
 	const static FName PixelsDebugGroupName = TEXT("PixelDebug");
 	const static FName PixelPositionsAttributeName = TEXT("Position");
 	const static FName PixelValueAttributeName = TEXT("Value");
-	const static FName PerimiterCurveGroupName = TEXT("PerimiterCurves");
-	const static FName PerimiterCurveAttributeName = TEXT("PerimiterCurve");
+	const static FName PerimeterCurveGroupName = TEXT("PerimeterCurves");
+	const static FName PerimeterCurveAttributeName = TEXT("PerimeterCurve");
 	const static FName InputRootPointsGroupName = TEXT("InputRootPoints");
 	const static FName RootPointsGroupName = TEXT("RootPoints");
 	const static FName RootPointsAttributeName = TEXT("Point");
 
-	const static PV::TDynamicMeshVertexAttributeDefinition<bool> PerimiterAttributeDefinition(TEXT("Perimiter"));
+	const static PV::TDynamicMeshVertexAttributeDefinition<bool> PerimeterAttributeDefinition(TEXT("Perimeter"));
 	
 	static int32 GetPixelIndex(const FVector2i& PixelCoords, int32 ImageWidth)
 	{
@@ -62,7 +63,7 @@ namespace Internal
 		return UE::Math::TVector<Vector3dType>(0, Vec2D.X * Scale, (1.0 - Vec2D.Y) * Scale);
 	}
 
-	static void AddDebugRootPointsToCollection(FManagedArrayCollection& InOutCollection, const TArray<FPVImportTexture2DRootPoint>& InRootPoints, float Scale)
+	static void AddRootPointsToCollection(FManagedArrayCollection& InOutCollection, const TArray<FPVImportTexture2DRootPoint>& InRootPoints, float Scale, const FVector& Offset = FVector::ZeroVector)
 	{
 		InOutCollection.AddGroup(InputRootPointsGroupName);
 		InOutCollection.AddElements(InRootPoints.Num(), InputRootPointsGroupName);
@@ -70,7 +71,7 @@ namespace Internal
 		auto& InputRootPointAttribute = InOutCollection.AddAttribute<FVector3f>(RootPointsAttributeName, InputRootPointsGroupName);
 		for (int32 i = 0; i < InRootPoints.Num(); i++)
 		{
-			InputRootPointAttribute[i] = ToVector3(InRootPoints[i].Position, Scale);
+			InputRootPointAttribute[i] = ToVector3(InRootPoints[i].Position, Scale) + FVector3f(Offset);
 		}
 	}
 	
@@ -97,23 +98,55 @@ namespace Internal
 		}
 	}
 
-	static void AddDebugPerimiterCurvesToCollection(FManagedArrayCollection& InOutCollection, const TArray<TArray<FVector2f>>& InPerimiterCurves, float Scale)
+	static void AddPerimeterCurvesToCollection(
+		FManagedArrayCollection& InOutCollection, 
+		const TArray<TArray<FVector2f>>& InPerimeterCurves, 
+		float Scale, 
+		const FVector& Offset = FVector::ZeroVector
+	)
 	{
-		InOutCollection.AddGroup(PerimiterCurveGroupName);
-		InOutCollection.AddElements(InPerimiterCurves.Num(), PerimiterCurveGroupName);
+		InOutCollection.AddGroup(PerimeterCurveGroupName);
+		InOutCollection.AddElements(InPerimeterCurves.Num(), PerimeterCurveGroupName);
     	
-		auto& PerimiterCurveAttribute = InOutCollection.AddAttribute<TArray<FVector3f>>(PerimiterCurveAttributeName, PerimiterCurveGroupName);
-		for (int32 i = 0; i < InPerimiterCurves.Num(); i++)
+		auto& PerimeterCurveAttribute = InOutCollection.AddAttribute<TArray<FVector3f>>(PerimeterCurveAttributeName, PerimeterCurveGroupName);
+		for (int32 i = 0; i < InPerimeterCurves.Num(); i++)
 		{
-			const TArray<FVector2f>& Curve = InPerimiterCurves[i];
+			const TArray<FVector2f>& Curve = InPerimeterCurves[i];
 			for (const FVector2f& CurvePos : Curve)
 			{
-				PerimiterCurveAttribute[i].Add(ToVector3(CurvePos, Scale));
+				PerimeterCurveAttribute[i].Add(ToVector3(CurvePos, Scale) + FVector3f(Offset));
 			}
 		}
 	}
 
-	static void AddDebugTextureQuadToCollection(FManagedArrayCollection& InOutCollection, UTexture2D& Texture2D, const UMaterialInterface* TextureVisualizationMaterial, float Scale)
+	static void AddDynamicMeshVerticesToCollection(
+		FManagedArrayCollection& InOutCollection, 
+		const FName GroupName,
+		const FName AttributeName,
+		const UE::Geometry::FDynamicMesh3& DynamicMesh, 
+		const FBooleanVertexAttribute& Filter,
+		const FVector& Offset
+	)
+	{
+		const TArray<int32> VertexIndices = Filter.FindAllNonZero();
+		
+		InOutCollection.AddGroup(GroupName);
+		InOutCollection.AddElements(VertexIndices.Num(), GroupName);
+
+		auto& Attribute = InOutCollection.AddAttribute<FVector3f>(AttributeName, GroupName);
+		for (int32 i = 0; i < VertexIndices.Num(); ++i)
+		{
+			Attribute[i] = FVector3f(DynamicMesh.GetVertex(VertexIndices[i]) + Offset);
+		}
+	}
+
+	static void AddTextureQuadToCollection(
+		FManagedArrayCollection& InOutCollection,
+		UTexture2D& Texture2D,
+		const UMaterialInterface* TextureVisualizationMaterial,
+		float Scale,
+		const FVector& Offset = FVector::ZeroVector
+	)
 	{
 		const TArray<FVector> Points = {
 			FVector(0, 0, 1),
@@ -142,7 +175,7 @@ namespace Internal
 			Normals,
 			Triangles,
 			UVs,
-			FTransform(FRotator::ZeroRotator, FVector::ZeroVector, FVector(Scale)),
+			FTransform(FRotator::ZeroRotator, Offset, FVector(Scale)),
 			FTransform::Identity,
 			1
 		);
@@ -154,6 +187,48 @@ namespace Internal
 		MeshGeomCollection->CopyTo(&InOutCollection);
 	}
 
+	static void ArrangePlantsInGrid(FManagedArrayCollection& InCollection, const FVector& InOffset = FVector::ZeroVector)
+	{
+		const PV::Facades::FPlantFacade PlantFacade(InCollection);
+		const PV::Facades::FBranchFacade BranchFacade(InCollection);
+		PV::Facades::FPointFacade PointFacade(InCollection);
+
+		struct FPlantOffsetInfo
+		{
+			int32 PlantNumber;
+			FBox3f Bounds;
+		};
+		TArray<FPlantOffsetInfo> PlantsOffsetInfo;
+
+		PV::FacadeHelper::ForEachPlant(PlantFacade, [&](int32 PlantNumber)
+		{
+			FPlantOffsetInfo& PlantOffsetInfo = PlantsOffsetInfo.AddDefaulted_GetRef();
+			PlantOffsetInfo.PlantNumber = PlantNumber;
+
+			PV::FacadeHelper::ForEachPlantPoint(PlantNumber, PlantFacade, BranchFacade, [&](int32 BranchIndex, int32 PointIndex)
+			{
+				PlantOffsetInfo.Bounds += PointFacade.GetPosition(PointIndex);
+				return PV::FacadeHelper::EForEachResult::Continue;
+			});
+			return PV::FacadeHelper::EForEachResult::Continue;
+		});
+
+		FVector3f Offset = FVector3f(InOffset);
+		for (int32 i = 0; i < PlantsOffsetInfo.Num(); ++i)
+		{
+			Offset.Y -= PlantsOffsetInfo[i].Bounds.Max.Y + 10.f;
+
+			PV::FacadeHelper::ForEachPlantPoint(PlantsOffsetInfo[i].PlantNumber, PlantFacade, BranchFacade, [&](int32 BranchIndex, int32 PointIndex)
+			{
+				FVector3f& PointPosition = PointFacade.ModifyPositions()[PointIndex];
+				PointPosition += Offset;
+				return PV::FacadeHelper::EForEachResult::Continue;
+			});
+
+			Offset.Y += PlantsOffsetInfo[i].Bounds.Min.Y;
+		}
+	}
+
 	static void SampleImage(FTextureSource& InTextureSource, FVector2i& InOutSampleSize, TArray<bool>& OutBlackAndWhiteImageData, bool bInvertImage, float WhiteLevel)
 	{
 		const FTextureSource::FMipLock LockedMip0(FTextureSource::ELockState::ReadOnly, &InTextureSource, 0);
@@ -357,9 +432,9 @@ namespace Internal
 	static void TriangulateCurvesAndSplitEdges(const TArray<TArray<FVector2f>>& InCurves, UE::Geometry::FDynamicMesh3& OutMesh, float Scale)
 	{
 		/*
-		* This function convertes the perimiter curves into a DynamicMesh and adds a vertex at the center of each edge spanning two perimiter points.
+		* This function convertes the perimeter curves into a DynamicMesh and adds a vertex at the center of each edge spanning two perimeter points.
 		* 
-		* It does this by first triangulating the curve using FDelaunay2, it then looks at each edge and if said edge has two perimiter points, it inserts a new vertex
+		* It does this by first triangulating the curve using FDelaunay2, it then looks at each edge and if said edge has two perimeter points, it inserts a new vertex
 		* at the center of the edge. It does this for all edges and then re-triangulates the entire mesh with this new vertex set.
 		* 
 		* The end result will look something like the following diagram:
@@ -374,7 +449,7 @@ namespace Internal
 		*   *     *       |/ | \|
 		*/
 
-		PerimiterAttributeDefinition.AttachAttribute(OutMesh);
+		PerimeterAttributeDefinition.AttachAttribute(OutMesh);
 
 		for (int32 CurveIndex = 0; CurveIndex < InCurves.Num(); CurveIndex++)
 		{
@@ -446,8 +521,8 @@ namespace Internal
 			}
 			
 			UE::Geometry::FDynamicMesh3 NewDynamicMesh;
-			PerimiterAttributeDefinition.AttachAttribute(NewDynamicMesh);
-			auto& PerimiterAttribute = PerimiterAttributeDefinition.GetAttributeChecked(NewDynamicMesh);
+			PerimeterAttributeDefinition.AttachAttribute(NewDynamicMesh);
+			auto& PerimeterAttribute = PerimeterAttributeDefinition.GetAttributeChecked(NewDynamicMesh);
 			
 			const TArray<FIndex3i> Triangles = Delaunay.GetFilledTriangles(OriginalCurveEdges, UE::Geometry::FDelaunay2::EFillMode::Solid);
 			const TArray<FVector2f>& Vertices = SubdividedMeshPoints;
@@ -461,7 +536,7 @@ namespace Internal
 				);
 				
 				const int32 VertexID = NewDynamicMesh.AppendVertex(VertexInfo);
-				PerimiterAttribute.SetValue(VertexID, i < Curve.Num());
+				PerimeterAttribute.SetValue(VertexID, i < Curve.Num());
 			}
 			
 			for (const FIndex3i& Triangle : Triangles)
@@ -475,24 +550,24 @@ namespace Internal
 
 	static void ComputeVertexScale2D(
 		const UE::Geometry::FDynamicMesh3& DynamicMesh,
-		const FBooleanVertexAttribute& PerimiterAttribute,
+		const FBooleanVertexAttribute& PerimeterAttribute,
 		FFloatVertexAttribute& VertexScaleAttribute
 	)
 	{
 		/*
-		* This function computes the vertex scale for each vertex in the mesh by finding the nearest point on the perimiter curve.
+		* This function computes the vertex scale for each vertex in the mesh by finding the nearest point on the perimeter curve.
 		*/
 
-		const auto FindDistToPerimiter = [&](const FVector3d& Point, const TArray<int32>& PerimiterVertices)
+		const auto FindDistToPerimeter = [&](const FVector3d& Point, const TArray<int32>& PerimeterVertices)
 		{
 			// TODO: This algorithm is N^2, we can probably reduce that by using FSparseDynamicOctree3
 
 			double ClosestDistSqr = MAX_flt;
 
-			for (int32 i = 0; i < PerimiterVertices.Num(); ++i)
+			for (int32 i = 0; i < PerimeterVertices.Num(); ++i)
 			{
-				const int32 PointIndex0 = PerimiterVertices[i];
-				const int32 PointIndex1 = PerimiterVertices[FMath::WrapExclusive(i + 1, 0, PerimiterVertices.Num())];
+				const int32 PointIndex0 = PerimeterVertices[i];
+				const int32 PointIndex1 = PerimeterVertices[FMath::WrapExclusive(i + 1, 0, PerimeterVertices.Num())];
 
 				const FVector3d& VertexLocation1 = DynamicMesh.GetVertex(PointIndex0);
 				const FVector3d& VertexLocation2 = DynamicMesh.GetVertex(PointIndex1);
@@ -513,14 +588,14 @@ namespace Internal
 
 		for (const auto& Component : ConnectedComponents)
 		{
-			TArray<int32> PerimiterVertices = PerimiterAttribute.FindAllNonZero(Component.Indices);
-			PerimiterVertices.Sort(); // This assumes perimiter vertices are ordered sequentially, which is currently the case.
+			TArray<int32> PerimeterVertices = PerimeterAttribute.FindAllNonZero(Component.Indices);
+			PerimeterVertices.Sort(); // This assumes perimeter vertices are ordered sequentially, which is currently the case.
 
 			for (int32 VertexID : Component.Indices)
 			{
-				if (!PerimiterAttribute.GetValue(VertexID))
+				if (!PerimeterAttribute.GetValue(VertexID))
 				{
-					const float VertexScale = FindDistToPerimiter(DynamicMesh.GetVertex(VertexID), PerimiterVertices);
+					const float VertexScale = FindDistToPerimeter(DynamicMesh.GetVertex(VertexID), PerimeterVertices);
 					VertexScaleAttribute.SetValue(VertexID, VertexScale);
 				}
 			}
@@ -532,12 +607,12 @@ namespace Internal
 		float MaxTipAngleSearchDist,
 		float TipAngleThresholdDeg,
 		FBooleanVertexAttribute& EndPointAttribute,
-		const FBooleanVertexAttribute& PerimiterAttribute
+		const FBooleanVertexAttribute& PerimeterAttribute
 	)
 	{
 		/*
-		* This function finds the end points (or "tips") on a 2D mesh by walking the perimiter vertices and checking the angle between itself and its nearest 
-		* two neighbours on the perimiter. To allow for curved (or flat-tipped) end points to be detected, the function takes a second pass where it walks outwards along the
+		* This function finds the end points (or "tips") on a 2D mesh by walking the perimeter vertices and checking the angle between itself and its nearest 
+		* two neighbours on the perimeter. To allow for curved (or flat-tipped) end points to be detected, the function takes a second pass where it walks outwards along the
 		* perimited (limited by MaxTipAngleSearchDist), and checks the angle of the current vertex towards the subsequent neighbours. (see diagram a for example if flat-tipped end point)
 		* 
 		*     a: ___     b:
@@ -550,20 +625,20 @@ namespace Internal
 
 		const static auto ComputeAngleDeg = [](
 			const UE::Geometry::FDynamicMesh3& Mesh,
-			const TArray<int32>& PerimiterVertices, 
+			const TArray<int32>& PerimeterVertices, 
 			int32 PerimeterIndex, 
 			int32 LeftPerimeterIndex, 
 			int32 RightPerimeterIndex)->float
 		{
-			check(PerimiterVertices.IsValidIndex(PerimeterIndex));
-			check(Mesh.IsVertex(PerimiterVertices[PerimeterIndex]));
+			check(PerimeterVertices.IsValidIndex(PerimeterIndex));
+			check(Mesh.IsVertex(PerimeterVertices[PerimeterIndex]));
 
-			LeftPerimeterIndex = FMath::WrapExclusive(LeftPerimeterIndex, 0, PerimiterVertices.Num());
-			RightPerimeterIndex = FMath::WrapExclusive(RightPerimeterIndex, 0, PerimiterVertices.Num());
+			LeftPerimeterIndex = FMath::WrapExclusive(LeftPerimeterIndex, 0, PerimeterVertices.Num());
+			RightPerimeterIndex = FMath::WrapExclusive(RightPerimeterIndex, 0, PerimeterVertices.Num());
 
-			const int32 VertexIndex = PerimiterVertices[PerimeterIndex];
-			const int32 NeighbourVertexIndex1 = PerimiterVertices[LeftPerimeterIndex];
-			const int32 NeighbourVertexIndex2 = PerimiterVertices[RightPerimeterIndex];
+			const int32 VertexIndex = PerimeterVertices[PerimeterIndex];
+			const int32 NeighbourVertexIndex1 = PerimeterVertices[LeftPerimeterIndex];
+			const int32 NeighbourVertexIndex2 = PerimeterVertices[RightPerimeterIndex];
 
 			const FVector3d& PointPos = Mesh.GetVertex(VertexIndex);
 			const FVector3d& PerimeterNeighbour1Pos = Mesh.GetVertex(NeighbourVertexIndex1);
@@ -581,71 +656,71 @@ namespace Internal
 
 		for (const auto& Component : ConnectedComponents)
 		{
-			const TArray<int32> PerimiterVertices = PerimiterAttribute.FindAllNonZero(Component.Indices);
+			const TArray<int32> PerimeterVertices = PerimeterAttribute.FindAllNonZero(Component.Indices);
 			 	
-			TArray<int32> ExcludedPerimiterVertices;
-			ExcludedPerimiterVertices.Reserve(PerimiterVertices.Num());
+			TArray<int32> ExcludedPerimeterVertices;
+			ExcludedPerimeterVertices.Reserve(PerimeterVertices.Num());
 			 	
-			TArray<float> PerimiterAngles;
-			PerimiterAngles.SetNumZeroed(PerimiterVertices.Num());
+			TArray<float> PerimeterAngles;
+			PerimeterAngles.SetNumZeroed(PerimeterVertices.Num());
 			 	
-			for (int32 PerimiterIndex = 0; PerimiterIndex < PerimiterVertices.Num(); ++PerimiterIndex)
+			for (int32 PerimeterIndex = 0; PerimeterIndex < PerimeterVertices.Num(); ++PerimeterIndex)
 			{
-				const int32 PointIndex = PerimiterVertices[PerimiterIndex];
+				const int32 PointIndex = PerimeterVertices[PerimeterIndex];
 			 	check(DynamicMesh.IsVertex(PointIndex));
 			 		
 			 	int32 NumInternalNeighbours = 0;
 				const TArray<int32> PointNeighbours = PV::DynamicMeshHelper::GetVertexNeighbours(DynamicMesh, PointIndex);
 			 	for (int32 NeighbourIndex : PointNeighbours)
 			 	{
-			 		if (!PerimiterAttribute.GetValue(NeighbourIndex))
+			 		if (!PerimeterAttribute.GetValue(NeighbourIndex))
 			 		{
 			 			NumInternalNeighbours += 1;
 			 		}
 			 	}
 			 	
 				// TODO: Since we do not know whether the tip is pointing inwards or outwards, we use the fact that most outward pointing tips only have one
-				// internal connection (e.g. a connection to a vertec which is not a perimiter vertex). This is however not a very reliable method of
-				// determining the direction of the tip. A better alternative would be to calculate the normal for each vertex in the perimiter curve
+				// internal connection (e.g. a connection to a vertec which is not a perimeter vertex). This is however not a very reliable method of
+				// determining the direction of the tip. A better alternative would be to calculate the normal for each vertex in the perimeter curve
 				// and use said normal to determine whether the tip is pointing outwads.
 			 	if (NumInternalNeighbours == 1)
 			 	{
-			 		const double TipAngle = ComputeAngleDeg(DynamicMesh, PerimiterVertices, PerimiterIndex, PerimiterIndex - 1, PerimiterIndex + 1);
+			 		const double TipAngle = ComputeAngleDeg(DynamicMesh, PerimeterVertices, PerimeterIndex, PerimeterIndex - 1, PerimeterIndex + 1);
 			 		if (TipAngle >= TipAngleThresholdDeg)
 			 		{
 						EndPointAttribute.SetValue(PointIndex, true);
-						PerimiterAngles[PerimiterIndex] = TipAngle;
+						PerimeterAngles[PerimeterIndex] = TipAngle;
 			 		}
 			 		else 
 			 		{
-						ExcludedPerimiterVertices.Add(PerimiterIndex);
+						ExcludedPerimeterVertices.Add(PerimeterIndex);
 			 		}
 			 	}
 			}
 			 	
 			const float MaxTipAngleSearchDistSqr = MaxTipAngleSearchDist * MaxTipAngleSearchDist;
-			for (int32 PerimiterIndex : ExcludedPerimiterVertices)
+			for (int32 PerimeterIndex : ExcludedPerimeterVertices)
 			{
-				const int32 PointIndex = PerimiterVertices[PerimiterIndex];
+				const int32 PointIndex = PerimeterVertices[PerimeterIndex];
 
 			 	float LeftDistSqr = 0.f;
 			 	float RightDistSqr = 0.f;
 			 		
-			 	int32 LeftPerimiterIndex = PerimiterIndex;
-			 	int32 RightPerimiterIndex = PerimiterIndex;
+			 	int32 LeftPerimeterIndex = PerimeterIndex;
+			 	int32 RightPerimeterIndex = PerimeterIndex;
 			 		
 			 	while (LeftDistSqr < MaxTipAngleSearchDistSqr || RightDistSqr < MaxTipAngleSearchDistSqr)
 			 	{
 			 		if (LeftDistSqr < MaxTipAngleSearchDistSqr)
 			 		{
-						LeftPerimiterIndex = FMath::WrapExclusive(LeftPerimiterIndex - 1, 0, PerimiterVertices.Num());
+						LeftPerimeterIndex = FMath::WrapExclusive(LeftPerimeterIndex - 1, 0, PerimeterVertices.Num());
 			 				
-			 			if (LeftPerimiterIndex == PerimiterIndex)
+			 			if (LeftPerimeterIndex == PerimeterIndex)
 			 			{
 			 				break;
 			 			}
 			 			
-						const int32 LeftVertexIndex = PerimiterVertices[LeftPerimiterIndex];
+						const int32 LeftVertexIndex = PerimeterVertices[LeftPerimeterIndex];
 
 			 			if (EndPointAttribute.GetValue(LeftVertexIndex))
 			 			{
@@ -659,14 +734,14 @@ namespace Internal
 			 			
 			 		if (RightDistSqr < MaxTipAngleSearchDistSqr)
 			 		{
-						RightPerimiterIndex = FMath::WrapExclusive(RightPerimiterIndex + 1, 0, PerimiterVertices.Num());
+						RightPerimeterIndex = FMath::WrapExclusive(RightPerimeterIndex + 1, 0, PerimeterVertices.Num());
 			 			
-			 			if (RightPerimiterIndex == PerimiterIndex)
+			 			if (RightPerimeterIndex == PerimeterIndex)
 			 			{
 			 				break;
 			 			}
 
-						const int32 RightPointIndex = PerimiterVertices[RightPerimiterIndex];
+						const int32 RightPointIndex = PerimeterVertices[RightPerimeterIndex];
 			 				
 			 			if (EndPointAttribute.GetValue(RightPointIndex))
 			 			{
@@ -678,11 +753,11 @@ namespace Internal
 			 			}
 			 		}
 			 			
-			 		const float TipAngle = ComputeAngleDeg(DynamicMesh, PerimiterVertices, PerimiterIndex, LeftPerimiterIndex, RightPerimiterIndex);
+			 		const float TipAngle = ComputeAngleDeg(DynamicMesh, PerimeterVertices, PerimeterIndex, LeftPerimeterIndex, RightPerimeterIndex);
 			 		if (TipAngle >= TipAngleThresholdDeg)
 			 		{
 						EndPointAttribute.SetValue(PointIndex, true);
-						PerimiterAngles[PerimiterIndex] = TipAngle;
+						PerimeterAngles[PerimeterIndex] = TipAngle;
 			 			break;
 			 		}
 			 	}
@@ -690,20 +765,20 @@ namespace Internal
 			 	
 			// When we have 2 neighbouring end points (e.g. when the tip is a plane with 2 45* corners that can both be considered end points),
 			// then select the one with the steepest angle
-			for (int32 PerimiterIndex = 0; PerimiterIndex < PerimiterVertices.Num(); ++PerimiterIndex)
+			for (int32 PerimeterIndex = 0; PerimeterIndex < PerimeterVertices.Num(); ++PerimeterIndex)
 			{
-				const int32 PointIndex = PerimiterVertices[PerimiterIndex];
+				const int32 PointIndex = PerimeterVertices[PerimeterIndex];
 				if (!EndPointAttribute.GetValue(PointIndex))
 				{
 					continue;
 				}
 
-			 	const int32 NextPerimiterIndex = FMath::WrapExclusive(PerimiterIndex + 1, 0, PerimiterVertices.Num());
-				const int32 NextPointIndex = PerimiterVertices[NextPerimiterIndex];
+			 	const int32 NextPerimeterIndex = FMath::WrapExclusive(PerimeterIndex + 1, 0, PerimeterVertices.Num());
+				const int32 NextPointIndex = PerimeterVertices[NextPerimeterIndex];
 
 			 	if (EndPointAttribute.GetValue(NextPointIndex))
 			 	{
-			 		const int32 EndPointToRemove = PerimiterAngles[NextPerimiterIndex] >= PerimiterAngles[PerimiterIndex] ? PointIndex : NextPointIndex;
+			 		const int32 EndPointToRemove = PerimeterAngles[NextPerimeterIndex] >= PerimeterAngles[PerimeterIndex] ? PointIndex : NextPointIndex;
 					EndPointAttribute.SetValue(EndPointToRemove, false);
 			 	}
 			}
@@ -846,9 +921,14 @@ PV::Texture2DImport::EImportResult PV::Texture2DImport::ImportGrowthDataFromText
 		return EImportResult::InvalidSourceTexture;
 	}
 
+	if (InParams.Scale <= 0)
+	{
+		return EImportResult::InvalidScale;
+	}
+
 	if (InParams.DebugState == EPVImportTexture2DDebugState::Texture)
 	{
-		AddDebugTextureQuadToCollection(OutCollection, *InParams.Texture2DAsset, TextureVisualizationMaterial, InParams.Scale);
+		AddTextureQuadToCollection(OutCollection, *InParams.Texture2DAsset, TextureVisualizationMaterial, InParams.Scale);
 		return EImportResult::Success;
 	}
 
@@ -869,55 +949,55 @@ PV::Texture2DImport::EImportResult PV::Texture2DImport::ImportGrowthDataFromText
 	if (InParams.DebugState == EPVImportTexture2DDebugState::SampleImage)
 	{
 		AddDebugPixelValuesToCollection(OutCollection, BlackAndWhiteImageData, SampleSize, InParams.Scale);
-		AddDebugTextureQuadToCollection(OutCollection, *InParams.Texture2DAsset, TextureVisualizationMaterial, InParams.Scale);
+		AddTextureQuadToCollection(OutCollection, *InParams.Texture2DAsset, TextureVisualizationMaterial, InParams.Scale);
 		return EImportResult::Success;
 	}
 
-	TArray<TArray<FVector2f>> Curves;
-	TraceImage(BlackAndWhiteImageData, SampleSize, Curves);
+	TArray<TArray<FVector2f>> PerimeterCurves;
+	TraceImage(BlackAndWhiteImageData, SampleSize, PerimeterCurves);
 	if (InParams.DebugState == EPVImportTexture2DDebugState::TraceImage)
 	{
-		AddDebugPerimiterCurvesToCollection(OutCollection, Curves, InParams.Scale);
-		AddDebugTextureQuadToCollection(OutCollection, *InParams.Texture2DAsset, TextureVisualizationMaterial, InParams.Scale);
+		AddPerimeterCurvesToCollection(OutCollection, PerimeterCurves, InParams.Scale);
+		AddTextureQuadToCollection(OutCollection, *InParams.Texture2DAsset, TextureVisualizationMaterial, InParams.Scale);
 		return EImportResult::Success;
 	}
 	
-	RemoveSmallCurves(Curves, InParams.MinBoundsArea);
+	RemoveSmallCurves(PerimeterCurves, InParams.MinBoundsArea);
 	if (InParams.DebugState == EPVImportTexture2DDebugState::RemoveSmallCurves)
 	{
-		AddDebugPerimiterCurvesToCollection(OutCollection, Curves, InParams.Scale);
-		AddDebugTextureQuadToCollection(OutCollection, *InParams.Texture2DAsset, TextureVisualizationMaterial, InParams.Scale);
+		AddPerimeterCurvesToCollection(OutCollection, PerimeterCurves, InParams.Scale);
+		AddTextureQuadToCollection(OutCollection, *InParams.Texture2DAsset, TextureVisualizationMaterial, InParams.Scale);
 		return EImportResult::Success;
 	}
 	
-	SmoothCurves(Curves, InParams.SmoothingIterations);
+	SmoothCurves(PerimeterCurves, InParams.SmoothingIterations);
 	if (InParams.DebugState == EPVImportTexture2DDebugState::SmoothCurves)
 	{
-		AddDebugPerimiterCurvesToCollection(OutCollection, Curves, InParams.Scale);
-		AddDebugTextureQuadToCollection(OutCollection, *InParams.Texture2DAsset, TextureVisualizationMaterial, InParams.Scale);
+		AddPerimeterCurvesToCollection(OutCollection, PerimeterCurves, InParams.Scale);
+		AddTextureQuadToCollection(OutCollection, *InParams.Texture2DAsset, TextureVisualizationMaterial, InParams.Scale);
 		return EImportResult::Success;
 	}
 	
-	SimplifyCurves(Curves, InParams.SimplificationEpsilon);
+	SimplifyCurves(PerimeterCurves, InParams.SimplificationEpsilon);
 	if (InParams.DebugState == EPVImportTexture2DDebugState::SimplifyCurves)
 	{
-		AddDebugPerimiterCurvesToCollection(OutCollection, Curves, InParams.Scale);
-		AddDebugTextureQuadToCollection(OutCollection, *InParams.Texture2DAsset, TextureVisualizationMaterial, InParams.Scale);
+		AddPerimeterCurvesToCollection(OutCollection, PerimeterCurves, InParams.Scale);
+		AddTextureQuadToCollection(OutCollection, *InParams.Texture2DAsset, TextureVisualizationMaterial, InParams.Scale);
 		return EImportResult::Success;
 	}
 	
 	UE::Geometry::FDynamicMesh3 DynamicMesh;
-	TriangulateCurvesAndSplitEdges(Curves, DynamicMesh, InParams.Scale);
+	TriangulateCurvesAndSplitEdges(PerimeterCurves, DynamicMesh, InParams.Scale);
 	if (InParams.DebugState == EPVImportTexture2DDebugState::TriangulateCurves)
 	{
 		AddDynamicMeshToCollection(DynamicMesh, OutCollection);
-		PerimiterAttributeDefinition.CopyAttributeToCollection(DynamicMesh, OutCollection);
+		PerimeterAttributeDefinition.CopyAttributeToCollection(DynamicMesh, OutCollection);
 		return EImportResult::Success;
 	}
 
 	ComputeVertexScale2D(
 		DynamicMesh, 
-		PerimiterAttributeDefinition.GetOrAttachAttribute(DynamicMesh),
+		PerimeterAttributeDefinition.GetOrAttachAttribute(DynamicMesh),
 		VertexScaleAttributeDefinition->GetOrAttachAttribute(DynamicMesh)
 	);
 	if (InParams.DebugState == EPVImportTexture2DDebugState::ComputeVertexScale)
@@ -932,7 +1012,7 @@ PV::Texture2DImport::EImportResult PV::Texture2DImport::ImportGrowthDataFromText
 		InParams.MaxTipAngleSearchDist * InParams.Scale,
 		InParams.TipAngleThresholdInDegrees,
 		EndPointAttributeDefinition->GetOrAttachAttribute(DynamicMesh),
-		PerimiterAttributeDefinition.GetOrAttachAttribute(DynamicMesh)
+		PerimeterAttributeDefinition.GetOrAttachAttribute(DynamicMesh)
 	);
 	if (InParams.DebugState == EPVImportTexture2DDebugState::FindTips)
 	{
@@ -941,6 +1021,18 @@ PV::Texture2DImport::EImportResult PV::Texture2DImport::ImportGrowthDataFromText
  		return EImportResult::Success;
 	}
 
+	const float HypotenuseLength = FMath::Sqrt(InParams.Scale * InParams.Scale * 2);
+	const FVector VisualizationOffset = FVector(0, -InParams.Scale / 2, HypotenuseLength);
+	// Add end points visualization before applying root points as that transforms the DynamicMesh
+	AddDynamicMeshVerticesToCollection(
+		Output.VisualizationCollection, 
+		TEXT("EndPoints"), 
+		TEXT("Vertex"), 
+		DynamicMesh, 
+		EndPointAttributeDefinition->GetOrAttachAttribute(DynamicMesh), 
+		VisualizationOffset
+	);
+
 	ApplyRootPoints2D(
 		DynamicMesh, 
 		InParams.Scale,
@@ -950,15 +1042,15 @@ PV::Texture2DImport::EImportResult PV::Texture2DImport::ImportGrowthDataFromText
 	);
 	if (InParams.DebugState == EPVImportTexture2DDebugState::ApplyRootPoints)
 	{
-		AddDebugTextureQuadToCollection(OutCollection, *InParams.Texture2DAsset, TextureVisualizationMaterial, InParams.Scale);
+		AddTextureQuadToCollection(OutCollection, *InParams.Texture2DAsset, TextureVisualizationMaterial, InParams.Scale);
 		AddDynamicMeshToCollection(DynamicMesh, OutCollection);
-		AddDebugRootPointsToCollection(OutCollection, InParams.RootPoints, InParams.Scale);
+		AddRootPointsToCollection(OutCollection, InParams.RootPoints, InParams.Scale);
 		EndPointAttributeDefinition->CopyAttributeToCollection(DynamicMesh, OutCollection);
 		RootPointAttributeDefinition->CopyAttributeToCollection(DynamicMesh, OutCollection);
  		return EImportResult::Success;
 	}
 
-	auto& PerimiterAttribute = PerimiterAttributeDefinition.GetOrAttachAttribute(DynamicMesh);
+	auto& PerimeterAttribute = PerimeterAttributeDefinition.GetOrAttachAttribute(DynamicMesh);
 	TArray<FPVBranchHierarchyDescription> BranchHierarchyDescriptions;
 	ComputeBranchHierarchies(
 		DynamicMesh,
@@ -967,7 +1059,7 @@ PV::Texture2DImport::EImportResult PV::Texture2DImport::ImportGrowthDataFromText
 		VertexScaleAttributeDefinition->GetOrAttachAttribute(DynamicMesh),
 		LengthFromRootAttributeDefinition->GetOrAttachAttribute(DynamicMesh),
 		NextBranchPointAttributeDefinition->GetOrAttachAttribute(DynamicMesh),
-		[&](int32 VertexID) { return !PerimiterAttribute.GetValue(VertexID); },
+		[&](int32 VertexID) { return !PerimeterAttribute.GetValue(VertexID); },
 		BranchHierarchyDescriptions
 	);
 	if (InParams.DebugState == EPVImportTexture2DDebugState::ComputeBranchHierarchies)
@@ -991,5 +1083,11 @@ PV::Texture2DImport::EImportResult PV::Texture2DImport::ImportGrowthDataFromText
 		return EImportResult::FailedToGenerateGrowthData;
 	}
 
+	OutCollection.CopyTo(&Output.VisualizationCollection);
+	ArrangePlantsInGrid(Output.VisualizationCollection, VisualizationOffset);
+	AddTextureQuadToCollection(Output.VisualizationCollection, *InParams.Texture2DAsset, TextureVisualizationMaterial, InParams.Scale, VisualizationOffset);
+	AddPerimeterCurvesToCollection(Output.VisualizationCollection, PerimeterCurves, InParams.Scale, VisualizationOffset);
+	AddRootPointsToCollection(Output.VisualizationCollection, InParams.RootPoints, InParams.Scale, VisualizationOffset);
+
 	return EImportResult::Success;
 }
\ No newline at end of file
diff --git a/Engine/Plugins/Experimental/ProceduralVegetationEditor/Source/ProceduralVegetationEditor/Private/Importer/PVImporter_Texture2D.h b/Engine/Plugins/Experimental/ProceduralVegetationEditor/Source/ProceduralVegetationEditor/Private/Importer/PVImporter_Texture2D.h
--- a/Engine/Plugins/Experimental/ProceduralVegetationEditor/Source/ProceduralVegetationEditor/Private/Importer/PVImporter_Texture2D.h
+++ b/Engine/Plugins/Experimental/ProceduralVegetationEditor/Source/ProceduralVegetationEditor/Private/Importer/PVImporter_Texture2D.h
@@ -12,6 +12,7 @@ namespace PV::Texture2DImport
 	enum class EImportResult
 	{
 		Success,
+		InvalidScale,
 		InvalidSourceTexture,
 		InvalidSampleSize,
 		InvalidTextureSize,
diff --git a/Engine/Plugins/Experimental/ProceduralVegetationEditor/Source/ProceduralVegetationEditor/Private/Visualizations/DebugVisualization/PVDebugVisualizer.cpp b/Engine/Plugins/Experimental/ProceduralVegetationEditor/Source/ProceduralVegetationEditor/Private/Visualizations/DebugVisualization/PVDebugVisualizer.cpp
--- a/Engine/Plugins/Experimental/ProceduralVegetationEditor/Source/ProceduralVegetationEditor/Private/Visualizations/DebugVisualization/PVDebugVisualizer.cpp
+++ b/Engine/Plugins/Experimental/ProceduralVegetationEditor/Source/ProceduralVegetationEditor/Private/Visualizations/DebugVisualization/PVDebugVisualizer.cpp
@@ -25,14 +25,15 @@ FVisualizerDrawContext::FVisualizerDrawContext(
 	, SceneSetupParams(InSceneParams)
 {}
 
-void FPVDebugVisualizer::Draw(const UPVData* InData, FPCGSceneSetupParams& InOutParams)
+void FPVDebugVisualizer::DrawDebugVisualizations(
+	const TArray<FPVVisualizationSettings>& VisualizationSettings, 
+	const FManagedArrayCollection& Collection, 
+	FPCGSceneSetupParams& InOutParams
+)
 {
-	check(InData);
-
-	const TArray<FPVVisualizationSettings>& VisualizationSettings = InData->GetDebugSettings().VisualizationSettings;
 	for (const FPVVisualizationSettings& VisualizationSetting : VisualizationSettings)
 	{
-		FVisualizerDrawContext Context(InData->GetCollection(), VisualizationSetting, InOutParams);
+		FVisualizerDrawContext Context(Collection, VisualizationSetting, InOutParams);
 		if (const FPVDebugVisualizationPtr VisualizationPtr = CreateVisualizer(VisualizationSetting))
 		{
 			VisualizationPtr->Draw(Context);
@@ -40,12 +41,28 @@ void FPVDebugVisualizer::Draw(const UPVData* InData, FPCGSceneSetupParams& InOut
 	}
 }
 
-void FPVDebugVisualizer::DrawLoopDebug(const UPVData* InData, FPCGSceneSetupParams& InOutParams)
+//void FPVDebugVisualizer::DrawVisualizationCollections(const TArray<FPVVisualizationCollection>& VisualizationCollections, FPCGSceneSetupParams& InOutParams)
+//{
+//	for (const FPVVisualizationCollection& VisualizationCollection : VisualizationCollections)
+//	{
+//		for (const FPVVisualizationSettings& VisualizationSetting : VisualizationCollection.VisualizationSettings)
+//		{
+//			FVisualizerDrawContext Context(VisualizationCollection.Collection, VisualizationSetting, InOutParams);
+//			if (const FPVDebugVisualizationPtr VisualizationPtr = CreateVisualizer(VisualizationSetting))
+//			{
+//				VisualizationPtr->Draw(Context);
+//			}
+//		}
+//	}
+//}
+
+void FPVDebugVisualizer::DrawLoopDebug(
+	const TArray<FPVLoopDebugSettings>& LoopDebugVisualizationSettings, 
+	bool bAutoFocusLoopDebug, 
+	FPCGSceneSetupParams& InOutParams
+)
 {
-	check(InData);
-
 	TArray<TPair<FColor, FText>> LegendItems;
-	const TArray<FPVLoopDebugSettings>& LoopDebugVisualizationSettings = InData->GetDebugSettings().LoopDebugVisualizationSettings;
 
 	float Alpha = 0.0f;
 	const float AlphaStep = 1.0f / static_cast<float>(LoopDebugVisualizationSettings.Num());
@@ -73,7 +90,7 @@ void FPVDebugVisualizer::DrawLoopDebug(const UPVData* InData, FPCGSceneSetupPara
 		return;
 	}
 
-	if (InData->GetDebugSettings().bAutoFocusLoopDebug)
+	if (bAutoFocusLoopDebug)
 	{
 		if (const TSharedPtr<FEditorViewportClient> ViewportClient = ViewportWidget->GetViewportClient())
 		{
diff --git a/Engine/Plugins/Experimental/ProceduralVegetationEditor/Source/ProceduralVegetationEditor/Private/Visualizations/DebugVisualization/PVDebugVisualizer.h b/Engine/Plugins/Experimental/ProceduralVegetationEditor/Source/ProceduralVegetationEditor/Private/Visualizations/DebugVisualization/PVDebugVisualizer.h
--- a/Engine/Plugins/Experimental/ProceduralVegetationEditor/Source/ProceduralVegetationEditor/Private/Visualizations/DebugVisualization/PVDebugVisualizer.h
+++ b/Engine/Plugins/Experimental/ProceduralVegetationEditor/Source/ProceduralVegetationEditor/Private/Visualizations/DebugVisualization/PVDebugVisualizer.h
@@ -24,8 +24,8 @@ struct FVisualizerDrawContext
 class FPVDebugVisualizer
 {
 public:
-	static void Draw(const UPVData* InData, FPCGSceneSetupParams& InOutParams);
-	static void DrawLoopDebug(const UPVData* InData, FPCGSceneSetupParams& InOutParams);
+	static void DrawDebugVisualizations(const TArray<FPVVisualizationSettings>& VisualizationSettings, const FManagedArrayCollection& Collection, FPCGSceneSetupParams& InOutParams);
+	static void DrawLoopDebug(const TArray<FPVLoopDebugSettings>& LoopDebugVisualizationSettings, bool bAutoFocusLoopDebug, FPCGSceneSetupParams& InOutParams);
 
 private:
 	static void VisualizeLoopDebuggerParam(
diff --git a/Engine/Plugins/Experimental/ProceduralVegetationEditor/Source/ProceduralVegetationEditor/Private/Visualizations/PVDataVisualization.cpp b/Engine/Plugins/Experimental/ProceduralVegetationEditor/Source/ProceduralVegetationEditor/Private/Visualizations/PVDataVisualization.cpp
--- a/Engine/Plugins/Experimental/ProceduralVegetationEditor/Source/ProceduralVegetationEditor/Private/Visualizations/PVDataVisualization.cpp
+++ b/Engine/Plugins/Experimental/ProceduralVegetationEditor/Source/ProceduralVegetationEditor/Private/Visualizations/PVDataVisualization.cpp
@@ -74,7 +74,7 @@ FPCGSetupSceneFunc FPVDataVisualization::GetViewportSetupFunc(const UPCGSettings
 		FPCGSceneSetupParams& InOutParams)
 	{
 		check(InOutParams.Scene)
-		if (!WeakData.IsValid())
+		if (!WeakData.IsValid() || !WeakSettings.IsValid())
 		{
 			UE_LOG(LogProceduralVegetationDataVisualization, Error, TEXT("Failed to setup data viewport, the data was lost or invalid."));
 			return;
@@ -111,35 +111,36 @@ FPCGSetupSceneFunc FPVDataVisualization::GetViewportSetupFunc(const UPCGSettings
 			{
 				InOutParams.Scene->GetLineBatcher()->Flush();
 				RenderMap[RenderType].Execute(Collection, InOutParams, Bounds);
+
+				for (const FPVVisualizationCollection& VisualizationCollection : WeakData->GetVisualizationCollections())
+				{
+					RenderMap[RenderType].Execute(VisualizationCollection.Collection, InOutParams, Bounds);
+				}
 			}
 		}
 
-		if (const UPVBaseSettings* PVBaseSettings = Cast<UPVBaseSettings>(WeakSettings))
+		for (const auto& ViewportObject : WeakSettings->GetViewportObjects())
 		{
-			for (const auto& ViewportObject : PVBaseSettings->GetViewportObjects())
+			if (UStaticMesh* Mesh = Cast<UStaticMesh>(ViewportObject.Key))
 			{
-				if (UStaticMesh* Mesh = Cast<UStaticMesh>(ViewportObject.Key))
-				{
-					TObjectPtr<UStaticMeshComponent> StaticMeshComponent = NewObject<UStaticMeshComponent>(GetTransientPackage(), NAME_None, RF_Transient);
-					StaticMeshComponent->SetStaticMesh(Mesh);
-					StaticMeshComponent->UpdateBounds();
+				TObjectPtr<UStaticMeshComponent> StaticMeshComponent = NewObject<UStaticMeshComponent>(GetTransientPackage(), NAME_None, RF_Transient);
+				StaticMeshComponent->SetStaticMesh(Mesh);
+				StaticMeshComponent->UpdateBounds();
 
-					InOutParams.ManagedResources.Add(StaticMeshComponent);
-					InOutParams.Scene->AddComponent(StaticMeshComponent, ViewportObject.Value);
-				}
+				InOutParams.ManagedResources.Add(StaticMeshComponent);
+				InOutParams.Scene->AddComponent(StaticMeshComponent, ViewportObject.Value);
 			}
 		}
- 
-		if (PV::Utilities::DebugModeEnabled())
+
+		if (PV::Utilities::DebugModeEnabled() && WeakSettings->bDebug)
 		{
-			if (WeakSettings.IsValid())
-			{
-				if (WeakSettings->bDebug)
-				{
-					FPVDebugVisualizer::Draw(WeakData.Get(), InOutParams);
-					FPVDebugVisualizer::DrawLoopDebug(WeakData.Get(), InOutParams);
-				}
-			}
+			FPVDebugVisualizer::DrawDebugVisualizations(WeakData->GetDebugSettings().VisualizationSettings, WeakData->GetCollection(), InOutParams);
+			FPVDebugVisualizer::DrawLoopDebug(WeakData->GetDebugSettings().LoopDebugVisualizationSettings, WeakData->GetDebugSettings().bAutoFocusLoopDebug, InOutParams);
+		}
+
+		for (const FPVVisualizationCollection& VisualizationCollection : WeakData->GetVisualizationCollections())
+		{
+			FPVDebugVisualizer::DrawDebugVisualizations(VisualizationCollection.VisualizationSettings, VisualizationCollection.Collection, InOutParams);
 		}
  
 		InOutParams.FocusBounds = FBoxSphereBounds(Bounds.Origin, Bounds.BoxExtent, Bounds.SphereRadius);
