[PVE] Resume Growth Functionality first implementation

#Synced-CL 50040261

diff --git a/Engine/Plugins/Experimental/ProceduralVegetationEditor/Source/ProceduralVegetation/Private/Facades/PVBranchFacade.cpp b/Engine/Plugins/Experimental/ProceduralVegetationEditor/Source/ProceduralVegetation/Private/Facades/PVBranchFacade.cpp
--- a/Engine/Plugins/Experimental/ProceduralVegetationEditor/Source/ProceduralVegetation/Private/Facades/PVBranchFacade.cpp
+++ b/Engine/Plugins/Experimental/ProceduralVegetationEditor/Source/ProceduralVegetation/Private/Facades/PVBranchFacade.cpp
@@ -499,4 +499,44 @@ namespace PV::Facades
 	{
 		FILL_COLLECTION_ATTRIBUTE(PlantNumbers, InputArray);
 	}
+
+	void FBranchFacade::GetParents(int32 Index, TArray<int32>& OutValue) const
+	{
+		GetValueFromAccessor<TArray<int32>>(Parents, OutValue, Index, TArray<int32>());
+	}
+
+	void FBranchFacade::GetChildren(int32 Index, TArray<int32>& OutValue) const
+	{
+		GetValueFromAccessor<TArray<int32>>(Children, OutValue, Index, TArray<int32>());
+	}
+
+	void FBranchFacade::GetBranchPoints(int32 Index, TArray<int32>& OutValue) const
+	{
+		GetValueFromAccessor<TArray<int32>>(BranchPoints, OutValue, Index, TArray<int32>());
+	}
+
+	void FBranchFacade::GetBranchNumber(int32 Index, int32& OutValue) const
+	{
+		GetValueFromAccessor<int32>(BranchNumbers, OutValue, Index, -1);
+	}
+
+	void FBranchFacade::GetBranchSourceBudNumber(int32 Index, int32& OutValue) const
+	{
+		GetValueFromAccessor<int32>(BranchSourceBudNumber, OutValue, Index, -1);
+	}
+
+	void FBranchFacade::GetBranchHierarchyNumber(int32 Index, int32& OutValue) const
+	{
+		GetValueFromAccessor<int32>(BranchHierarchyNumber, OutValue, Index, -1);
+	}
+
+	void FBranchFacade::GetBranchParentNumber(int32 Index, int32& OutValue) const
+	{
+		GetValueFromAccessor<int32>(BranchParentNumbers, OutValue, Index, -1);
+	}
+
+	void FBranchFacade::GetPlantNumber(int32 Index, int32& OutValue) const
+	{
+		GetValueFromAccessor<int32>(PlantNumbers, OutValue, Index, -1);
+	}
 }
diff --git a/Engine/Plugins/Experimental/ProceduralVegetationEditor/Source/ProceduralVegetation/Private/Facades/PVPointFacade.cpp b/Engine/Plugins/Experimental/ProceduralVegetationEditor/Source/ProceduralVegetation/Private/Facades/PVPointFacade.cpp
--- a/Engine/Plugins/Experimental/ProceduralVegetationEditor/Source/ProceduralVegetation/Private/Facades/PVPointFacade.cpp
+++ b/Engine/Plugins/Experimental/ProceduralVegetationEditor/Source/ProceduralVegetation/Private/Facades/PVPointFacade.cpp
@@ -15,6 +15,8 @@ namespace PV::Facades
 		, GroundGradient(InCollection, AttributeNames::GroundGradient, GroupNames::PointGroup)
 		, PointScale(InCollection, AttributeNames::PointScale, GroupNames::PointGroup)
 		, BudDirections(InCollection, AttributeNames::BudDirection, GroupNames::PointGroup)
+		, BudStatus(InCollection, AttributeNames::BudStatus, GroupNames::PointGroup)
+		, BudLateralMeristem(InCollection, AttributeNames::BudLateralMeristem, GroupNames::PointGroup)
 		, BranchGradients(InCollection, AttributeNames::BranchGradient, GroupNames::PointGroup)
 		, BudHormoneLevels(InCollection, AttributeNames::BudHormoneLevels, GroupNames::PointGroup)
 		, PlantGradients(InCollection, AttributeNames::PlantGradient, GroupNames::PointGroup)
@@ -38,6 +40,8 @@ namespace PV::Facades
 		, GroundGradient(InCollection, AttributeNames::GroundGradient, GroupNames::PointGroup)
 		, PointScale(InCollection, AttributeNames::PointScale, GroupNames::PointGroup)
 		, BudDirections(InCollection, AttributeNames::BudDirection, GroupNames::PointGroup)
+		, BudStatus(InCollection, AttributeNames::BudStatus, GroupNames::PointGroup)
+		, BudLateralMeristem(InCollection, AttributeNames::BudLateralMeristem, GroupNames::PointGroup)
 		, BranchGradients(InCollection, AttributeNames::BranchGradient, GroupNames::PointGroup)
 		, BudHormoneLevels(InCollection, AttributeNames::BudHormoneLevels, GroupNames::PointGroup)
 		, PlantGradients(InCollection, AttributeNames::PlantGradient, GroupNames::PointGroup)
@@ -478,18 +482,27 @@ namespace PV::Facades
 			BudLightDetected.ModifyAt(ToIndex, BudLightDetected[FromIndex]);
 			BudDevelopment.ModifyAt(ToIndex, BudDevelopment[FromIndex]);
 			
-			if(TextureCoordV.IsValid())
+			if(TextureCoordV.IsValid() && TextureCoordV.IsValidIndex(ToIndex) && TextureCoordV.IsValidIndex(FromIndex))
 			{
 				TextureCoordV.ModifyAt(ToIndex, TextureCoordV[FromIndex]);
 			}
-			if(TextureCoordUOffset.IsValid())
+			if(TextureCoordUOffset.IsValid() && TextureCoordUOffset.IsValidIndex(ToIndex) && TextureCoordUOffset.IsValidIndex(FromIndex))
 			{
 				TextureCoordUOffset.ModifyAt(ToIndex, TextureCoordUOffset[FromIndex]);
 			}
-			if(URange.IsValid())
+			if(URange.IsValid() && URange.IsValidIndex(ToIndex) && URange.IsValidIndex(FromIndex))
 			{
 				URange.ModifyAt(ToIndex, URange[FromIndex]);
 			}
+			if (BudStatus.IsValid() && BudStatus.IsValidIndex(ToIndex) && BudStatus.IsValidIndex(FromIndex))
+			{
+				BudStatus.ModifyAt(ToIndex, BudStatus[FromIndex]);
+			}
+
+			if (BudLateralMeristem.IsValid() && BudLateralMeristem.IsValidIndex(ToIndex) && BudLateralMeristem.IsValidIndex(FromIndex))
+			{
+				BudLateralMeristem.ModifyAt(ToIndex, BudLateralMeristem[FromIndex]);
+			}
 		}
 	}
 
@@ -590,4 +603,94 @@ namespace PV::Facades
 	{
 		FILL_COLLECTION_ATTRIBUTE(NjordPixelIndex, InputArray);
 	}
+
+	void FPointFacade::FillBudStatus(const TArray<TArray<int>>& InputArray)
+	{
+		FILL_COLLECTION_ATTRIBUTE(BudStatus, InputArray);
+	}
+
+	void FPointFacade::FillBudLateralMeristem(const TArray<TArray<float>>& InputArray)
+	{
+		FILL_COLLECTION_ATTRIBUTE(BudLateralMeristem, InputArray);
+	}
+
+	void FPointFacade::GetLFR(int32 Index, float& OutValue) const
+	{
+		GetValueFromAccessor<float>(LengthFromRoot, OutValue, Index, -1);
+	}
+
+	void FPointFacade::GetPointScaleGradient(int32 Index, float& OutValue) const
+	{
+		GetValueFromAccessor<float>(PointScaleGradient, OutValue, Index, -1);
+	}
+
+	void FPointFacade::GetHullGradient(int32 Index, float& OutValue) const
+	{
+		GetValueFromAccessor<float>(HullGradient, OutValue, Index, -1);
+	}
+
+	void FPointFacade::GetMainTrunkGradients(int32 Index, float& OutValue) const
+	{
+		GetValueFromAccessor<float>(MainTrunkGradient, OutValue, Index, -1);
+	}
+
+	void FPointFacade::GetGroundGradients(int32 Index, float& OutValue) const
+	{
+		GetValueFromAccessor<float>(GroundGradient, OutValue, Index, -1);
+	}
+
+	void FPointFacade::GetPointScales(int32 Index, float& OutValue) const
+	{
+		GetValueFromAccessor<float>(PointScale, OutValue, Index , -1);
+	}
+
+	void FPointFacade::GetBudDirections(int32 Index, TArray<FVector3f>& OutValue) const
+	{
+		GetValueFromAccessor<TArray<FVector3f>>(BudDirections, OutValue, Index, TArray<FVector3f>());
+	}
+
+	void FPointFacade::GetBranchGradients(int32 Index, float& OutValue) const
+	{
+		GetValueFromAccessor<float>(BranchGradients, OutValue, Index, -1);
+	}
+
+	void FPointFacade::GetBudHormoneLevels(int32 Index, TArray<float>& OutValue) const
+	{
+		GetValueFromAccessor<TArray<float>>(BudHormoneLevels, OutValue, Index, TArray<float>());
+	}
+
+	void FPointFacade::GetPlantGradients(int32 Index, float& OutValue) const
+	{
+		GetValueFromAccessor<float>(PlantGradients, OutValue, Index, -1);
+	}
+
+	void FPointFacade::GetBudNumbers(int32 Index, int32& OutValue) const
+	{
+		GetValueFromAccessor<int32>(BudNumber, OutValue, Index, -1);
+	}
+
+	void FPointFacade::GetBudLightDetected(int32 Index, TArray<float>& OutValue) const
+	{
+		GetValueFromAccessor<TArray<float>>(BudLightDetected, OutValue, Index, TArray<float>());
+	}
+
+	void FPointFacade::GetBudDevelopment(int32 Index, TArray<int>& OutValue) const
+	{
+		GetValueFromAccessor<TArray<int>>(BudDevelopment, OutValue, Index, TArray<int>());
+	}
+
+	void FPointFacade::GetNjordPixelID(int32 Index, float& OutValue) const
+	{
+		GetValueFromAccessor<float>(NjordPixelIndex, OutValue, Index , -1);
+	}
+
+	void FPointFacade::GetBudStatus(int32 Index, TArray<int>& OutValue) const
+	{
+		GetValueFromAccessor<TArray<int>>(BudStatus, OutValue, Index, TArray<int>());
+	}
+
+	void FPointFacade::GetBudLateralMeristem(int32 Index, TArray<float>& OutValue) const
+	{
+		GetValueFromAccessor<TArray<float>>(BudLateralMeristem, OutValue, Index, TArray<float>());
+	}
 }
diff --git a/Engine/Plugins/Experimental/ProceduralVegetationEditor/Source/ProceduralVegetation/Private/Implementations/PVGrower.cpp b/Engine/Plugins/Experimental/ProceduralVegetationEditor/Source/ProceduralVegetation/Private/Implementations/PVGrower.cpp
--- a/Engine/Plugins/Experimental/ProceduralVegetationEditor/Source/ProceduralVegetation/Private/Implementations/PVGrower.cpp
+++ b/Engine/Plugins/Experimental/ProceduralVegetationEditor/Source/ProceduralVegetation/Private/Implementations/PVGrower.cpp
@@ -83,12 +83,24 @@ void FPVGrower::Grow(const FPVSeedPoint& SeedPoint, const FPVGrowerParams& InGro
 	TObjectPtr<UPVGrowerData> Skeleton = NewObject<UPVGrowerData>(GetTransientPackage(), UPVGrowerData::StaticClass(),FName("Skeleton"), RF_Public | RF_Transactional);
 	InitializeGrowth(InGrowerParams, SeedPoint, Skeleton);
 	IterateGrowthCycles(InGrowerParams, Skeleton);
-	ConvertPointsToMeters(Skeleton);
+	ConvertPointsToCentimeters(Skeleton);
 	SplitBranchSourcePoints(Skeleton);
 	FillCollection(InCollection, Skeleton);
 	//PrintPoints();
 }
 
+void FPVGrower::Grow(const FManagedArrayCollection& InCollection, const FPVGrowerParams& InGrowerParams, FManagedArrayCollection& Collection)
+{
+	TRACE_CPUPROFILER_EVENT_SCOPE(PV::Grower::Grow);
+	TObjectPtr<UPVGrowerData> Skeleton = NewObject<UPVGrowerData>(GetTransientPackage(), UPVGrowerData::StaticClass(),FName("Skeleton"), RF_Public | RF_Transactional);
+	InitializeGrowth(InCollection, Skeleton);
+	IterateGrowthCycles(InGrowerParams, Skeleton);
+	ConvertPointsToCentimeters(Skeleton);
+	SplitBranchSourcePoints(Skeleton);
+	FillCollection(Collection, Skeleton);
+	//PrintPoints();
+}
+
 void FPVGrower::FillCollection(FManagedArrayCollection& Collection, UPVGrowerData* Skeleton)
 {
 	PV::Facades::FPointFacade PointFacade = PV::Facades::FPointFacade(Collection);
@@ -109,6 +121,8 @@ void FPVGrower::FillCollection(FManagedArrayCollection& Collection, UPVGrowerDat
 	TArray<int32> BudNumbers;
 	TArray<TArray<float>> BudLightDetected;
 	TArray<TArray<int32>> BudDevelopment;
+	TArray<TArray<int32>> BudStatus;
+	TArray<TArray<float>> BudLateralMeristem;
 	TArray<float> NjordPixelIndices;
 	
 	int PointIndex = 0;
@@ -124,6 +138,8 @@ void FPVGrower::FillCollection(FManagedArrayCollection& Collection, UPVGrowerDat
 		BudHormoneLevels.Add(Point->Bud.HarmoneLevels.GetDataArray());
 		BudDevelopment.Add(Point->Bud.Development.GetDataArray());
 		BudLightDetected.Add(Point->Bud.LightDetected.GetDataArray());
+		BudStatus.Add(Point->Bud.Status.GetDataArray());
+		BudLateralMeristem.Add(Point->Bud.LateralMeristem.GetDataArray());
 		NjordPixelIndices.Add(PointIndex+1);
 		PointIndex++;
 	}
@@ -137,6 +153,8 @@ void FPVGrower::FillCollection(FManagedArrayCollection& Collection, UPVGrowerDat
 	PointFacade.FillGroundGradients(GroundGradients);
 	PointFacade.FillPointScales(PointScales);
 	PointFacade.FillBudDirections(BudDirections);
+	PointFacade.FillBudStatus(BudStatus);
+	PointFacade.FillBudLateralMeristem(BudLateralMeristem);
 	PointFacade.FillBranchGradients(BranchGradients);
 	PointFacade.FillBudHormoneLevels(BudHormoneLevels);
 	PointFacade.FillPlantGradients(PlantGradients);
@@ -205,6 +223,177 @@ void FPVGrower::FillCollection(FManagedArrayCollection& Collection, UPVGrowerDat
 	check(BranchFacade.IsValid());
 }
 
+void FPVGrower::FillFromCollection(const FManagedArrayCollection& Collection, UPVGrowerData* Skeleton)
+{
+	FillPointsFromCollection(Collection, Skeleton);
+	FillBranchesFromCollection(Collection, Skeleton);
+}
+
+void FPVGrower::FillPointsFromCollection(const FManagedArrayCollection& Collection, UPVGrowerData* Skeleton)
+{
+	PV::Facades::FPointFacade PointFacade = PV::Facades::FPointFacade(Collection);
+	const TManagedArray<FVector3f>& Positions = PointFacade.GetPositions();
+
+	Skeleton->Points.Reserve(Positions.Num());
+
+	for (int i=0; i< Positions.Num(); i++)
+	{
+		int32 BudNumber = PointFacade.GetBudNumber(i);
+
+		// To ignore the Fused points
+		if (Skeleton->GetPointIndex(BudNumber) == INDEX_NONE)
+		{
+			UPVGrowerPoint* Point = NewObject<UPVGrowerPoint>(Skeleton,UPVGrowerPoint::StaticClass(),NAME_None, RF_Transactional);
+
+			FPVBud Bud;
+			Bud.BudNumber = BudNumber;
+
+			Skeleton->MaxBudNumber = FMath::Max(BudNumber, Skeleton->MaxBudNumber);
+			
+			TArray<int> DevelopmentArray;
+			PointFacade.GetBudDevelopment(i, DevelopmentArray);
+			Bud.Development.SetData(DevelopmentArray);
+
+			TArray<int> StatusArray;
+			PointFacade.GetBudStatus(i, StatusArray);
+			Bud.Status.SetData(StatusArray);
+		
+			TArray<float> HormoneLevelsArray;
+			PointFacade.GetBudHormoneLevels(i, HormoneLevelsArray);
+			Bud.HarmoneLevels.SetData(HormoneLevelsArray);
+
+			TArray<FVector3f> DirectionArray;
+			PointFacade.GetBudDirections(i, DirectionArray);
+			Bud.Direction.SetData(DirectionArray);
+		
+			TArray<float> LightDetectedArray;
+			PointFacade.GetBudLightDetected(i, LightDetectedArray);
+			Bud.LightDetected.SetData(LightDetectedArray);
+
+			TArray<float> LateralMeristemArray;
+			PointFacade.GetBudLateralMeristem(i, LateralMeristemArray);
+			Bud.LateralMeristem.SetData(LateralMeristemArray);
+
+			Point->Position = FVector(Positions[i]) * 0.01f;
+		
+			float LengthFromRoot;
+			PointFacade.GetLFR(i, LengthFromRoot);
+			Point->LengthFromRoot = LengthFromRoot;
+
+			float PScale;
+			PointFacade.GetPointScales(i, PScale);
+			Point->PScale = PScale * 0.01f;
+
+			PointFacade.GetBudNumbers(i, Bud.BudNumber);
+
+			Point->Bud = Bud;
+
+			Point->SeedInfo.PScale = 0.1;
+			Point->SeedInfo.PScaleRatio = 1;
+
+			Skeleton->AddPoint(Point, Bud.BudNumber);
+		}
+	}
+}
+
+void FPVGrower::FillBranchesFromCollection(const FManagedArrayCollection& Collection, UPVGrowerData* Skeleton)
+{
+	PV::Facades::FBranchFacade BranchFacade = PV::Facades::FBranchFacade(Collection);
+
+	const int32 NumBranches = BranchFacade.GetElementCount();
+	Skeleton->Primitives.Reserve(NumBranches);
+
+	TArray<UPVGrowerPoint*> SplitPoints;
+
+	auto SwapSplitPointWithFusedPoint = [&SplitPoints, Skeleton](TArray<int32>& BranchPoints, int Index, int SourceBudNumber)
+	{
+		if (Index == 0)
+		{
+			UPVGrowerPoint* Point = Skeleton->GetPointFromIndex(BranchPoints[Index]);
+			if (Point && Point->Bud.Status.Seed == 0)
+			{
+				SplitPoints.Add(Point);
+				BranchPoints[Index] = Skeleton->GetPointIndex(SourceBudNumber);
+			}
+		}
+	};
+	
+	for (int i=0; i< NumBranches; i++)
+	{
+		UPVGrowerPrimitive* Primitive = NewObject<UPVGrowerPrimitive>(Skeleton,UPVGrowerPrimitive::StaticClass(),NAME_None, RF_Transactional);
+
+		TArray<int32> BranchPoints;
+		BranchFacade.GetBranchPoints(i, BranchPoints);
+
+		int32 BranchSourceBudNumber;
+		BranchFacade.GetBranchSourceBudNumber(i, BranchSourceBudNumber);
+		Primitive->BranchSourceBudNumber = BranchSourceBudNumber;
+
+		for (int j=0; j< BranchPoints.Num(); j++)
+		{
+			SwapSplitPointWithFusedPoint(BranchPoints, j, BranchSourceBudNumber);
+			
+			UPVGrowerPoint* Point = Skeleton->GetPointFromIndex(BranchPoints[j]);
+			if (Point)
+			{
+				Point->Primitive = Primitive;
+				Primitive->BranchBuds.Add(Point->Bud.BudNumber);
+
+				if (BranchPoints.IsValidIndex(j-1))
+				{
+					UPVGrowerPoint* PrevPoint = Skeleton->GetPointFromIndex(BranchPoints[j - 1]);
+					if (PrevPoint)
+					{
+						Point->Neighbors.AddUnique(PrevPoint);
+					}
+				}
+
+				if (BranchPoints.IsValidIndex(j + 1))
+				{
+					UPVGrowerPoint* NextPoint = Skeleton->GetPointFromIndex(BranchPoints[j + 1]);
+					if (NextPoint)
+					{
+						Point->Neighbors.AddUnique(NextPoint);
+					}
+				}
+			}
+			else
+			{
+				UE_LOG(LogProceduralVegetation, Warning, TEXT("Point Not found for index %i"), BranchPoints[j]);
+			}
+		}
+
+		TArray<int32> BranchParents;
+		BranchFacade.GetParents(i, BranchParents);
+		Primitive->Parents = BranchParents;
+
+		TArray<int32> BranchChildren;
+		BranchFacade.GetChildren(i, BranchChildren);
+		Primitive->Children = BranchChildren;
+
+		int32 BranchNumber = BranchFacade.GetBranchNumber(i);
+		Primitive->BranchNumber = BranchNumber;
+
+		int32 BranchParentNumber;
+		BranchFacade.GetBranchParentNumber(i, BranchParentNumber);
+		Primitive->BranchParentNumber = BranchParentNumber;
+
+		int32 BranchHierarchyNumber;
+		BranchFacade.GetBranchHierarchyNumber(i, BranchHierarchyNumber);
+
+		int32 BranchPlantNumber;
+		BranchFacade.GetPlantNumber(i, BranchPlantNumber);
+		Primitive->PlantNumber = BranchPlantNumber;
+
+		Skeleton->Primitives.Add(Primitive);
+	}
+
+	for (UPVGrowerPoint* Point : SplitPoints)
+	{
+		Skeleton->RemovePoint(Point->Bud.BudNumber);
+	}
+}
+
 float FPVGrower::RandomValue(uint32 Seed)
 {
 	std::default_random_engine RandomSeedEngine(Seed);
@@ -231,6 +420,15 @@ void FPVGrower::InitializeGrowth(const FPVGrowerParams& InGrowerParams, const FP
 	//Calculate MinRootDist from seed points
 }
 
+void FPVGrower::InitializeGrowth(const FManagedArrayCollection& InCollection, UPVGrowerData* Skeleton)
+{
+	TRACE_CPUPROFILER_EVENT_SCOPE(PV::Grower::InitializeGrowth);
+	check(Skeleton);
+
+	FillFromCollection(InCollection, Skeleton);
+	Skeleton->MaxBranchNumber = Skeleton->Primitives.Num();
+}
+
 void FPVGrower::IterateGrowthCycles(const FPVGrowerParams& InGrowerParams, UPVGrowerData* Skeleton)
 {
 	FPVColliderMeshData ColliderMeshData;
@@ -252,7 +450,7 @@ void FPVGrower::PrintPoints( UPVGrowerData* Skeleton)
 	}
 }
 
-void FPVGrower::ConvertPointsToMeters(UPVGrowerData* Skeleton)
+void FPVGrower::ConvertPointsToCentimeters(UPVGrowerData* Skeleton)
 {
 	for (auto Point : Skeleton->Points)
 	{
diff --git a/Engine/Plugins/Experimental/ProceduralVegetationEditor/Source/ProceduralVegetation/Private/Implementations/PVGrower.h b/Engine/Plugins/Experimental/ProceduralVegetationEditor/Source/ProceduralVegetation/Private/Implementations/PVGrower.h
--- a/Engine/Plugins/Experimental/ProceduralVegetationEditor/Source/ProceduralVegetation/Private/Implementations/PVGrower.h
+++ b/Engine/Plugins/Experimental/ProceduralVegetationEditor/Source/ProceduralVegetation/Private/Implementations/PVGrower.h
@@ -104,17 +104,27 @@ struct FPVGrower
 {
 	static void Grow(const FPVSeedPoint& SeedPoint, const FPVGrowerParams& InGrowerParams, FManagedArrayCollection& InCollection);
 	
+	static void Grow(const FManagedArrayCollection& InCollection, const FPVGrowerParams& InGrowerParams, FManagedArrayCollection& Collection);
+
 	static void FillCollection(FManagedArrayCollection& Collection, UPVGrowerData* Skeleton);
+	
+	static void FillFromCollection(const FManagedArrayCollection& Collection, UPVGrowerData* Skeleton);
+	
+	static void FillPointsFromCollection(const FManagedArrayCollection& Collection, UPVGrowerData* Skeleton);
+	
+	static void FillBranchesFromCollection(const FManagedArrayCollection& Collection, UPVGrowerData* Skeleton);
 
 	static float RandomValue(uint32 Seed);
 
 	static void InitializeGrowth(const FPVGrowerParams& InGrowerParams, const FPVSeedPoint& SeedPoint, UPVGrowerData* Skeleton);
+	
+	static void InitializeGrowth(const FManagedArrayCollection& InCollection, UPVGrowerData* Skeleton);
 
 	static void IterateGrowthCycles(const FPVGrowerParams& InGrowerParams, UPVGrowerData* Skeleton);
 
 	static void PrintPoints(UPVGrowerData* Skeleton);
 
-	static void ConvertPointsToMeters(UPVGrowerData* Skeleton);
+	static void ConvertPointsToCentimeters(UPVGrowerData* Skeleton);
 	
 	static void SplitBranchSourcePoints(UPVGrowerData* Skeleton);
 
diff --git a/Engine/Plugins/Experimental/ProceduralVegetationEditor/Source/ProceduralVegetation/Private/Implementations/PVGrowerData.cpp b/Engine/Plugins/Experimental/ProceduralVegetationEditor/Source/ProceduralVegetation/Private/Implementations/PVGrowerData.cpp
--- a/Engine/Plugins/Experimental/ProceduralVegetationEditor/Source/ProceduralVegetation/Private/Implementations/PVGrowerData.cpp
+++ b/Engine/Plugins/Experimental/ProceduralVegetationEditor/Source/ProceduralVegetation/Private/Implementations/PVGrowerData.cpp
@@ -2,6 +2,87 @@
 
 #include "PVGrowerData.h"
 
+void FPVBudDevelopment::SetData(const TArray<int>& InArray)
+{
+	if (ensure(InArray.Num() == 6))
+	{
+		Generation = InArray[0];
+		BudAge = InArray[1];
+		BranchAge = InArray[2];
+		AgeSenescense = InArray[3];
+		LightSenescense = InArray[4];
+		RelativeBudAge = InArray[5];
+	}
+}
+
+void FPVBudDirection::SetData(const TArray<FVector3f>& InArray)
+{
+	if (ensure(InArray.Num() == 6))
+	{
+		Apical = FVector(InArray[0]);
+		Axillary = FVector(InArray[1]);
+		LightOptimal = FVector(InArray[2]);
+		LightSubOptimal = FVector(InArray[3]);
+		CurveGuide = FVector(InArray[4]);
+		UpVector = FVector(InArray[5]);
+	}
+}
+
+void FPVBudHarmoneLevels::SetData(const TArray<float>& InArray)
+{
+	if (ensure(InArray.Num() == 6))
+	{
+		Apical = InArray[0];
+		Axillary = InArray[1];
+		AxillaryInhibition = InArray[2];
+		Radical = InArray[3];
+		Ethylene = InArray[4];
+		Cykininin = InArray[5];
+	}
+}
+
+void FPVBudLateralMeristem::SetData(const TArray<float>& InArray)
+{
+	if (ensure(InArray.Num() == 7))
+	{
+		LateralMeristem = InArray[0];
+		Multiplier = InArray[1];
+		Inactive = InArray[2];
+		Davinci = InArray[3];
+		ParentDot = InArray[4];
+		RootDistance = InArray[5];
+		Degredation = InArray[6];
+	}
+}
+
+void FPVBudLightDetected::SetData(const TArray<float>& InArray)
+{
+	if (ensure(InArray.Num() == 4))
+	{
+		Available = InArray[0];
+		Resource = InArray[1];
+		Branch = InArray[2];
+		Collision = InArray[3];
+	}
+}
+
+void FPVBudStatus::SetData(const TArray<int>& InArray)
+{
+	if (ensure(InArray.Num() == 10))
+	{
+		ApicalMeristem = InArray[0];
+		Codominant = InArray[1];
+		Axillary = InArray[2];
+		Seed = InArray[3];
+		Dormant = InArray[4];
+		Triggered = InArray[5];
+		NumTriggered = InArray[6];
+		Inactive = InArray[7];
+		BrokenTip = InArray[8];
+		Broken = InArray[9];
+	}
+}
+
 void UPVGrowerData::AddPoint(TObjectPtr<UPVGrowerPoint> Point, TObjectPtr<UPVGrowerPrimitive> Primitive, bool bAddToBranch /*= true*/)
 {
 	MaxBudNumber += 1;
@@ -15,6 +96,12 @@ void UPVGrowerData::AddPoint(TObjectPtr<UPVGrowerPoint> Point, TObjectPtr<UPVGro
 	BudNumberIndexMap.Add(MaxBudNumber , Points.Num() - 1);
 }
 
+void UPVGrowerData::AddPoint(TObjectPtr<UPVGrowerPoint> Point, int32 BudNumber)
+{
+	Points.Add(Point);
+	BudNumberIndexMap.Add(BudNumber , Points.Num() - 1);
+}
+
 void UPVGrowerData::RemovePoint(int32 BudNumber)
 {
 	if (!BudNumberIndexMap.Contains(BudNumber))
@@ -64,6 +151,16 @@ TObjectPtr<UPVGrowerPoint> UPVGrowerData::GetPoint(int32 BudNumber) const
 	return Points[BudNumberIndexMap[BudNumber]];
 }
 
+TObjectPtr<UPVGrowerPoint> UPVGrowerData::GetPointFromIndex(int32 PointIndex) const
+{
+	if (Points.IsValidIndex(PointIndex))
+	{
+		return Points[PointIndex];
+	}
+
+	return nullptr;
+}
+
 int32 UPVGrowerData::GetPrimitiveIndex(int32 BranchNumber) const
 {
 	return Primitives.IndexOfByPredicate([BranchNumber](UPVGrowerPrimitive* Primitive)
diff --git a/Engine/Plugins/Experimental/ProceduralVegetationEditor/Source/ProceduralVegetation/Private/Implementations/PVGrowerData.h b/Engine/Plugins/Experimental/ProceduralVegetationEditor/Source/ProceduralVegetation/Private/Implementations/PVGrowerData.h
--- a/Engine/Plugins/Experimental/ProceduralVegetationEditor/Source/ProceduralVegetation/Private/Implementations/PVGrowerData.h
+++ b/Engine/Plugins/Experimental/ProceduralVegetationEditor/Source/ProceduralVegetation/Private/Implementations/PVGrowerData.h
@@ -26,6 +26,8 @@ struct FPVBudDevelopment
 	{
 		return {Generation, BudAge, BranchAge, AgeSenescense, LightSenescense, RelativeBudAge};
 	}
+
+	void SetData(const TArray<int>& InArray);
 };
 
 USTRUCT()
@@ -49,6 +51,8 @@ struct FPVBudDirection
 	{
 		return { FVector3f(Apical), FVector3f(Axillary), FVector3f(LightOptimal), FVector3f(LightSubOptimal), FVector3f(CurveGuide), FVector3f(UpVector)};//-UpVector.X, UpVector.Z, UpVector.Y)};
 	}
+
+	void SetData(const TArray<FVector3f>& InArray);
 };
 
 USTRUCT()
@@ -72,6 +76,8 @@ struct FPVBudHarmoneLevels
 	{
 		return { Apical, Axillary, AxillaryInhibition, Radical, Ethylene, Cykininin };
 	}
+
+	void SetData(const TArray<float>& InArray);
 };
 
 USTRUCT()
@@ -97,6 +103,8 @@ struct FPVBudLateralMeristem
 	{
 		return { LateralMeristem, Multiplier, Inactive, Davinci, ParentDot, RootDistance, Degredation };
 	}
+
+	void SetData(const TArray<float>& InArray);
 };
 
 USTRUCT()
@@ -116,6 +124,8 @@ struct FPVBudLightDetected
 	{
 		return { Available, Resource, Branch, Collision };
 	}
+
+	void SetData(const TArray<float>& InArray);
 };
 
 USTRUCT()
@@ -147,6 +157,8 @@ struct FPVBudStatus
 	{
 		return { ApicalMeristem, Codominant, Axillary, Seed, Dormant, Triggered, NumTriggered, Inactive, BrokenTip, Broken };
 	}
+
+	void SetData(const TArray<int>& InArray);
 };
 
 USTRUCT()
@@ -330,13 +342,24 @@ public:
 	int32 MaxBudNumber = 0;
 
 public:
+	/*
+	 * Add the point to Points list and associate to the branch
+	 * Increments the BudNumber and add point to BudNumberIndexMap
+	 */
 	void AddPoint(TObjectPtr<UPVGrowerPoint> Point, TObjectPtr<UPVGrowerPrimitive> Primitive, bool bAddToBranch = true);
-	
+
+	/*
+	 * Add the point to points list and also add it to BudNumberIndexMap
+	 */
+	void AddPoint(TObjectPtr<UPVGrowerPoint> Point, int32 BudNumber);
+
 	void RemovePoint(int32 BudNumber);
 	
 	int32 GetPointIndex(int32 BudNumber);
 	
 	TObjectPtr<UPVGrowerPoint> GetPoint(int32 BudNumber) const;
+	
+	TObjectPtr<UPVGrowerPoint> GetPointFromIndex(int32 PointIndex) const;
 
 	int32 GetPrimitiveIndex(int32 BranchNumber) const;
 
diff --git a/Engine/Plugins/Experimental/ProceduralVegetationEditor/Source/ProceduralVegetation/Private/Nodes/PVGrowerSettings.cpp b/Engine/Plugins/Experimental/ProceduralVegetationEditor/Source/ProceduralVegetation/Private/Nodes/PVGrowerSettings.cpp
--- a/Engine/Plugins/Experimental/ProceduralVegetationEditor/Source/ProceduralVegetation/Private/Nodes/PVGrowerSettings.cpp
+++ b/Engine/Plugins/Experimental/ProceduralVegetationEditor/Source/ProceduralVegetation/Private/Nodes/PVGrowerSettings.cpp
@@ -1,6 +1,8 @@
 ﻿// Copyright Epic Games, Inc. All Rights Reserved.
 
 #include "PVGrowerSettings.h"
+
+#include "ProceduralVegetationModule.h"
 #include "DataTypes/PVData.h"
 #include "DataTypes/PVGrowthData.h"
 #include "Data/PCGBasePointData.h"
@@ -39,7 +41,19 @@ TMap<UObject*, FTransform> UPVGrowerSettings::GetViewportObjects() const
 
 TArray<FPCGPinProperties> UPVGrowerSettings::InputPinProperties() const
 {
-	return TArray<FPCGPinProperties>();
+	TArray<FPCGPinProperties> Properties;
+
+	FPCGPinProperties& Pin = Properties.Emplace_GetRef(PCGPinConstants::DefaultInputLabel, GetInputPinTypeIdentifier());
+	Pin.SetAllowMultipleConnections(false);
+	Pin.SetOverrideOrUserParamPin();
+	Pin.bAllowMultipleData = false;
+
+	return Properties;
+}
+
+FPCGDataTypeIdentifier UPVGrowerSettings::GetInputPinTypeIdentifier() const
+{
+	return FPCGDataTypeIdentifier{ FPVDataTypeInfoGrowth::AsId() };
 }
 
 FPCGDataTypeIdentifier UPVGrowerSettings::GetOutputPinTypeIdentifier() const
@@ -62,15 +76,39 @@ bool FPVGrowerElement::ExecuteInternal(FPCGContext* InContext) const
 	
 	if (Settings)
 	{
-		FManagedArrayCollection OutCollection;
+		const TArray<FPCGTaggedData>& Inputs = InContext->InputData.GetInputsByPin(PCGPinConstants::DefaultInputLabel);
+		if (Inputs.IsValidIndex(0))
+		{
+			if (const UPVGrowthData* InputData = Cast<UPVGrowthData>(Inputs[0].Data))
+			{
+				const FManagedArrayCollection& SkeletonCollection = InputData->GetCollection();
+
+				FManagedArrayCollection OutCollection;
+		
+				FPVGrower::Grow(SkeletonCollection, Settings->GrowerParams, OutCollection);
+			
+				UPVGrowthData* OutManagedArrayCollectionData = FPCGContext::NewObject_AnyThread<UPVGrowthData>(InContext);
+				OutManagedArrayCollectionData->Initialize(MoveTemp(OutCollection));
+		
+				InContext->OutputData.TaggedData.Emplace(OutManagedArrayCollectionData);
+			}
+			else
+			{
+				UE_LOG(LogProceduralVegetation, Log, TEXT("Invalid Input Data for FPVGrowerElement"));
+			}
+		}
+		else
+		{
+			FManagedArrayCollection OutCollection;
 		
-		FPVSeedPoint SeedPoint;
-		FPVGrower::Grow(SeedPoint, Settings->GrowerParams, OutCollection);
+			FPVSeedPoint SeedPoint;
+			FPVGrower::Grow(SeedPoint, Settings->GrowerParams, OutCollection);
 			
-		UPVGrowthData* OutManagedArrayCollectionData = FPCGContext::NewObject_AnyThread<UPVGrowthData>(InContext);
-		OutManagedArrayCollectionData->Initialize(MoveTemp(OutCollection));
+			UPVGrowthData* OutManagedArrayCollectionData = FPCGContext::NewObject_AnyThread<UPVGrowthData>(InContext);
+			OutManagedArrayCollectionData->Initialize(MoveTemp(OutCollection));
 		
-		InContext->OutputData.TaggedData.Emplace(OutManagedArrayCollectionData);
+			InContext->OutputData.TaggedData.Emplace(OutManagedArrayCollectionData);
+		}
 	}
 	else
 	{
diff --git a/Engine/Plugins/Experimental/ProceduralVegetationEditor/Source/ProceduralVegetation/Private/Nodes/PVGrowerSettings.h b/Engine/Plugins/Experimental/ProceduralVegetationEditor/Source/ProceduralVegetation/Private/Nodes/PVGrowerSettings.h
--- a/Engine/Plugins/Experimental/ProceduralVegetationEditor/Source/ProceduralVegetation/Private/Nodes/PVGrowerSettings.h
+++ b/Engine/Plugins/Experimental/ProceduralVegetationEditor/Source/ProceduralVegetation/Private/Nodes/PVGrowerSettings.h
@@ -25,6 +25,7 @@ public:
 
 protected:
 	virtual TArray<FPCGPinProperties> InputPinProperties() const override;
+	virtual FPCGDataTypeIdentifier GetInputPinTypeIdentifier() const override;
 	virtual FPCGDataTypeIdentifier GetOutputPinTypeIdentifier() const override;
 	virtual FPCGElementPtr CreateElement() const override;
 
diff --git a/Engine/Plugins/Experimental/ProceduralVegetationEditor/Source/ProceduralVegetation/Public/Facades/PVAttributesNames.h b/Engine/Plugins/Experimental/ProceduralVegetationEditor/Source/ProceduralVegetation/Public/Facades/PVAttributesNames.h
--- a/Engine/Plugins/Experimental/ProceduralVegetationEditor/Source/ProceduralVegetation/Public/Facades/PVAttributesNames.h
+++ b/Engine/Plugins/Experimental/ProceduralVegetationEditor/Source/ProceduralVegetation/Public/Facades/PVAttributesNames.h
@@ -52,6 +52,8 @@ namespace PV::AttributeNames
 	inline static const FName BudNumber = FName(TEXT("budNumber"));
 	inline static const FName BudHormoneLevels = FName(TEXT("budHormoneLevels"));
 	inline static const FName BudDirection = FName(TEXT("budDirection"));
+	inline static const FName BudStatus = FName(TEXT("budStatus"));
+	inline static const FName BudLateralMeristem = FName(TEXT("budLateralMeristem"));
 	inline static const FName BranchGradient = FName(TEXT("branchGradient"));
 	inline static const FName PlantGradient = FName(TEXT("plantGradient"));
 	inline static const FName NjordPixelIndex = FName(TEXT("njord_pixelIdx"));
diff --git a/Engine/Plugins/Experimental/ProceduralVegetationEditor/Source/ProceduralVegetation/Public/Facades/PVBranchFacade.h b/Engine/Plugins/Experimental/ProceduralVegetationEditor/Source/ProceduralVegetation/Public/Facades/PVBranchFacade.h
--- a/Engine/Plugins/Experimental/ProceduralVegetationEditor/Source/ProceduralVegetation/Public/Facades/PVBranchFacade.h
+++ b/Engine/Plugins/Experimental/ProceduralVegetationEditor/Source/ProceduralVegetation/Public/Facades/PVBranchFacade.h
@@ -113,6 +113,35 @@ namespace PV::Facades
 		
 		void FillPlantNumbers(const TArray<int32>& InputArray);
 
+		template<typename T>
+		FORCEINLINE bool GetValueFromAccessor(TManagedArrayAccessor<T> Accessor, T& OutValue, int32 Index, const T DefaultValue) const
+		{
+			if ((Accessor).IsValid() && (Accessor).IsValidIndex(Index))
+			{ 
+				OutValue = (Accessor)[Index]; 
+				return true;; 
+			}
+
+			OutValue = DefaultValue;
+			return false;
+		}
+
+		void GetParents(int32 Index, TArray<int32>& OutValue) const;
+
+		void GetChildren(int32 Index, TArray<int32>& OutValue) const;
+
+		void GetBranchPoints(int32 Index, TArray<int32>& OutValue) const;
+		
+		void GetBranchNumber(int32 Index, int32& OutValue) const;
+		
+		void GetBranchSourceBudNumber(int32 Index, int32& OutValue) const;
+		
+		void GetBranchHierarchyNumber(int32 Index, int32& OutValue) const;
+		
+		void GetBranchParentNumber(int32 Index, int32& OutValue) const;
+		
+		void GetPlantNumber(int32 Index, int32& OutValue) const;
+
 	private:
 		FManagedArrayCollection* Collection = nullptr;
 
diff --git a/Engine/Plugins/Experimental/ProceduralVegetationEditor/Source/ProceduralVegetation/Public/Facades/PVPointFacade.h b/Engine/Plugins/Experimental/ProceduralVegetationEditor/Source/ProceduralVegetation/Public/Facades/PVPointFacade.h
--- a/Engine/Plugins/Experimental/ProceduralVegetationEditor/Source/ProceduralVegetation/Public/Facades/PVPointFacade.h
+++ b/Engine/Plugins/Experimental/ProceduralVegetationEditor/Source/ProceduralVegetation/Public/Facades/PVPointFacade.h
@@ -4,10 +4,25 @@
 
 #include "CoreMinimal.h"
 #include "IShrinkableFacade.h"
+#include "VisualizeTexture.h"
+#include "Distributions/Distribution.h"
 #include "GeometryCollection/ManagedArrayAccessor.h"
 
 namespace PV::Facades
 {
+#define FILL_COLLECTION_ATTRIBUTE(Accessor, InputArray) \
+if (!Accessor.IsValid()) \
+{ \
+Accessor.Add(); \
+} \
+if (Accessor.Num() < InputArray.Num()) \
+{ \
+Accessor.AddElements(InputArray.Num() - Accessor.Num()); \
+} \
+for (int32 i = 0; i < InputArray.Num(); i++) \
+{ \
+Accessor.ModifyAt(i, InputArray[i]); \
+}
 	/**
 	 * FPointFacade is used to access and manipulate the Point Group data from the ProceduralVegetation's FManagedArrayCollection
 	 * Only add the frequently used Point attributes and their access to this facade, for the specific access write a new facade
@@ -144,6 +159,55 @@ namespace PV::Facades
 
 		void FillNjordPixelID(const TArray<float>& InputArray);
 
+		void FillBudStatus(const TArray<TArray<int>>& InputArray);
+
+		void FillBudLateralMeristem(const TArray<TArray<float>>& InputArray);
+
+		template<typename T>
+		FORCEINLINE bool GetValueFromAccessor(TManagedArrayAccessor<T> Accessor, T& OutValue, int32 Index, const T DefaultValue) const
+		{
+			if ((Accessor).IsValid() && (Accessor).IsValidIndex(Index))
+			{ 
+				OutValue = (Accessor)[Index]; 
+				return true;; 
+			}
+
+			OutValue = DefaultValue;
+			return false;
+		}
+		
+		void GetLFR(int32 Index, float& LFR) const;
+
+		void GetPointScaleGradient(int32 Index, float& OutPointScaleGradient) const;
+		
+		void GetHullGradient(int32 Index, float& OutHullGradient) const;
+
+		void GetMainTrunkGradients(int32 Index, float& OutValue) const;
+
+		void GetGroundGradients(int32 Index, float& OutValue) const;
+
+		void GetPointScales(int32 Index, float& OutValue) const;
+
+		void GetBudDirections(int32 Index, TArray<FVector3f>& OutValue) const;
+
+		void GetBranchGradients(int32 Index, float& OutValue) const;
+
+		void GetBudHormoneLevels(int32 Index, TArray<float>& OutValue) const;
+
+		void GetPlantGradients(int32 Index, float& OutValue) const;
+
+		void GetBudNumbers(int32 Index, int32& OutValue) const;
+
+		void GetBudLightDetected(int32 Index, TArray<float>& OutValue) const;
+
+		void GetBudDevelopment(int32 Index, TArray<int>& OutValue) const;
+		
+		void GetNjordPixelID(int32 Index, float& OutValue) const;
+
+		void GetBudStatus(int32 Index, TArray<int>& OutValue) const;
+		
+		void GetBudLateralMeristem(int32 Index, TArray<float>& OutValue) const;
+
 	private:
 		TManagedArrayAccessor<FVector3f> Positions;
 		TManagedArrayAccessor<float> LengthFromRoot;
@@ -153,6 +217,8 @@ namespace PV::Facades
 		TManagedArrayAccessor<float> GroundGradient;
 		TManagedArrayAccessor<float> PointScale;
 		TManagedArrayAccessor<TArray<FVector3f>> BudDirections;
+		TManagedArrayAccessor<TArray<int>> BudStatus;
+		TManagedArrayAccessor<TArray<float>> BudLateralMeristem;
 		TManagedArrayAccessor<float> BranchGradients;
 		TManagedArrayAccessor<TArray<float>> BudHormoneLevels;
 		TManagedArrayAccessor<float> PlantGradients;
