Fix warnings in PVLightDetectionCS.
- using inout for some function parameters to match behavior of references in C++.
- prevented running with r.Shaders.WarningsAsErrors=1

#Synced-CL 50178168

diff --git a/Engine/Plugins/Experimental/ProceduralVegetationEditor/Shaders/PVLightDetectionCS.usf b/Engine/Plugins/Experimental/ProceduralVegetationEditor/Shaders/PVLightDetectionCS.usf
--- a/Engine/Plugins/Experimental/ProceduralVegetationEditor/Shaders/PVLightDetectionCS.usf
+++ b/Engine/Plugins/Experimental/ProceduralVegetationEditor/Shaders/PVLightDetectionCS.usf
@@ -70,7 +70,7 @@ int GetTransformIndex(int Index)
 }
 
 //This function is the HLSL translation of FindIntersection in IterRay3Triangle3.h
-bool RayTriangleHit(float3 RayOrigin, float3 RayDir, float3 Vert0, float3 Vert1, float3 Vert2, out float t, out float3 HitPos, out float3 HitNormal)
+bool RayTriangleHit(float3 RayOrigin, float3 RayDir, float3 Vert0, float3 Vert1, float3 Vert2, inout float t, inout float3 HitPos, inout float3 HitNormal)
 {
 	// Compute the offset origin, edges, and normal.
 	float3 Diff = RayOrigin - Vert0;
@@ -168,7 +168,7 @@ bool RayMeshHit(float3 RayOrigin, float3 RayDir, out float tOut, out float3 HitP
 	return bHit;
 }
 
-bool RayCylinderHit(float3 RayOrigin, float3 RayDirection, PVCollisionData Data, out float t, out float3 Normal, out int FalureType)
+bool RayCylinderHit(float3 RayOrigin, float3 RayDirection, PVCollisionData Data, inout float t, inout float3 Normal, inout int FalureType)
 {
 	float3 PA = Data.Position;
 	float3 PB = Data.Position + Data.Direction * Data.Extents.y;
