Fix PVImporter_Texture2D failing to triangulate perimeter curve if it was simplified too much

- Refactor PVImporter_Texture2D to use FPolygon2f
- Replace custom simplification algorithm with FPolygon2f::Simplify
- Replace GetFilledTriangles with GetFilledTrianglesGeneralizedWinding to resolve triangulation failure

#Synced-CL 50196042

diff --git a/Engine/Plugins/Experimental/ProceduralVegetationEditor/Source/ProceduralVegetation/Public/PVImportTexture2DParams.h b/Engine/Plugins/Experimental/ProceduralVegetationEditor/Source/ProceduralVegetation/Public/PVImportTexture2DParams.h
--- a/Engine/Plugins/Experimental/ProceduralVegetationEditor/Source/ProceduralVegetation/Public/PVImportTexture2DParams.h
+++ b/Engine/Plugins/Experimental/ProceduralVegetationEditor/Source/ProceduralVegetation/Public/PVImportTexture2DParams.h
@@ -54,10 +54,10 @@ struct FPVImportTexture2DParams
 	float MinBoundsArea = 0.05f;
 	
 	UPROPERTY(EditAnywhere, Category = "Image Trace", meta=(ClampMin=0, ClampMax=100, UIMin=0, UIMax=10))
-	int32 SmoothingIterations = 1;
+	int32 SmoothingIterations = 2;
 	
-	UPROPERTY(EditAnywhere, Category = "Image Trace", meta=(ClampMin=0, ClampMax=0.1, UIMin=0, UIMax=0.01))
-	float SimplificationEpsilon = 0.001;
+	UPROPERTY(EditAnywhere, Category = "Image Trace", meta=(ClampMin=0, ClampMax=100, UIMin=0, UIMax=10.0))
+	float SimplificationAmount = 1.0;
 
 	UPROPERTY(EditAnywhere, Category = "Texture")
 	TObjectPtr<UTexture2D> Texture2DAsset = nullptr;
diff --git a/Engine/Plugins/Experimental/ProceduralVegetationEditor/Source/ProceduralVegetationEditor/Private/Importer/PVImporter_Texture2D.cpp b/Engine/Plugins/Experimental/ProceduralVegetationEditor/Source/ProceduralVegetationEditor/Private/Importer/PVImporter_Texture2D.cpp
--- a/Engine/Plugins/Experimental/ProceduralVegetationEditor/Source/ProceduralVegetationEditor/Private/Importer/PVImporter_Texture2D.cpp
+++ b/Engine/Plugins/Experimental/ProceduralVegetationEditor/Source/ProceduralVegetationEditor/Private/Importer/PVImporter_Texture2D.cpp
@@ -11,6 +11,7 @@
 #include "Helpers/PVFacadeHelper.h"
 
 #include "Engine/Texture2D.h"
+#include "Polygon2.h"
 
 #include "Materials/MaterialInterface.h"
 
@@ -25,7 +26,8 @@
 namespace Internal
 {
 	using FBooleanVertexAttribute = PV::DynamicMeshHelper::FBooleanVertexAttribute;
-	using FFloatVertexAttribute = PV::DynamicMeshHelper::FFloatVertexAttribute; using FIndex2i = UE::Geometry::FIndex2i;
+	using FFloatVertexAttribute = PV::DynamicMeshHelper::FFloatVertexAttribute; 
+	using FIndex2i = UE::Geometry::FIndex2i;
 
 	using FIndex3i = UE::Geometry::FIndex3i;
 	using FVector2i = UE::Geometry::FVector2i;
@@ -33,6 +35,8 @@ namespace Internal
 	using FDynamicMeshAABBTree3 = UE::Geometry::FDynamicMeshAABBTree3;
 	using FDynamicMesh3 = UE::Geometry::FDynamicMesh3;
 	using FDynamicMeshQueries3 = UE::Geometry::TMeshQueries<FDynamicMesh3>;
+	using FPolygon2f = UE::Geometry::FPolygon2f;
+	using FDelaunay2 = UE::Geometry::FDelaunay2;
 
 	const static FName PixelsDebugGroupName = TEXT("PixelDebug");
 	const static FName PixelPositionsAttributeName = TEXT("Position");
@@ -100,7 +104,7 @@ namespace Internal
 
 	static void AddPerimeterCurvesToCollection(
 		FManagedArrayCollection& InOutCollection, 
-		const TArray<TArray<FVector2f>>& InPerimeterCurves, 
+		const TArray<FPolygon2f>& InPerimeterCurves,
 		float Scale, 
 		const FVector& Offset = FVector::ZeroVector
 	)
@@ -111,7 +115,7 @@ namespace Internal
 		auto& PerimeterCurveAttribute = InOutCollection.AddAttribute<TArray<FVector3f>>(PerimeterCurveAttributeName, PerimeterCurveGroupName);
 		for (int32 i = 0; i < InPerimeterCurves.Num(); i++)
 		{
-			const TArray<FVector2f>& Curve = InPerimeterCurves[i];
+			const TArray<FVector2f>& Curve = InPerimeterCurves[i].GetVertices();
 			for (const FVector2f& CurvePos : Curve)
 			{
 				PerimeterCurveAttribute[i].Add(ToVector3(CurvePos, Scale) + FVector3f(Offset));
@@ -260,7 +264,7 @@ namespace Internal
 		}
 	}
 	
-	static void TraceImage(const TArray<bool>& ImageData, const FVector2i& ImageSize, TArray<TArray<FVector2f>>& OutCurves)
+	static void TraceImage(const TArray<bool>& ImageData, const FVector2i& ImageSize, TArray<FPolygon2f>& OutCurves)
 	{
 		TRACE_CPUPROFILER_EVENT_SCOPE(TraceImage);
 		
@@ -276,6 +280,7 @@ namespace Internal
 			false
 		);
 		
+		TArray<FVector2f> PerimeterCurve;
 		OutCurves.Reserve(Results.Num());
 		for (const PCGSpatialAlgo::FPCGMarchingSquareResult& Result : Results)
 		{
@@ -284,21 +289,21 @@ namespace Internal
 				continue;
 			}
 			
-			TArray<FVector2f>& OutCurve = OutCurves.AddDefaulted_GetRef();
-			OutCurve.Reserve(Result.LinkedGridCoordinates.Num());
+			PerimeterCurve.Reset(Result.LinkedGridCoordinates.Num());
 			for (const FVector2D& Coordinate : Result.LinkedGridCoordinates)
 			{
-				OutCurve.Add(FVector2f(Coordinate.X / (ImageSize.X - 1), Coordinate.Y / (ImageSize.Y - 1)));
+				PerimeterCurve.Add(FVector2f(Coordinate.X / (ImageSize.X - 1), Coordinate.Y / (ImageSize.Y - 1)));
 			}
+			OutCurves.Emplace(MoveTemp(PerimeterCurve));
 		}
 	}
 	
-	static void RemoveSmallCurves(TArray<TArray<FVector2f>>& InOutCurves, float Threshold)
+	static void RemoveSmallCurves(TArray<FPolygon2f>& InOutPerimeterCurves, float Threshold)
 	{
 		const float ThresholdSqr = Threshold * Threshold;
-		for (int32 i = InOutCurves.Num() - 1; i >= 0; i--)
+		for (int32 i = InOutPerimeterCurves.Num() - 1; i >= 0; i--)
 		{
-			const TArray<FVector2f>& Curve = InOutCurves[i];
+			const TArray<FVector2f>& Curve = InOutPerimeterCurves[i].GetVertices();
 			FVector2f Min = FVector2f(MAX_flt);
 			FVector2f Max = FVector2f(-MAX_flt);
 			
@@ -314,122 +319,50 @@ namespace Internal
 			const float AreaSqr = Area * Area;
 			if (AreaSqr < ThresholdSqr)
 			{
-				InOutCurves.RemoveAtSwap(i);
+				InOutPerimeterCurves.RemoveAtSwap(i);
 			}
 		}
 	}
 	
-	static void SmoothCurves(TArray<TArray<FVector2f>>& InOutCurves, int32 SmoothingIterations)
+	static void SmoothCurves(TArray<FPolygon2f>& InOutPerimeterCurves, int32 SmoothingIterations)
 	{
 		while (SmoothingIterations != 0)
 		{
-			for (TArray<FVector2f>& Curve : InOutCurves)
+			for (FPolygon2f& PerimeterCurve : InOutPerimeterCurves)
 			{
+				const TArray<FVector2f>& Vertices = PerimeterCurve.GetVertices();
+
 				TArray<FVector2f> SmoothedCurve;
-				SmoothedCurve.SetNum(Curve.Num());
+				SmoothedCurve.SetNum(Vertices.Num());
 			
-				for (int32 i = 0; i < Curve.Num(); i++)
+				for (int32 i = 0; i < Vertices.Num(); i++)
 				{
-					const FVector2f& Pos1 = i == 0 ? Curve.Last() : Curve[i - 1];
-					const FVector2f& Pos2 = Curve[i];
-					const FVector2f& Pos3 = i == Curve.Num() - 1 ? Curve[0] : Curve[i + 1];
+					const FVector2f& Pos1 = i == 0 ? Vertices.Last() : Vertices[i - 1];
+					const FVector2f& Pos2 = Vertices[i];
+					const FVector2f& Pos3 = i == Vertices.Num() - 1 ? Vertices[0] : Vertices[i + 1];
 				
 					const FVector2f AvgPos = (Pos1 + Pos2 + Pos3) / 3; 
 					SmoothedCurve[i] = AvgPos;
 				}
 			
-				Curve = MoveTemp(SmoothedCurve);
+				PerimeterCurve = FPolygon2f(MoveTemp(SmoothedCurve));
 			}
 			
 			SmoothingIterations -= 1;
 		}
 	}
 	
-	static void SimplifyCurves(TArray<TArray<FVector2f>>& InOutCurves, float SimplificationEpsilon)
+	static void SimplifyCurves(TArray<FPolygon2f>& InOutPerimeterCurves, float SimplificationAmount)
 	{
-		// TODO: We should probably be using TPolygon2::Simplify for this algorithm, however, it seems to produce slightly different results so more 
-		// thorough testing is required.
-		
-		/*
-		 * Non-recursive version of Ramer-Douglas-Peucker algorithm (https://en.wikipedia.org/wiki/Ramer%E2%80%93Douglas%E2%80%93Peucker_algorithm)
-		 *
-		 * Given the following curve, the algorithm will take the start and end point (d,e) and find the point furthest away (a) from the line segment
-		 * made up from d->e. It then takes the points (d,a) and (e,a), and finds the point between d->a and e->a that are furthest away from the line
-		 * segments made up from d->a and e->a respectively. In this case that would be the points b and c. This is repeated until the point found is
-		 * closer to the line than the SimplificationEpsilon, at which point the subdivision stops for that segment.
-		 *
-		 *    a
-		 *   /\
-		 * b/  \c
-		 *  \  / 
-		 *  d  e
-		 */
-	
-		const static auto FindFarPoint = [](const TArray<FVector2f>& CurvePoints, int32 StartPointIndex, int32 EndPointIndex, float Epsilon)->int32
+		for (FPolygon2f& PerimeterCurve : InOutPerimeterCurves)
 		{
-			const FVector2f& StartPos = CurvePoints[StartPointIndex];
-			const FVector2f& EndPos = CurvePoints[EndPointIndex];
-			float MaxDistSqr = Epsilon * Epsilon;
-			int32 FarPointIndex = INDEX_NONE;
-			for (int32 i = StartPointIndex + 1; i < EndPointIndex - 1; i++)
-			{
-				const FVector2f& Pos = CurvePoints[i];
-				const FVector2f ClosestPoint = FMath::ClosestPointOnSegment2D(Pos, StartPos, EndPos);
-				const float DistSqr = (Pos - ClosestPoint).SizeSquared();
-				if (DistSqr >= MaxDistSqr)
-				{ 
-					MaxDistSqr = DistSqr;
-					FarPointIndex = i;
-				}
-			}
-			return FarPointIndex;
-		};
-		
-		if (SimplificationEpsilon < 0.000001)
-		{
-			return; 
-		}
-		
-		TArray<int32> DouglasPeuckerQueue;
-		TArray<int32> NewDouglasPeuckerQueue;
-		TArray<int32> KeepPoints;
-		for (TArray<FVector2f>& Curve : InOutCurves)
-		{
-			DouglasPeuckerQueue.Reset(Curve.Num() * 2);
-			NewDouglasPeuckerQueue.Reset(Curve.Num() * 2);
-			KeepPoints.Reset(Curve.Num());
-			
-			DouglasPeuckerQueue.Append({ 0, Curve.Num() - 1 });
-			KeepPoints.Append(DouglasPeuckerQueue);
-			
-			while (DouglasPeuckerQueue.Num() > 0)
-			{
-				NewDouglasPeuckerQueue.Reset();
-				for (int32 i = 0; i < DouglasPeuckerQueue.Num(); i += 2)
-				{
-					const int32 StartPointIndex = DouglasPeuckerQueue[i];
-					const int32 EndPointIndex = DouglasPeuckerQueue[i + 1];
-					const int32 FarPointIndex = FindFarPoint(Curve, StartPointIndex, EndPointIndex, SimplificationEpsilon);
-					if (FarPointIndex != INDEX_NONE)
-					{ 
-						KeepPoints.Add(FarPointIndex);
-						NewDouglasPeuckerQueue.Append({ StartPointIndex, FarPointIndex, FarPointIndex, EndPointIndex });
-					}
-				}
-				Swap(DouglasPeuckerQueue, NewDouglasPeuckerQueue);
-			}
-			
-			KeepPoints.Sort();
-			for (int32 i = 0; i < KeepPoints.Num(); i++)
-			{
-				const int32 PointIndex = KeepPoints[i];
-				Curve[i] = Curve[PointIndex]; // It's safe to overwrite the Curve array with itself since KeepPoints is sorted, so i can never be larger than PointIndex
-			}
-			Curve.SetNum(KeepPoints.Num());
+			const float LineDeviationTolerance = SimplificationAmount * 0.001;
+			const float ClusterTolerance = LineDeviationTolerance * 0.001;
+			PerimeterCurve.Simplify(ClusterTolerance, LineDeviationTolerance);
 		}
 	}
 	
-	static void TriangulateCurvesAndSplitEdges(const TArray<TArray<FVector2f>>& InCurves, UE::Geometry::FDynamicMesh3& OutMesh, float Scale)
+	static void TriangulateCurvesAndSplitEdges(const TArray<FPolygon2f>& InPerimeterCurves, UE::Geometry::FDynamicMesh3& OutMesh, float Scale)
 	{
 		/*
 		* This function convertes the perimeter curves into a DynamicMesh and adds a vertex at the center of each edge spanning two perimeter points.
@@ -451,34 +384,45 @@ namespace Internal
 
 		PerimeterAttributeDefinition.AttachAttribute(OutMesh);
 
-		for (int32 CurveIndex = 0; CurveIndex < InCurves.Num(); CurveIndex++)
+		for (int32 CurveIndex = 0; CurveIndex < InPerimeterCurves.Num(); CurveIndex++)
 		{
-			const TArray<FVector2f>& Curve = InCurves[CurveIndex];
-			
+			const FPolygon2f& PerimeterCurve = InPerimeterCurves[CurveIndex];
+			const TArray<FVector2f>& CurveVertices = PerimeterCurve.GetVertices();
+
 			const TArray<FIndex2i> OriginalCurveEdges = Invoke([&]() 
 			{
 				TArray<FIndex2i> Edges;
-				Edges.Reserve(Curve.Num());
-				for (int32 i = 0; i < Curve.Num() - 1; i++)
+				Edges.Reserve(CurveVertices.Num());
+
+				for (int32 Last = CurveVertices.Num() - 1, Idx = 0; Idx < CurveVertices.Num(); Last = Idx++)
 				{
-					Edges.Emplace(i, i + 1);
+					Edges.Emplace(Last, Idx);
 				}
-				Edges.Emplace(0, Curve.Num() - 1);
+
 				return Edges;
 			});
 			
-			UE::Geometry::FDelaunay2 Delaunay;
-			bool bTriangulateSucess = Delaunay.Triangulate(Curve, OriginalCurveEdges);
+			FDelaunay2 Delaunay;
+			Delaunay.Triangulate(CurveVertices, OriginalCurveEdges);
+			
+			bool bTriangulateSucess = Delaunay.GetResult() == UE::Geometry::FDelaunay2::EResult::MissingEdges
+				|| Delaunay.GetResult() == UE::Geometry::FDelaunay2::EResult::Success;
+
 			if (!bTriangulateSucess)
 			{
 				continue;
 			}
+
+			const FDelaunay2::EFillMode FillMode = PerimeterCurve.IsClockwise() 
+				? FDelaunay2::EFillMode::NegativeWinding 
+				: FDelaunay2::EFillMode::PositiveWinding;
 			
-			const TArray<FIndex3i> OriginalCurveTriangles = Delaunay.GetFilledTriangles(OriginalCurveEdges, UE::Geometry::FDelaunay2::EFillMode::Solid);
-			
+			TArray<FIndex3i> OriginalCurveTriangles;
+			Delaunay.GetFilledTrianglesGeneralizedWinding(OriginalCurveTriangles, CurveVertices, OriginalCurveEdges, FillMode);
+
 			TArray<FVector2f> SubdividedMeshPoints;
-			SubdividedMeshPoints.Reserve(Curve.Num() + OriginalCurveTriangles.Num());
-			SubdividedMeshPoints.Append(Curve);
+			SubdividedMeshPoints.Reserve(CurveVertices.Num() + OriginalCurveTriangles.Num());
+			SubdividedMeshPoints.Append(CurveVertices);
 			
 			TArray<FIndex2i> SubdividedMeshEdges;
 			SubdividedMeshEdges.Reserve(OriginalCurveTriangles.Num() * 4);
@@ -501,8 +445,8 @@ namespace Internal
 					{
 						const int32 PointIndex1 = EdgeIndex[0];
 						const int32 PointIndex2 = EdgeIndex[1];
-						const FVector2f& Pos1 = Curve[PointIndex1];
-						const FVector2f& Pos2 = Curve[PointIndex2];
+						const FVector2f& Pos1 = CurveVertices[PointIndex1];
+						const FVector2f& Pos2 = CurveVertices[PointIndex2];
 						const int32 NewPointIndex = SubdividedMeshPoints.Add((Pos1 + Pos2) / 2);
 						
 						const FIndex2i NewEdge1(PointIndex1, NewPointIndex);
@@ -514,7 +458,11 @@ namespace Internal
 				}
 			}
 			
-			bTriangulateSucess = Delaunay.Triangulate(SubdividedMeshPoints, SubdividedMeshEdges);
+			Delaunay.Triangulate(SubdividedMeshPoints, SubdividedMeshEdges);
+
+			bTriangulateSucess = Delaunay.GetResult() == UE::Geometry::FDelaunay2::EResult::MissingEdges
+				|| Delaunay.GetResult() == UE::Geometry::FDelaunay2::EResult::Success;
+
 			if (!bTriangulateSucess)
 			{
 				continue;
@@ -524,8 +472,9 @@ namespace Internal
 			PerimeterAttributeDefinition.AttachAttribute(NewDynamicMesh);
 			auto& PerimeterAttribute = PerimeterAttributeDefinition.GetAttributeChecked(NewDynamicMesh);
 			
-			const TArray<FIndex3i> Triangles = Delaunay.GetFilledTriangles(OriginalCurveEdges, UE::Geometry::FDelaunay2::EFillMode::Solid);
 			const TArray<FVector2f>& Vertices = SubdividedMeshPoints;
+			TArray<FIndex3i> Triangles;
+			Delaunay.GetFilledTrianglesGeneralizedWinding(Triangles, Vertices, OriginalCurveEdges, FillMode);
 			
 			for (int32 i = 0; i < Vertices.Num(); ++i)
 			{
@@ -536,7 +485,7 @@ namespace Internal
 				);
 				
 				const int32 VertexID = NewDynamicMesh.AppendVertex(VertexInfo);
-				PerimeterAttribute.SetValue(VertexID, i < Curve.Num());
+				PerimeterAttribute.SetValue(VertexID, i < CurveVertices.Num());
 			}
 			
 			for (const FIndex3i& Triangle : Triangles)
@@ -953,7 +902,7 @@ PV::Texture2DImport::EImportResult PV::Texture2DImport::ImportGrowthDataFromText
 		return EImportResult::Success;
 	}
 
-	TArray<TArray<FVector2f>> PerimeterCurves;
+	TArray<FPolygon2f> PerimeterCurves;
 	TraceImage(BlackAndWhiteImageData, SampleSize, PerimeterCurves);
 	if (InParams.DebugState == EPVImportTexture2DDebugState::TraceImage)
 	{
@@ -978,7 +927,7 @@ PV::Texture2DImport::EImportResult PV::Texture2DImport::ImportGrowthDataFromText
 		return EImportResult::Success;
 	}
 	
-	SimplifyCurves(PerimeterCurves, InParams.SimplificationEpsilon);
+	SimplifyCurves(PerimeterCurves, InParams.SimplificationAmount);
 	if (InParams.DebugState == EPVImportTexture2DDebugState::SimplifyCurves)
 	{
 		AddPerimeterCurvesToCollection(OutCollection, PerimeterCurves, InParams.Scale);
@@ -1090,4 +1039,4 @@ PV::Texture2DImport::EImportResult PV::Texture2DImport::ImportGrowthDataFromText
 	AddRootPointsToCollection(Output.VisualizationCollection, InParams.RootPoints, InParams.Scale, VisualizationOffset);
 
 	return EImportResult::Success;
-}
\ No newline at end of file
+}
