[PVE] More Visualization Modes

#Synced-CL 50043089

diff --git a/Engine/Plugins/Experimental/ProceduralVegetationEditor/Source/ProceduralVegetation/Private/Facades/PVMetaInfoFacade.cpp b/Engine/Plugins/Experimental/ProceduralVegetationEditor/Source/ProceduralVegetation/Private/Facades/PVMetaInfoFacade.cpp
--- a/Engine/Plugins/Experimental/ProceduralVegetationEditor/Source/ProceduralVegetation/Private/Facades/PVMetaInfoFacade.cpp
+++ b/Engine/Plugins/Experimental/ProceduralVegetationEditor/Source/ProceduralVegetation/Private/Facades/PVMetaInfoFacade.cpp
@@ -9,6 +9,9 @@ namespace PV::Facades
 		: ConstCollection(InCollection)
 		, Collection(&InCollection)
 		, GuidAttribute(InCollection, AttributeNames::Guid, GroupNames::DetailsGroup)
+		, LeafGrowthAttribute(InCollection, AttributeNames::LeafGrowth, GroupNames::DetailsGroup)
+		, AbscissionSenescenseAttribute(InCollection, AttributeNames::AbscissionSenescense, GroupNames::DetailsGroup)
+		, LateralElongationAttribute(InCollection, AttributeNames::LateralElongation, GroupNames::DetailsGroup)
 	{
 		DefineSchema(InitialSize);
 	}
@@ -17,6 +20,9 @@ namespace PV::Facades
 		: ConstCollection(InCollection)
 		, Collection(nullptr)
 		, GuidAttribute(InCollection, AttributeNames::Guid, GroupNames::DetailsGroup)
+		, LeafGrowthAttribute(InCollection, AttributeNames::LeafGrowth, GroupNames::DetailsGroup)
+		, AbscissionSenescenseAttribute(InCollection, AttributeNames::AbscissionSenescense, GroupNames::DetailsGroup)
+		, LateralElongationAttribute(InCollection, AttributeNames::LateralElongation, GroupNames::DetailsGroup)
 	{}
 
 	void FMetaInfoFacade::DefineSchema(const int32 InitialSize)
@@ -59,4 +65,37 @@ namespace PV::Facades
 
 		return FGuid::NewGuid();
 	}
+
+	const TArray<float>& FMetaInfoFacade::GetLeafGrowth() const
+	{
+		if(LeafGrowthAttribute.IsValid() && LeafGrowthAttribute.IsValidIndex(0))
+		{
+			return LeafGrowthAttribute[0];
+		}
+
+		static TArray<float> DefaultLeafGrowthArray;
+		return DefaultLeafGrowthArray;
+	}
+
+	const TArray<float>& FMetaInfoFacade::GetAbscissionSenescense() const
+	{
+		if(AbscissionSenescenseAttribute.IsValid() && AbscissionSenescenseAttribute.IsValidIndex(0))
+		{
+			return AbscissionSenescenseAttribute[0];
+		}
+
+		static TArray<float> DefaultAbscissionSenescenseArray;
+		return DefaultAbscissionSenescenseArray;
+	}
+	
+	const TArray<float>& FMetaInfoFacade::GetLateralElongation() const
+	{
+		if(LateralElongationAttribute.IsValid() && LateralElongationAttribute.IsValidIndex(0))
+		{
+			return LateralElongationAttribute[0];
+		}
+
+		static TArray<float> DefaultLateralElongationArray;
+		return DefaultLateralElongationArray;
+	}
 }
diff --git a/Engine/Plugins/Experimental/ProceduralVegetationEditor/Source/ProceduralVegetation/Private/Facades/PVPointFacade.cpp b/Engine/Plugins/Experimental/ProceduralVegetationEditor/Source/ProceduralVegetation/Private/Facades/PVPointFacade.cpp
--- a/Engine/Plugins/Experimental/ProceduralVegetationEditor/Source/ProceduralVegetation/Private/Facades/PVPointFacade.cpp
+++ b/Engine/Plugins/Experimental/ProceduralVegetationEditor/Source/ProceduralVegetation/Private/Facades/PVPointFacade.cpp
@@ -297,6 +297,17 @@ namespace PV::Facades
 		}
 	}
 
+	const TArray<float>& FPointFacade::GetBudLateralMeristem(int32 Index) const
+	{
+		if (BudLateralMeristem.IsValid() && BudLateralMeristem.IsValidIndex(Index))
+		{
+			return BudLateralMeristem[Index];
+		}
+
+		static const TArray<float> EmptyArray;
+		return EmptyArray;
+	}
+
 	float FPointFacade::GetBranchGradient(int32 Index) const
 	{
 		if (BranchGradients.IsValid() && BranchGradients.IsValidIndex(Index))
diff --git a/Engine/Plugins/Experimental/ProceduralVegetationEditor/Source/ProceduralVegetation/Private/Facades/PVTreeFacade.cpp b/Engine/Plugins/Experimental/ProceduralVegetationEditor/Source/ProceduralVegetation/Private/Facades/PVTreeFacade.cpp
--- a/Engine/Plugins/Experimental/ProceduralVegetationEditor/Source/ProceduralVegetation/Private/Facades/PVTreeFacade.cpp
+++ b/Engine/Plugins/Experimental/ProceduralVegetationEditor/Source/ProceduralVegetation/Private/Facades/PVTreeFacade.cpp
@@ -1,9 +1,12 @@
 ﻿// Copyright Epic Games, Inc. All Rights Reserved.
 
-#include "PVTreeFacade.h"
+#include "Facades/PVTreeFacade.h"
+
 #include "Facades/PVBranchFacade.h"
 #include "Facades/PVFoliageFacade.h"
+#include "Facades/PVMetaInfoFacade.h"
 #include "Facades/PVPointFacade.h"
+
 #include "GeometryCollection/ManagedArrayAccessor.h"
 
 namespace PV::Facades
@@ -148,6 +151,259 @@ namespace PV::Facades
 		return BudDevelopment[0];
 	}
 
+	TArray<float> FTreeFacade::GetVisualizationValues(const FManagedArrayCollection& Collection, ESkeletonVisualizationModes VisualizationMode)
+	{
+		const FBranchFacade BranchFacade(Collection);
+		const FPointFacade PointFacade(Collection);
+		const FMetaInfoFacade DetailsFacade(Collection);
+
+		TArray<float> VisualizationValues;
+		VisualizationValues.SetNum(PointFacade.GetElementCount());
+
+		switch (VisualizationMode)
+		{
+		case ESkeletonVisualizationModes::None:
+			{
+				VisualizationValues.Empty();
+				break;
+			}
+		case ESkeletonVisualizationModes::GravitationalStress:
+			{
+				const TArray<float>& AbscissionSenescenseArray = DetailsFacade.GetAbscissionSenescense();
+				const float WeightAbscissionThreshold = 1 - AbscissionSenescenseArray[0];
+
+				for (int BranchIndex = 0; BranchIndex < BranchFacade.GetElementCount(); ++BranchIndex)
+				{
+					const TArray<int32>& BranchPoints = BranchFacade.GetPoints(BranchIndex);
+					const FVector3f FirstBudDirection = PointFacade.GetBudDirection(BranchPoints[0])[0];
+					const FVector3f SecondBudDirection = PointFacade.GetBudDirection(BranchPoints[1])[0];
+					const float BirthParentChildDot = PointFacade.GetBudLateralMeristem(BranchPoints[1])[4];
+
+					const float GravityParentChildDot = FVector3f::DotProduct(SecondBudDirection, FirstBudDirection);
+					const float GravitationalPressure = FMath::GetMappedRangeValueUnclamped(
+						FVector2f{0.0f, WeightAbscissionThreshold},
+						FVector2f{0.0f, 1.0f},
+						BirthParentChildDot - GravityParentChildDot
+					);
+					for (const int32 BranchPointIndex : BranchPoints)
+					{
+						VisualizationValues[BranchPointIndex] = GravitationalPressure;
+					}
+				}
+				break;
+			}
+		case ESkeletonVisualizationModes::CellDensity:
+			{
+				const TArray<float>& LateralElongationArray = DetailsFacade.GetLateralElongation();
+
+				const float CellDensity = (1 - LateralElongationArray[1]) * (LateralElongationArray[0] * 100.0f);
+
+				for (int BranchIndex = 0; BranchIndex < BranchFacade.GetElementCount(); ++BranchIndex)
+				{
+					const TArray<int32>& BranchPoints = BranchFacade.GetPoints(BranchIndex);
+					for (const int32 BranchPointIndex : BranchPoints)
+					{
+						const float CurrentLateralMeristem = PointFacade.GetBudLateralMeristem(BranchPointIndex)[0];
+						VisualizationValues[BranchPointIndex] = FMath::GetMappedRangeValueClamped(
+							FVector2f{0.0f, CellDensity},
+							FVector2f{0.0f, 1.0f},
+							CurrentLateralMeristem
+						);
+					}
+				}
+				break;
+			}
+		case ESkeletonVisualizationModes::CellWeight:
+			{
+				const TArray<float>& LateralElongationArray = DetailsFacade.GetLateralElongation();
+
+				const float CellWeight = LateralElongationArray[2];
+
+				for (int BranchIndex = 0; BranchIndex < BranchFacade.GetElementCount(); ++BranchIndex)
+				{
+					const TArray<int32>& BranchPoints = BranchFacade.GetPoints(BranchIndex);
+					for (const int32 BranchPointIndex : BranchPoints)
+					{
+						const float CurrentLateralWeight = PointFacade.GetBudLateralMeristem(BranchPointIndex)[3] * CellWeight;
+						VisualizationValues[BranchPointIndex] = FMath::Clamp(CurrentLateralWeight, 0.0f, 1.0f);
+					}
+				}
+				break;
+			}
+		case ESkeletonVisualizationModes::LightSenescence:
+			{
+				const TArray<float>& AbscissionSenescenseArray = DetailsFacade.GetAbscissionSenescense();
+				const float LightAbscissionMin = AbscissionSenescenseArray[6];
+				const float LightAbscissionMax = AbscissionSenescenseArray[7];
+				for (int BranchIndex = 0; BranchIndex < BranchFacade.GetElementCount(); ++BranchIndex)
+				{
+					const TArray<int32>& BranchPoints = BranchFacade.GetPoints(BranchIndex);
+					const int32 BranchNumber = BranchFacade.GetBranchNumber(BranchIndex);
+					FRandomStream RandomStream(BranchNumber);
+
+					const int32 LightSenescence = PointFacade.GetBudDevelopment(BranchPoints[1])[4];
+					const float LightAbscissionThreshold = FMath::RoundToInt32(RandomStream.FRandRange(LightAbscissionMin, LightAbscissionMax));
+					const float LightSenescenceValue = FMath::GetMappedRangeValueUnclamped(
+						FVector2f{0.0f, LightAbscissionThreshold},
+						FVector2f{0.0f, 1.0f},
+						LightSenescence
+					);
+					for (const int32 BranchPointIndex : BranchPoints)
+					{
+						VisualizationValues[BranchPointIndex] = LightSenescenceValue;
+					}
+				}
+				break;
+			}
+		case ESkeletonVisualizationModes::BranchLight:
+			{
+				for (int BranchIndex = 0; BranchIndex < BranchFacade.GetElementCount(); ++BranchIndex)
+				{
+					const TArray<int32>& BranchPoints = BranchFacade.GetPoints(BranchIndex);
+					const float LightDetected = PointFacade.GetBudLightDetected(BranchPoints[1])[2];
+					for (const int32 BranchPointIndex : BranchPoints)
+					{
+						VisualizationValues[BranchPointIndex] = LightDetected;
+					}
+				}
+				break;
+			}
+		case ESkeletonVisualizationModes::LightAbscissionRetention:
+			{
+				const TArray<float>& AbscissionSenescenseArray = DetailsFacade.GetAbscissionSenescense();
+				const TArray<float>& LateralElongationArray = DetailsFacade.GetLateralElongation();
+
+				const float LightRetention = AbscissionSenescenseArray[8];
+				const float LateralElongation = LateralElongationArray[0];
+				for (int BranchIndex = 0; BranchIndex < BranchFacade.GetElementCount(); ++BranchIndex)
+				{
+					const TArray<int32>& BranchPoints = BranchFacade.GetPoints(BranchIndex);
+					for (const int32 BranchPointIndex : BranchPoints)
+					{
+						const float LateralMeristem = PointFacade.GetBudLateralMeristem(BranchPointIndex)[0];
+						const float LightRetentionThreshold = (LateralElongation * 100.0f) * (1 - LightRetention);
+						const float LightRetentionValue = FMath::GetMappedRangeValueUnclamped(
+							FVector2f{0.0f, LightRetentionThreshold},
+							FVector2f{0.0f, 1.0f},
+							LateralMeristem
+						);
+						VisualizationValues[BranchPointIndex] = LightRetentionValue;
+					}
+				}
+				break;
+			}
+		case ESkeletonVisualizationModes::AgeSenescence:
+			{
+				const TArray<float>& AbscissionSenescenseArray = DetailsFacade.GetAbscissionSenescense();
+
+				const float AgeAbscissionMin = AbscissionSenescenseArray[2];
+				const float AgeAbscissionMax = AbscissionSenescenseArray[3];
+				for (int BranchIndex = 0; BranchIndex < BranchFacade.GetElementCount(); ++BranchIndex)
+				{
+					const TArray<int32>& BranchPoints = BranchFacade.GetPoints(BranchIndex);
+					const int32 BranchNumber = BranchFacade.GetBranchNumber(BranchIndex);
+					FRandomStream RandomStream(BranchNumber);
+
+					const int32 AgeSenescence = PointFacade.GetBudDevelopment(BranchPoints[1])[3];
+					const float AgeAbscissionThreshold = FMath::RoundToInt32(RandomStream.FRandRange(AgeAbscissionMin, AgeAbscissionMax));
+					const float AgeSenescenceValue = FMath::GetMappedRangeValueUnclamped(
+						FVector2f{0.0f, AgeAbscissionThreshold},
+						FVector2f{0.0f, 1.0f},
+						AgeSenescence
+					);
+					for (const int32 BranchPointIndex : BranchPoints)
+					{
+						VisualizationValues[BranchPointIndex] = AgeSenescenceValue;
+					}
+				}
+				break;
+			}
+		case ESkeletonVisualizationModes::AgeAbscissionRetention:
+			{
+				const TArray<float>& AbscissionSenescenseArray = DetailsFacade.GetAbscissionSenescense();
+				const TArray<float>& LateralElongationArray = DetailsFacade.GetLateralElongation();
+
+				const float AgeRetention = AbscissionSenescenseArray[4];
+				const float LateralElongation = LateralElongationArray[0];
+				for (int BranchIndex = 0; BranchIndex < BranchFacade.GetElementCount(); ++BranchIndex)
+				{
+					const TArray<int32>& BranchPoints = BranchFacade.GetPoints(BranchIndex);
+					for (const int32 BranchPointIndex : BranchPoints)
+					{
+						const float LateralMeristem = PointFacade.GetBudLateralMeristem(BranchPointIndex)[0];
+						const float AgeRetentionThreshold = (LateralElongation * 100.0f) * (1 - AgeRetention);
+						const float AgeRetentionValue = FMath::GetMappedRangeValueUnclamped(
+							FVector2f{0.0f, AgeRetentionThreshold},
+							FVector2f{0.0f, 1.0f},
+							LateralMeristem
+						);
+						VisualizationValues[BranchPointIndex] = AgeRetentionValue;
+					}
+				}
+				break;
+			}
+		case ESkeletonVisualizationModes::Apical:
+			{
+				for (int32 Index = 0; Index < PointFacade.GetElementCount(); ++Index)
+				{
+					VisualizationValues[Index] = PointFacade.GetBudHormoneLevels(Index)[0];
+				}
+				break;
+			}
+		case ESkeletonVisualizationModes::Auxin:
+			{
+				for (int32 Index = 0; Index < PointFacade.GetElementCount(); ++Index)
+				{
+					VisualizationValues[Index] = PointFacade.GetBudHormoneLevels(Index)[1];
+				}
+				break;
+			}
+		case ESkeletonVisualizationModes::Radical:
+			{
+				for (int32 Index = 0; Index < PointFacade.GetElementCount(); ++Index)
+				{
+					VisualizationValues[Index] = PointFacade.GetBudHormoneLevels(Index)[3];
+				}
+				break;
+			}
+		case ESkeletonVisualizationModes::Ethylene:
+			{
+				for (int32 Index = 0; Index < PointFacade.GetElementCount(); ++Index)
+				{
+					VisualizationValues[Index] = PointFacade.GetBudHormoneLevels(Index)[4];
+				}
+				break;
+			}
+		case ESkeletonVisualizationModes::AuxinRetention:
+			{
+				const TArray<float>& LeafGrowthArray = DetailsFacade.GetLeafGrowth();
+				const float AuxinRetention = LeafGrowthArray[5];
+
+				for (int32 Index = 0; Index < PointFacade.GetElementCount(); ++Index)
+				{
+					const float EthyleneLevel = FMath::Clamp(PointFacade.GetBudHormoneLevels(Index)[4], 0.0f, 1.0f);
+					const float ApicalAuxin = 1.0f - FMath::Lerp(EthyleneLevel, 0, PointFacade.GetBudHormoneLevels(Index)[0]);
+					VisualizationValues[Index] = FMath::Lerp(0, ApicalAuxin, AuxinRetention);
+				}
+				break;
+			}
+		case ESkeletonVisualizationModes::Cytokinin:
+			{
+				for (int32 Index = 0; Index < PointFacade.GetElementCount(); ++Index)
+				{
+					VisualizationValues[Index] = PointFacade.GetBudHormoneLevels(Index)[5];
+				}
+				break;
+			}
+		default:
+			{
+				VisualizationValues.Empty();
+			}
+		}
+		
+		return VisualizationValues;
+	}
+
 	void FTreeFacade::GatherChildBranches(const FBranchFacade& InBranchFacade, const int InParentIndex, TArray<int32>& OutBranchesToRemove)
 	{
 		const TArray<int32>& BranchChildren = InBranchFacade.GetChildren(InParentIndex);
diff --git a/Engine/Plugins/Experimental/ProceduralVegetationEditor/Source/ProceduralVegetation/Private/Facades/PVTreeFacade.h b/Engine/Plugins/Experimental/ProceduralVegetationEditor/Source/ProceduralVegetation/Private/Facades/PVTreeFacade.h
deleted file mode 100644
diff --git a/Engine/Plugins/Experimental/ProceduralVegetationEditor/Source/ProceduralVegetation/Private/Helpers/PVUtilities.cpp b/Engine/Plugins/Experimental/ProceduralVegetationEditor/Source/ProceduralVegetation/Private/Helpers/PVUtilities.cpp
--- a/Engine/Plugins/Experimental/ProceduralVegetationEditor/Source/ProceduralVegetation/Private/Helpers/PVUtilities.cpp
+++ b/Engine/Plugins/Experimental/ProceduralVegetationEditor/Source/ProceduralVegetation/Private/Helpers/PVUtilities.cpp
@@ -280,13 +280,7 @@ namespace PV::Utilities
 	{
 		Alpha = FMath::Clamp(Alpha, 0.f, 1.f);
 		const float HueDeg = 240.f * Alpha;
-		const FRandomStream Rand(Alpha * 255);
-		const float V = FMath::GetMappedRangeValueUnclamped(
-			FVector2f{0.0f, 1.0f},
-			FVector2f{0.25f, 1.0f},
-			Rand.FRand()
-		);
-		return FLinearColor(HueDeg, 0.9, V).HSVToLinearRGB();
+		return FLinearColor(HueDeg, 1, 1).HSVToLinearRGB();
 	};
  
 	bool IsValidPVData(const FManagedArrayCollection& Collection)
diff --git a/Engine/Plugins/Experimental/ProceduralVegetationEditor/Source/ProceduralVegetation/Public/DataTypes/PVData.h b/Engine/Plugins/Experimental/ProceduralVegetationEditor/Source/ProceduralVegetation/Public/DataTypes/PVData.h
--- a/Engine/Plugins/Experimental/ProceduralVegetationEditor/Source/ProceduralVegetation/Public/DataTypes/PVData.h
+++ b/Engine/Plugins/Experimental/ProceduralVegetationEditor/Source/ProceduralVegetation/Public/DataTypes/PVData.h
@@ -46,7 +46,8 @@ enum class EPVDebugValueVisualizationMode
 	Sphere,
 	Vector,
 	Point,
-	Curve
+	Curve,
+	ColorRamp
 };
 
 USTRUCT()
diff --git a/Engine/Plugins/Experimental/ProceduralVegetationEditor/Source/ProceduralVegetation/Public/Facades/PVAttributesNames.h b/Engine/Plugins/Experimental/ProceduralVegetationEditor/Source/ProceduralVegetation/Public/Facades/PVAttributesNames.h
--- a/Engine/Plugins/Experimental/ProceduralVegetationEditor/Source/ProceduralVegetation/Public/Facades/PVAttributesNames.h
+++ b/Engine/Plugins/Experimental/ProceduralVegetationEditor/Source/ProceduralVegetation/Public/Facades/PVAttributesNames.h
@@ -42,6 +42,7 @@ namespace PV::AttributeNames
 	inline static const FName LengthFromSeed = FName(TEXT("lengthFromSeed"));
 	inline static const FName BudLightDetected = FName(TEXT("budLightDetected"));
 	inline static const FName BudDevelopment = FName(TEXT("budDevelopment"));
+	inline static const FName BudLateralMeristem = FName(TEXT("budLateralMeristem"));
 	inline static const FName TextureCoordV = FName(TEXT("uv_v"));
 	inline static const FName TextureCoordUOffset = FName(TEXT("uv_uOffset"));
 	inline static const FName URange = FName(TEXT("uv_uRange"));
@@ -53,7 +54,6 @@ namespace PV::AttributeNames
 	inline static const FName BudHormoneLevels = FName(TEXT("budHormoneLevels"));
 	inline static const FName BudDirection = FName(TEXT("budDirection"));
 	inline static const FName BudStatus = FName(TEXT("budStatus"));
-	inline static const FName BudLateralMeristem = FName(TEXT("budLateralMeristem"));
 	inline static const FName BranchGradient = FName(TEXT("branchGradient"));
 	inline static const FName PlantGradient = FName(TEXT("plantGradient"));
 	inline static const FName NjordPixelIndex = FName(TEXT("njord_pixelIdx"));
@@ -75,6 +75,9 @@ namespace PV::AttributeNames
 	inline static const FName LeafPhyllotaxy = FName(TEXT("phyllotaxyLeaf"));
 	inline static const FName FoliagePath = FName(TEXT("FoliagePath"));
 	inline static const FName Guid = FName(TEXT("Guid"));
+	inline static const FName LeafGrowth = FName(TEXT("leafGrowth"));
+	inline static const FName AbscissionSenescense = FName(TEXT("abscissionSenescense"));
+	inline static const FName LateralElongation = FName(TEXT("lateralElongation"));
 
 	//PlantProfiles Attributes
 	inline static const FName ProfilePoints = FName(TEXT("ProfilePoints"));
diff --git a/Engine/Plugins/Experimental/ProceduralVegetationEditor/Source/ProceduralVegetation/Public/Facades/PVMetaInfoFacade.h b/Engine/Plugins/Experimental/ProceduralVegetationEditor/Source/ProceduralVegetation/Public/Facades/PVMetaInfoFacade.h
--- a/Engine/Plugins/Experimental/ProceduralVegetationEditor/Source/ProceduralVegetation/Public/Facades/PVMetaInfoFacade.h
+++ b/Engine/Plugins/Experimental/ProceduralVegetationEditor/Source/ProceduralVegetation/Public/Facades/PVMetaInfoFacade.h
@@ -18,6 +18,10 @@ namespace PV::Facades
 
 		void CreateGuid(const FString& InPath);
 		FGuid GetGuid() const;
+		
+		const TArray<float>& GetLeafGrowth() const;
+		const TArray<float>& GetAbscissionSenescense() const;
+		const TArray<float>& GetLateralElongation() const;
 
 	protected:
 		void DefineSchema(const int32 InitialSize = 0);
@@ -26,5 +30,8 @@ namespace PV::Facades
 		FManagedArrayCollection* Collection = nullptr;
 		
 		TManagedArrayAccessor<FGuid> GuidAttribute;
+		TManagedArrayAccessor<TArray<float>> LeafGrowthAttribute;
+		TManagedArrayAccessor<TArray<float>> AbscissionSenescenseAttribute;
+		TManagedArrayAccessor<TArray<float>> LateralElongationAttribute;
 	};
 }
diff --git a/Engine/Plugins/Experimental/ProceduralVegetationEditor/Source/ProceduralVegetation/Public/Facades/PVPointFacade.h b/Engine/Plugins/Experimental/ProceduralVegetationEditor/Source/ProceduralVegetation/Public/Facades/PVPointFacade.h
--- a/Engine/Plugins/Experimental/ProceduralVegetationEditor/Source/ProceduralVegetation/Public/Facades/PVPointFacade.h
+++ b/Engine/Plugins/Experimental/ProceduralVegetationEditor/Source/ProceduralVegetation/Public/Facades/PVPointFacade.h
@@ -81,6 +81,8 @@ Accessor.ModifyAt(i, InputArray[i]); \
 
 		void SetBudHormoneLevels(int32 Index, TArray<float> InBudHormoneLevels);
 
+		const TArray<float>& GetBudLateralMeristem(int32 Index) const;
+
 		float GetBranchGradient(int32 Index) const;
 
 		void SetBranchGradient(int32 Index, float InBranchGradient);
@@ -88,7 +90,7 @@ Accessor.ModifyAt(i, InputArray[i]); \
 		float GetPlantGradient(int32 Index) const;
 
 		void SetPlantGradient(int32 Index, float InPlantGradient);
-
+	
 		float GetNjordPixelIndex(int32 Index) const;
 
 		void SetNjordPixelIndex(int32 Index, float InNjordPixelIndex);
diff --git a/Engine/Plugins/Experimental/ProceduralVegetationEditor/Source/ProceduralVegetation/Public/Facades/PVTreeFacade.h b/Engine/Plugins/Experimental/ProceduralVegetationEditor/Source/ProceduralVegetation/Public/Facades/PVTreeFacade.h
new file mode 100644
--- /dev/null
+++ b/Engine/Plugins/Experimental/ProceduralVegetationEditor/Source/ProceduralVegetation/Public/Facades/PVTreeFacade.h
@@ -0,0 +1,63 @@
+// Copyright Epic Games, Inc. All Rights Reserved.
+
+#pragma once
+
+#include "CoreMinimal.h"
+
+#include "UObject/ObjectMacros.h"
+
+#include "PVTreeFacade.generated.h"
+
+struct FManagedArrayCollection;
+
+UENUM()
+enum class ESkeletonVisualizationModes : uint8
+{
+	None UMETA(Category=""),
+	GravitationalStress UMETA(Category="Gravity"),
+	CellDensity UMETA(Category="Gravity"),
+	CellWeight UMETA(Category="Gravity"),
+	LightSenescence UMETA(Category="Light"),
+	BranchLight UMETA(Category="Light"),
+	LightAbscissionRetention UMETA(Category="Light"),
+	AgeSenescence UMETA(Category="Age"),
+	AgeAbscissionRetention UMETA(Category="Age"),
+	Apical UMETA(Category="Auxin"),
+	Auxin UMETA(Category="Auxin"),
+	Radical UMETA(Category="Auxin"),
+	Ethylene UMETA(Category="Ethlyene"),
+	AuxinRetention UMETA(Category="Ethlyene"),
+	Cytokinin UMETA(Category="Cytokinin")
+};
+
+namespace PV::Facades
+{
+	struct FRemoveEntriesResult
+	{
+		TMap<int32, int32> PointsOldIDsToNewIDs;
+		TMap<int32, int32> BranchesOldIDsToNewIDs;
+		TMap<int32, int32> FoliageInstancesOldIDsToNewIDs;
+	};
+
+	class FBranchFacade;
+
+	/**
+	 * FTreeFacade is used to remove branches, points, and foliage instances from a tree collection.
+	 * It does not validate whether the input data is consistent i.e. the points actually belong to the
+	 * same branches that are asking to be removed etc.
+	 */
+	class PROCEDURALVEGETATION_API FTreeFacade
+	{
+	public:
+		static FRemoveEntriesResult RemoveEntriesAndReIndexAttributes(FManagedArrayCollection& OutCollection, TArray<bool>& PointsToRemove,
+		                                                              TArray<bool>& BranchesToRemove, TArray<bool>& FoliageInstancesToRemove);
+
+		static void RemoveBranches(const FBranchFacade& InBranchFacade, TArray<int>& InBranchesToRemove, FManagedArrayCollection& OutCollection);
+		
+		static int32 GetGenerationNumberFromSourceBud(const FManagedArrayCollection& Collection, int32 BranchIndex);
+		
+		static TArray<float> GetVisualizationValues(const FManagedArrayCollection& Collection, ESkeletonVisualizationModes VisualizationMode);
+	private:
+		static void GatherChildBranches(const FBranchFacade& InBranchFacade, const int InParentIndex, TArray<int32>& OutBranchesToRemove);
+	};
+}
diff --git a/Engine/Plugins/Experimental/ProceduralVegetationEditor/Source/ProceduralVegetationEditor/Private/SPVEditorViewport.cpp b/Engine/Plugins/Experimental/ProceduralVegetationEditor/Source/ProceduralVegetationEditor/Private/SPVEditorViewport.cpp
--- a/Engine/Plugins/Experimental/ProceduralVegetationEditor/Source/ProceduralVegetationEditor/Private/SPVEditorViewport.cpp
+++ b/Engine/Plugins/Experimental/ProceduralVegetationEditor/Source/ProceduralVegetationEditor/Private/SPVEditorViewport.cpp
@@ -33,7 +33,9 @@
 
 #include "Visualizations/PVScaleVisualizationComponent.h"
 
+#include "Widgets/SPVGradientBox.h"
 #include "Widgets/Input/SNumericEntryBox.h"
+#include "Widgets/Input/SSlider.h"
 #include "Widgets/Text/SRichTextBlock.h"
 #include "Widgets/Views/SListView.h"
 #include "Widgets/Views/STableRow.h"
@@ -297,6 +299,66 @@ void SPVEditorViewport::PopulateViewportOverlays(TSharedRef<SOverlay> Overlay)
 				]
 			]
 		];
+	
+	Overlay
+		->AddSlot()
+		.VAlign(VAlign_Bottom)
+		.HAlign(HAlign_Right)
+		.Padding(6.0f, 25.f)
+		[
+			SNew(SBorder)
+			.BorderImage(FAppStyle::Get().GetBrush("FloatingBorder"))
+			.Visibility_Lambda([this]
+				{
+					return SkeletonVisualizerComponent && SkeletonVisualizerComponent->GetVisualizationMode() != ESkeletonVisualizationModes::None
+						? EVisibility::Visible
+						: EVisibility::Collapsed;
+				})
+			.Padding(6.0f, 4.0f)
+			[
+				SNew(SVerticalBox)
+				+ SVerticalBox::Slot()
+				.HAlign(HAlign_Center)
+				[
+					SNew(STextBlock)
+					.Justification(ETextJustify::Center)
+					.Text_Lambda([this]()
+						{
+							if (SkeletonVisualizerComponent)
+							{
+								return StaticEnum<ESkeletonVisualizationModes>()->GetDisplayNameTextByIndex(
+									static_cast<int32>(SkeletonVisualizerComponent->GetVisualizationMode())
+								);
+							}
+							return FText::GetEmpty();
+						})
+				]
+				+ SVerticalBox::Slot()
+				[
+					SNew(SPVGradientBox)
+					.BoxSize(160.0f)
+				]
+				+ SVerticalBox::Slot()
+				[
+					SNew(SHorizontalBox)
+					+ SHorizontalBox::Slot()
+					.HAlign(HAlign_Left)
+					[
+						SNew(STextBlock)
+						.Justification(ETextJustify::Left)
+						.Text(FText::FromString(TEXT("Min")))
+					]
+					+ SHorizontalBox::Slot()
+					.HAlign(HAlign_Right)
+					[
+						SNew(STextBlock)
+						.Justification(ETextJustify::Right)
+
+						.Text(FText::FromString(TEXT("Max")))
+					]
+				]
+			]
+		];
 }
 
 void SPVEditorViewport::InitVisualizationScene()
@@ -682,6 +744,11 @@ FToolMenuEntry SPVEditorViewport::CreateVisualizationModeToolbarMenu()
 
 				LabelAttribute = TAttribute<FText>::CreateSPLambda(this, [this]
 					{
+						if (SkeletonVisualizerComponent && SkeletonVisualizerComponent->GetVisualizationMode() != ESkeletonVisualizationModes::None)
+						{
+							const UEnum* EnumPtr = StaticEnum<ESkeletonVisualizationModes>();
+							return EnumPtr->GetDisplayNameTextByIndex(static_cast<int32>(SkeletonVisualizerComponent->GetVisualizationMode()));
+						}
 						const UEnum* EnumPtr = StaticEnum<EPVVisualizationMode>();
 						return EnumPtr->GetDisplayNameTextByIndex(static_cast<int>(CurrentVisualizationMode));
 					}
@@ -705,7 +772,7 @@ FToolMenuEntry SPVEditorViewport::CreateVisualizationModeToolbarMenu()
 								for (int32 i = 0; i < EnumPtr->NumEnums() - 1; i++) // last entry is _MAX or hidden
 								{
 									FText ModeName = EnumPtr->GetDisplayNameTextByIndex(i);
-									EPVVisualizationMode Mode = static_cast<EPVVisualizationMode>(EnumPtr->GetValueByIndex(i));
+									const EPVVisualizationMode Mode = static_cast<EPVVisualizationMode>(EnumPtr->GetValueByIndex(i));
 
 									RenderModeSelectionSection.AddMenuEntry(
 										NAME_None,
@@ -744,6 +811,61 @@ FToolMenuEntry SPVEditorViewport::CreateVisualizationModeToolbarMenu()
 									);
 								}
 							}
+
+							RenderModeSelectionSection.AddSubMenu(
+								TEXT("DebugModeSection"),
+								LOCTEXT("DebugModeSelectionSubSectionLabel", "Debug Visualizations"),
+								LOCTEXT("DebugModeSelectionSubSectionLabel", "Toggle Debug Visualizations"),
+								FNewToolMenuDelegate::CreateSPLambda(this,
+									[this](UToolMenu* SubMenu)
+										{
+											UEnum* EnumPtr = StaticEnum<ESkeletonVisualizationModes>();
+											if (EnumPtr)
+											{
+												for (int32 i = 0; i < EnumPtr->NumEnums() - 1; i++)
+												{
+													const FText ModeName = EnumPtr->GetDisplayNameTextByIndex(i);
+													const FString Category = EnumPtr->GetMetaData(TEXT("Category"), i);
+													const ESkeletonVisualizationModes Mode = static_cast<ESkeletonVisualizationModes>(EnumPtr->GetValueByIndex(i));
+
+													FToolMenuSection& DebugModeSelectionSection = SubMenu->FindOrAddSection(*Category, FText::FromString(Category));
+
+													DebugModeSelectionSection.AddMenuEntry(
+														NAME_None,
+														ModeName,
+														FText(),
+														FSlateIcon(),
+														FUIAction(
+															FExecuteAction::CreateSPLambda(
+																this, [this, Mode]
+																	{
+																		if (SkeletonVisualizerComponent)
+																		{
+																			SkeletonVisualizerComponent->SetVisualizationMode(Mode);
+																		}
+																	}
+															),
+															FCanExecuteAction::CreateSPLambda(
+																this, [this, Mode]
+																	{
+																		return SkeletonVisualizerComponent != nullptr;
+																	}
+															),
+															FIsActionChecked::CreateSPLambda(
+																this, [this, Mode]
+																	{
+																		return SkeletonVisualizerComponent && SkeletonVisualizerComponent->
+																			GetVisualizationMode() == Mode;
+																	}
+															)
+														),
+														EUserInterfaceActionType::RadioButton
+													);
+												}
+											}
+										}
+								)
+							);
 						}
 					),
 					false,
@@ -830,4 +952,9 @@ void SPVEditorViewport::AddInstancedFoliageComponent(TObjectPtr<UMeshComponent>
 	Component->SetVisibility(FoliageVisibility->FindOrAdd(Component->GetName(), true), true);
 }
 
+void SPVEditorViewport::SetSkeletonVisualizerComponent(TObjectPtr<UPVSkeletonVisualizerComponent> InSkeletonVisualizerComponent)
+{
+	SkeletonVisualizerComponent = InSkeletonVisualizerComponent;
+}
+
 #undef LOCTEXT_NAMESPACE
diff --git a/Engine/Plugins/Experimental/ProceduralVegetationEditor/Source/ProceduralVegetationEditor/Private/SPVEditorViewport.h b/Engine/Plugins/Experimental/ProceduralVegetationEditor/Source/ProceduralVegetationEditor/Private/SPVEditorViewport.h
--- a/Engine/Plugins/Experimental/ProceduralVegetationEditor/Source/ProceduralVegetationEditor/Private/SPVEditorViewport.h
+++ b/Engine/Plugins/Experimental/ProceduralVegetationEditor/Source/ProceduralVegetationEditor/Private/SPVEditorViewport.h
@@ -4,6 +4,8 @@
 
 #include "CoreMinimal.h"
 
+#include "Visualizations/PVSkeletonVisualizerComponent.h"
+
 #include "Widgets/AssetEditorViewport/SPCGEditorViewport.h"
 #include "Widgets/Text/SRichTextBlock.h"
 
@@ -74,8 +76,7 @@ public:
 	void PopulateLegendOverlayText(const TArrayView<TPair<FColor, FText>> LegendItems);
 
 	void AddInstancedFoliageComponent(TObjectPtr<UMeshComponent> Component);
-
-	
+	void SetSkeletonVisualizerComponent(TObjectPtr<UPVSkeletonVisualizerComponent> InSkeletonVisualizerComponent);
 
 protected:
 	//~ Begin SEditorViewport Interface
@@ -128,6 +129,7 @@ protected:
 private:
 	TObjectPtr<UStaticMeshComponent> MannequinComponent = nullptr;
 	TObjectPtr<UPVScaleVisualizationComponent> ScaleVisualizationComponent = nullptr;
+	TObjectPtr<UPVSkeletonVisualizerComponent> SkeletonVisualizerComponent = nullptr;
 
 	TSharedPtr<TArray<TObjectPtr<UMeshComponent>>> FoliageComponents;
 	TSharedPtr<TMap<FString, bool>> FoliageVisibility;
@@ -136,6 +138,7 @@ private:
 	TSharedPtr<SRichTextBlock> OverlayText;
 	TSharedPtr<SRichTextBlock> StatsOverlayText;
 	TArray<TSharedPtr<FText>> LegendData;
+	ESkeletonVisualizationModes DebugVisualizationMode = ESkeletonVisualizationModes::None;
 
 	TSharedPtr<FSlateBrush> PreviewNodeBackgroundBrush;
 };
diff --git a/Engine/Plugins/Experimental/ProceduralVegetationEditor/Source/ProceduralVegetationEditor/Private/Visualizations/DebugVisualization/PVDebugVisualizationBase.cpp b/Engine/Plugins/Experimental/ProceduralVegetationEditor/Source/ProceduralVegetationEditor/Private/Visualizations/DebugVisualization/PVDebugVisualizationBase.cpp
--- a/Engine/Plugins/Experimental/ProceduralVegetationEditor/Source/ProceduralVegetationEditor/Private/Visualizations/DebugVisualization/PVDebugVisualizationBase.cpp
+++ b/Engine/Plugins/Experimental/ProceduralVegetationEditor/Source/ProceduralVegetationEditor/Private/Visualizations/DebugVisualization/PVDebugVisualizationBase.cpp
@@ -12,6 +12,8 @@
 
 #include "Engine/StaticMesh.h"
 
+#include "Materials/MaterialInstance.h"
+
 #include "UObject/Package.h"
 
 #include "Visualizations/PVLineBatchComponent.h"
@@ -615,6 +617,30 @@ void FPVDebugVisualizationBase::DrawAttributes(FVisualizerDrawContext& InContext
 			}
 			break;
 		}
+		case EPVDebugValueVisualizationMode::ColorRamp:
+		{
+			using namespace PVE::AttributeToScalar;
+
+			TArray<double> AttributeAsScalar = ConvertAttributeToScalar(Collection, AttributeName, GroupName);
+			if (AttributeAsScalar.Num() != NumElements)
+			{
+				break;
+			}
+
+			const double MaxValue = *Algo::MaxElement(AttributeAsScalar);
+			const double MinValue = *Algo::MinElement(AttributeAsScalar);
+			const double Range = MaxValue - MinValue;
+
+			for (int32 i = 0; i < NumElements; i++)
+			{
+				FVector3f OutPos;
+				float OutScale;
+				GetPivot(Collection, i, OutPos, OutScale);
+
+				AddDrawAsPoint(static_cast<FVector>(OutPos), OutScale, PV::Utilities::GetRandomHueColor(AttributeAsScalar[i] - MinValue) / Range);
+			}
+			break;
+		}
 		case EPVDebugValueVisualizationMode::Vector:
 			ensureMsgf(false, TEXT("Unexpected visualization mode"));
 			return;
@@ -624,29 +650,22 @@ void FPVDebugVisualizationBase::DrawAttributes(FVisualizerDrawContext& InContext
 	{
 		if (VisualizationSettings.bDrawPointAsMesh)
 		{
-			TArray<FTransform> PointTransforms;
-			PointTransforms.Reserve(PointsToDraw.Num());
-			for (const auto& It : PointsToDraw)
-			{
-				PointTransforms.Add(FTransform(FRotator::ZeroRotator, It.Key, FVector(It.Value * 0.005f)));
-			}
-
-			DrawAsMeshPoints(GetOrCreatePointComponent(InContext.SceneSetupParams), PointTransforms);
+			DrawAsMeshPoints(GetOrCreatePointComponent(InContext.SceneSetupParams), PointsToDraw);
 		}
 		else
 		{
-			DrawAsPoints(GetOrCreateLineComponent(InContext.SceneSetupParams), PointsToDraw, VisualizationSettings.Color);
+			DrawAsPoints(GetOrCreateLineComponent(InContext.SceneSetupParams), PointsToDraw);
 		}
 
 		PointsToDraw.Reset();
 	}
 }
 
-void FPVDebugVisualizationBase::AddDrawAsPoint(const FVector& InPos, float PointScale)
+void FPVDebugVisualizationBase::AddDrawAsPoint(const FVector& InPos, float PointScale, const FLinearColor& InColor)
 {
 	if (!FMath::IsNearlyZero(PointScale))
 	{
-		PointsToDraw.Add({ InPos, PointScale });
+		PointsToDraw.Add({ InPos, PointScale, InColor });
 	}
 }
 
@@ -678,10 +697,13 @@ void FPVDebugVisualizationBase::DrawAsSphere(UPVLineBatchComponent* LineBatchCom
 	LineBatchComponent->AddSphere(InPos, Radius, InColor);
 }
 
-void FPVDebugVisualizationBase::DrawAsPoints(UPVLineBatchComponent* LineBatchComponent, const TArray<TPair<FVector, float>>& InPoints, const FLinearColor& InColor)
+void FPVDebugVisualizationBase::DrawAsPoints(UPVLineBatchComponent* LineBatchComponent, const TArray<TTuple<FVector, float, FLinearColor>>& InPoints)
 {
 	check(LineBatchComponent);
-	LineBatchComponent->AddPoints(InPoints, InColor);
+	for (const auto& [Position,Size,Color] : InPoints)
+	{
+		LineBatchComponent->AddPoint(Position, Size * 10.0f, Color, SDPG_Foreground);
+	}
 }
 
 void FPVDebugVisualizationBase::DrawAsMeshPoints(UInstancedStaticMeshComponent* InstancedStaticMeshComponent, const TArray<FTransform>& InTransforms, const FLinearColor& InColor)
@@ -695,6 +717,20 @@ void FPVDebugVisualizationBase::DrawAsMeshPoints(UInstancedStaticMeshComponent*
 	InstancedStaticMeshComponent->AddInstances(InTransforms, false);
 }
 
+void FPVDebugVisualizationBase::DrawAsMeshPoints(UInstancedStaticMeshComponent* InstancedStaticMeshComponent,
+	TArray<TTuple<FVector, float, FLinearColor>> InPoints)
+{
+	check(InstancedStaticMeshComponent != nullptr);
+	for (const auto& [Position, Size, Color] : InPoints)
+	{
+		const int32 InstanceIndex = InstancedStaticMeshComponent->AddInstance(FTransform(FRotator::ZeroRotator, Position, FVector(Size * 0.005f)));
+		InstancedStaticMeshComponent->SetCustomDataValue(InstanceIndex, 0, Color.R, false);
+		InstancedStaticMeshComponent->SetCustomDataValue(InstanceIndex, 1, Color.G, false);
+		InstancedStaticMeshComponent->SetCustomDataValue(InstanceIndex, 2, Color.B, false);
+	}
+	InstancedStaticMeshComponent->MarkRenderStateDirty();
+}
+
 UPVLineBatchComponent* FPVDebugVisualizationBase::GetOrCreateLineComponent(FPCGSceneSetupParams& InOutParams)
 {
 	UPVLineBatchComponent* LineComponent = nullptr;
@@ -723,6 +759,7 @@ UInstancedStaticMeshComponent* FPVDebugVisualizationBase::GetOrCreatePointCompon
 		InstancedStaticMeshComponent = NewObject<UInstancedStaticMeshComponent>(GetTransientPackage(), NAME_None);
 		InstancedStaticMeshComponent->SetStaticMesh(PointSphere);
 		InstancedStaticMeshComponent->SetCollisionEnabled(ECollisionEnabled::NoCollision);
+		InstancedStaticMeshComponent->SetMaterial(0, LoadObject<UMaterialInterface>(nullptr, TEXT("/ProceduralVegetationEditor/Materials/FoliageAttachmentPointMaterial.FoliageAttachmentPointMaterial")));
 
 		AddToScene(InOutParams, InstancedStaticMeshComponent);
 	}
diff --git a/Engine/Plugins/Experimental/ProceduralVegetationEditor/Source/ProceduralVegetationEditor/Private/Visualizations/DebugVisualization/PVDebugVisualizationBase.h b/Engine/Plugins/Experimental/ProceduralVegetationEditor/Source/ProceduralVegetationEditor/Private/Visualizations/DebugVisualization/PVDebugVisualizationBase.h
--- a/Engine/Plugins/Experimental/ProceduralVegetationEditor/Source/ProceduralVegetationEditor/Private/Visualizations/DebugVisualization/PVDebugVisualizationBase.h
+++ b/Engine/Plugins/Experimental/ProceduralVegetationEditor/Source/ProceduralVegetationEditor/Private/Visualizations/DebugVisualization/PVDebugVisualizationBase.h
@@ -24,15 +24,16 @@ public:
 protected:
 	void DrawPivotPoints(const FVisualizerDrawContext& InContext);
 	void DrawAttributes(FVisualizerDrawContext& InContext);
-	void AddDrawAsPoint(const FVector& InPos, float PointScale);
+	void AddDrawAsPoint(const FVector& InPos, float PointScale, const FLinearColor& InColor = FLinearColor::White);
 	virtual TArray<FVector3f> GetPivotPositions(const FManagedArrayCollection& InCollection) = 0;
 	virtual void GetPivot(const FManagedArrayCollection& InCollection, const int InIndex, FVector3f& OutPos, float& OutScale) = 0;
 	
 	static void DrawAsText(FPCGSceneSetupParams& InOutParams, const FText& InTextToDraw, const FVector3f& InPos, float InScale, float TextSize, const FLinearColor& InColor = FLinearColor::White, const FVector3f TextOffset = FVector3f::Zero());
 	static void DrawAsLine(UPVLineBatchComponent* LineBatchComponent, const FVector& InStartPos, const FVector& InEndPos, const FLinearColor& InColor = FLinearColor::White);
 	static void DrawAsSphere(UPVLineBatchComponent* LineBatchComponent, const FVector& InPos, float Radius, const FLinearColor& InColor = FLinearColor::White);
-	static void DrawAsPoints(UPVLineBatchComponent* LineBatchComponent, const TArray<TPair<FVector, float>>& InPoints, const FLinearColor& InColor = FLinearColor::White);
+	static void DrawAsPoints(UPVLineBatchComponent* LineBatchComponent, const TArray<TTuple<FVector, float, FLinearColor>>& InPoints);
 	static void DrawAsMeshPoints(UInstancedStaticMeshComponent* InstancedStaticMeshComponent, const TArray<FTransform>& InTransforms, const FLinearColor& InColor = FLinearColor::White);
+	static void DrawAsMeshPoints(UInstancedStaticMeshComponent* InstancedStaticMeshComponent, TArray<TTuple<FVector, float, FLinearColor>> InPoints);
 	
 	static UPVLineBatchComponent* GetOrCreateLineComponent(FPCGSceneSetupParams& InOutParams);
 	static UInstancedStaticMeshComponent* GetOrCreatePointComponent(FPCGSceneSetupParams& InOutParams);
@@ -40,7 +41,7 @@ protected:
 	static void AddToScene(FPCGSceneSetupParams& InOutParams, UPrimitiveComponent* InComponent);
 	
 private:
-	TArray<TPair<FVector, float>> PointsToDraw;
+	TArray<TTuple<FVector, float, FLinearColor>> PointsToDraw;
 };
 typedef TSharedPtr<FPVDebugVisualizationBase> FPVDebugVisualizationPtr;
 
diff --git a/Engine/Plugins/Experimental/ProceduralVegetationEditor/Source/ProceduralVegetationEditor/Private/Visualizations/PVDataVisualization.cpp b/Engine/Plugins/Experimental/ProceduralVegetationEditor/Source/ProceduralVegetationEditor/Private/Visualizations/PVDataVisualization.cpp
--- a/Engine/Plugins/Experimental/ProceduralVegetationEditor/Source/ProceduralVegetationEditor/Private/Visualizations/PVDataVisualization.cpp
+++ b/Engine/Plugins/Experimental/ProceduralVegetationEditor/Source/ProceduralVegetationEditor/Private/Visualizations/PVDataVisualization.cpp
@@ -10,7 +10,9 @@
 #include "PlanarCut.h"
 #include "PVBoneComponent.h"
 #include "SPVEditorViewport.h"
- 
+
+#include "Algo/Unique.h"
+
 #include "Components/DynamicMeshComponent.h"
 #include "Components/InstancedStaticMeshComponent.h"
 #include "Components/LineBatchComponent.h"
@@ -193,6 +195,11 @@ void FPVDataVisualization::SkeletonRenderer(const FManagedArrayCollection& InCol
 		GetTransientPackage(), NAME_None, RF_Transient);
 	VisualizerComponent->SetCollection(&InCollection);
 	VisualizerComponent->UpdateBounds();
+
+	const TSharedPtr<SPVEditorViewport> ViewportWidget = StaticCastSharedPtr<SPVEditorViewport>(
+		InOutParams.EditorViewportClient->GetEditorViewportWidget());
+	
+	ViewportWidget->SetSkeletonVisualizerComponent(VisualizerComponent);
  
 	InOutParams.ManagedResources.Add(VisualizerComponent);
 	InOutParams.Scene->AddComponent(VisualizerComponent, FTransform::Identity);
diff --git a/Engine/Plugins/Experimental/ProceduralVegetationEditor/Source/ProceduralVegetationEditor/Private/Visualizations/PVLineBatchComponent.cpp b/Engine/Plugins/Experimental/ProceduralVegetationEditor/Source/ProceduralVegetationEditor/Private/Visualizations/PVLineBatchComponent.cpp
--- a/Engine/Plugins/Experimental/ProceduralVegetationEditor/Source/ProceduralVegetationEditor/Private/Visualizations/PVLineBatchComponent.cpp
+++ b/Engine/Plugins/Experimental/ProceduralVegetationEditor/Source/ProceduralVegetationEditor/Private/Visualizations/PVLineBatchComponent.cpp
@@ -43,12 +43,12 @@ void FPVLineSceneProxy::GetDynamicMeshElements(
 
 			if (Line.PointDrawSettings == EPointDrawSettings::Start || Line.PointDrawSettings == EPointDrawSettings::Both)
 			{
-				PDI->DrawPoint(Line.StartPos, FLinearColor::White, PointSize, Line.DepthPriorityGroup);
+				PDI->DrawPoint(Line.StartPos, Line.Color, PointSize, Line.DepthPriorityGroup);
 			}
 
 			if (Line.PointDrawSettings == EPointDrawSettings::End || Line.PointDrawSettings == EPointDrawSettings::Both)
 			{
-				PDI->DrawPoint(Line.EndPos, FLinearColor::White, PointSize, Line.DepthPriorityGroup);
+				PDI->DrawPoint(Line.EndPos, Line.Color, PointSize, Line.DepthPriorityGroup);
 			}
 		}
 	};
@@ -184,12 +184,12 @@ void UPVLineBatchComponent::AddPoint(const FVector& PointLocation, float InPoint
 	MarkRenderStateDirty();
 }
 
-void UPVLineBatchComponent::AddPoints(const TArray<TPair<FVector, float>>& InPoints, const FLinearColor& InColor, const ESceneDepthPriorityGroup InDepthPriorityGroup)
+void UPVLineBatchComponent::AddPoints(const TArray<TTuple<FVector, float, FLinearColor>>& InPoints, const ESceneDepthPriorityGroup InDepthPriorityGroup)
 {
-	for (const auto& It: InPoints)
+	for (const auto& [Position, Scale, Color] : InPoints)
 	{
-		Points.Emplace(It.Key, It.Value, InColor, InDepthPriorityGroup);
-		BBox += It.Key;
+		Points.Emplace(Position, Scale, Color, InDepthPriorityGroup);
+		BBox += Position;
 	}
 	MarkRenderStateDirty();
 }
diff --git a/Engine/Plugins/Experimental/ProceduralVegetationEditor/Source/ProceduralVegetationEditor/Private/Visualizations/PVLineBatchComponent.h b/Engine/Plugins/Experimental/ProceduralVegetationEditor/Source/ProceduralVegetationEditor/Private/Visualizations/PVLineBatchComponent.h
--- a/Engine/Plugins/Experimental/ProceduralVegetationEditor/Source/ProceduralVegetationEditor/Private/Visualizations/PVLineBatchComponent.h
+++ b/Engine/Plugins/Experimental/ProceduralVegetationEditor/Source/ProceduralVegetationEditor/Private/Visualizations/PVLineBatchComponent.h
@@ -83,7 +83,7 @@ public:
 	void AddPoint(const FVector& PointLocation, float PointSize, const FLinearColor& InColor = FLinearColor::White,
 		const ESceneDepthPriorityGroup InDepthPriorityGroup = ESceneDepthPriorityGroup::SDPG_World);
 
-	void AddPoints(const TArray<TPair<FVector, float>>& InPoints, const FLinearColor& InColor = FLinearColor::White,
+	void AddPoints(const TArray<TTuple<FVector, float, FLinearColor>>& InPoints,
 		const ESceneDepthPriorityGroup InDepthPriorityGroup = ESceneDepthPriorityGroup::SDPG_World);
 
 	void Flush();
diff --git a/Engine/Plugins/Experimental/ProceduralVegetationEditor/Source/ProceduralVegetationEditor/Private/Visualizations/PVSkeletonVisualizerComponent.cpp b/Engine/Plugins/Experimental/ProceduralVegetationEditor/Source/ProceduralVegetationEditor/Private/Visualizations/PVSkeletonVisualizerComponent.cpp
--- a/Engine/Plugins/Experimental/ProceduralVegetationEditor/Source/ProceduralVegetationEditor/Private/Visualizations/PVSkeletonVisualizerComponent.cpp
+++ b/Engine/Plugins/Experimental/ProceduralVegetationEditor/Source/ProceduralVegetationEditor/Private/Visualizations/PVSkeletonVisualizerComponent.cpp
@@ -1,11 +1,16 @@
 // Copyright Epic Games, Inc. All Rights Reserved.
 
-
 #include "PVSkeletonVisualizerComponent.h"
-#include "ProceduralVegetationEditorModule.h"
+
 #include "PrimitiveDrawInterface.h"
+#include "ProceduralVegetationEditorModule.h"
+
+#include "Facades/PVAttributesNames.h"
 #include "Facades/PVBranchFacade.h"
 #include "Facades/PVPointFacade.h"
+#include "Facades/PVTreeFacade.h"
+
+#include "Helpers/PVUtilities.h"
 
 UPVSkeletonVisualizerComponent::UPVSkeletonVisualizerComponent()
 	: SkeletonCollection(nullptr)
@@ -16,7 +21,7 @@ UPVSkeletonVisualizerComponent::UPVSkeletonVisualizerComponent()
 	bUseEditorCompositing = true;
 	SetGenerateOverlapEvents(false);
 	SetIgnoreStreamingManagerUpdate(true);
-	
+
 	BBox = FBox(ForceInit);
 	Bounds = FBoxSphereBounds(FVector::ZeroVector, FVector::OneVector, 1);
 }
@@ -42,6 +47,8 @@ void UPVSkeletonVisualizerComponent::SetCollection(const FManagedArrayCollection
 	}
 
 	SkeletonCollection = InSkeletonCollection;
+	
+	TArray<FLinearColor> PointColors = GetVisualizationColors();
 
 	const TManagedArray<FVector3f>& PointsPositions = PointFacade.GetPositions();
 
@@ -54,11 +61,56 @@ void UPVSkeletonVisualizerComponent::SetCollection(const FManagedArrayCollection
 
 		for (int32 BranchPointIndex = 1; BranchPointIndex < BranchPoints.Num(); ++BranchPointIndex)
 		{
-			FVector CurrentPosition = FVector(PointsPositions[BranchPoints[BranchPointIndex]]);
+			const int32 PointIndex = BranchPoints[BranchPointIndex];
+			FVector CurrentPosition = FVector(PointsPositions[PointIndex]);
 
-			AddLine(PreviousPosition, CurrentPosition, FLinearColor::White, ESceneDepthPriorityGroup::SDPG_World, EPointDrawSettings::Start);
+			AddLine(PreviousPosition, CurrentPosition, PointColors[PointIndex], ESceneDepthPriorityGroup::SDPG_World, EPointDrawSettings::End);
 
 			PreviousPosition = CurrentPosition;
 		}
 	}
-}
\ No newline at end of file
+}
+
+ESkeletonVisualizationModes UPVSkeletonVisualizerComponent::GetVisualizationMode() const
+{
+	return CurrentVisualizationMode;
+}
+
+void UPVSkeletonVisualizerComponent::SetVisualizationMode(ESkeletonVisualizationModes InMode)
+{
+	CurrentVisualizationMode = InMode;
+	SetCollection(GetCollection());
+}
+
+TArray<FLinearColor> UPVSkeletonVisualizerComponent::GetVisualizationColors() const
+{
+	const int32 NumPoints = SkeletonCollection->NumElements(PV::GroupNames::PointGroup);
+
+	TArray<FLinearColor> PointColors;
+	PointColors.Init(FLinearColor::White, NumPoints);
+
+	TArray<float> VisualizationValues = PV::Facades::FTreeFacade::GetVisualizationValues(*SkeletonCollection, CurrentVisualizationMode);
+
+	if (VisualizationValues.IsEmpty())
+	{
+		return PointColors;
+	}
+
+	const float MaxValue = *Algo::MaxElement(VisualizationValues);
+	const float MinValue = *Algo::MinElement(VisualizationValues);
+	const float Range = MaxValue - MinValue;
+
+	if (Range < UE_SMALL_NUMBER)
+	{
+		PointColors.Init(FLinearColor::Black, NumPoints);
+	}
+	else
+	{
+		for (int32 Index = 0; Index < NumPoints; ++Index)
+		{
+			PointColors[Index] = PV::Utilities::GetRandomHueColor(1 - (VisualizationValues[Index] - MinValue) / Range);
+		}
+	}
+
+	return PointColors;
+}
diff --git a/Engine/Plugins/Experimental/ProceduralVegetationEditor/Source/ProceduralVegetationEditor/Private/Visualizations/PVSkeletonVisualizerComponent.h b/Engine/Plugins/Experimental/ProceduralVegetationEditor/Source/ProceduralVegetationEditor/Private/Visualizations/PVSkeletonVisualizerComponent.h
--- a/Engine/Plugins/Experimental/ProceduralVegetationEditor/Source/ProceduralVegetationEditor/Private/Visualizations/PVSkeletonVisualizerComponent.h
+++ b/Engine/Plugins/Experimental/ProceduralVegetationEditor/Source/ProceduralVegetationEditor/Private/Visualizations/PVSkeletonVisualizerComponent.h
@@ -4,6 +4,9 @@
 
 #include "CoreMinimal.h"
 #include "PVLineBatchComponent.h"
+
+#include "Facades/PVTreeFacade.h"
+
 #include "PVSkeletonVisualizerComponent.generated.h"
 
 struct FManagedArrayCollection;
@@ -19,6 +22,11 @@ public:
 	const FManagedArrayCollection* GetCollection() const;
 	void SetCollection(const FManagedArrayCollection* const InSkeletonCollection);
 
+	ESkeletonVisualizationModes GetVisualizationMode() const;
+	void SetVisualizationMode(ESkeletonVisualizationModes InMode);
+	TArray<FLinearColor> GetVisualizationColors() const;
+
 private:
+	ESkeletonVisualizationModes CurrentVisualizationMode = ESkeletonVisualizationModes::None;
 	const FManagedArrayCollection* SkeletonCollection;
 };
diff --git a/Engine/Plugins/Experimental/ProceduralVegetationEditor/Source/ProceduralVegetationEditor/Private/Widgets/SPVGradientBox.cpp b/Engine/Plugins/Experimental/ProceduralVegetationEditor/Source/ProceduralVegetationEditor/Private/Widgets/SPVGradientBox.cpp
new file mode 100644
--- /dev/null
+++ b/Engine/Plugins/Experimental/ProceduralVegetationEditor/Source/ProceduralVegetationEditor/Private/Widgets/SPVGradientBox.cpp
@@ -0,0 +1,35 @@
+// Copyright Epic Games, Inc. All Rights Reserved.
+
+#include "SPVGradientBox.h"
+
+void SPVGradientBox::Construct(const FArguments& InArgs)
+{
+	BoxSize = InArgs._BoxSize;
+
+	for (int32 Index = 0; Index < 7; ++Index)
+	{
+		const FLinearColor DrawColor = FLinearColor(240.0f - (Index * 40.f), 1.f, 1.f).HSVToLinearRGB();
+		GradientStops.Add(FSlateGradientStop(BoxSize * (static_cast<float>(Index) / 6.0f), DrawColor));
+	}
+}
+
+int32 SPVGradientBox::OnPaint(const FPaintArgs& Args, const FGeometry& AllottedGeometry, const FSlateRect& MyCullingRect,
+                              FSlateWindowElementList& OutDrawElements, int32 LayerId, const FWidgetStyle& InWidgetStyle, bool bParentEnabled) const
+{
+	const FVector4f CornerRadius = FVector4f(4.0f, 4.0f, 4.0f, 4.0f);
+	FSlateDrawElement::MakeGradient(
+		OutDrawElements,
+		LayerId++,
+		AllottedGeometry.ToPaintGeometry(FVector2f(BoxSize, AllottedGeometry.GetLocalSize().Y), FSlateLayoutTransform()),
+		GradientStops,
+		Orient_Vertical,
+		ESlateDrawEffect::None,
+		CornerRadius
+	);
+	return LayerId;
+}
+
+FVector2D SPVGradientBox::ComputeDesiredSize(float LayoutScaleMultiplier) const
+{
+	return FVector2D(BoxSize, 20.0f);
+}
diff --git a/Engine/Plugins/Experimental/ProceduralVegetationEditor/Source/ProceduralVegetationEditor/Private/Widgets/SPVGradientBox.h b/Engine/Plugins/Experimental/ProceduralVegetationEditor/Source/ProceduralVegetationEditor/Private/Widgets/SPVGradientBox.h
new file mode 100644
--- /dev/null
+++ b/Engine/Plugins/Experimental/ProceduralVegetationEditor/Source/ProceduralVegetationEditor/Private/Widgets/SPVGradientBox.h
@@ -0,0 +1,24 @@
+// Copyright Epic Games, Inc. All Rights Reserved.
+
+#pragma once
+#include "Widgets/SCompoundWidget.h"
+
+class SPVGradientBox : public SCompoundWidget
+{
+	SLATE_BEGIN_ARGS(SPVGradientBox) {}
+		SLATE_ARGUMENT(float, BoxSize)
+	SLATE_END_ARGS()
+
+public:
+	void Construct(const FArguments& InArgs);
+
+	virtual int32 OnPaint(const FPaintArgs& Args, const FGeometry& AllottedGeometry, const FSlateRect& MyCullingRect,
+	                      FSlateWindowElementList& OutDrawElements, int32 LayerId, const FWidgetStyle& InWidgetStyle,
+	                      bool bParentEnabled) const override;
+	virtual FVector2D ComputeDesiredSize(float LayoutScaleMultiplier) const override;
+	
+private:
+	float BoxSize = 0.0f;
+	
+	TArray<FSlateGradientStop> GradientStops;
+};
