Fix broken color and GizmoScale for PVE debug point visualization

#Synced-CL 50245014

diff --git a/Engine/Plugins/Experimental/ProceduralVegetationEditor/Source/ProceduralVegetationEditor/Private/Visualizations/DebugVisualization/PVDebugVisualizationBase.cpp b/Engine/Plugins/Experimental/ProceduralVegetationEditor/Source/ProceduralVegetationEditor/Private/Visualizations/DebugVisualization/PVDebugVisualizationBase.cpp
--- a/Engine/Plugins/Experimental/ProceduralVegetationEditor/Source/ProceduralVegetationEditor/Private/Visualizations/DebugVisualization/PVDebugVisualizationBase.cpp
+++ b/Engine/Plugins/Experimental/ProceduralVegetationEditor/Source/ProceduralVegetationEditor/Private/Visualizations/DebugVisualization/PVDebugVisualizationBase.cpp
@@ -444,7 +444,7 @@ void FPVDebugVisualizationBase::DrawPivotPoints(const FVisualizerDrawContext& In
 		for (int32 i = 0; i < Positions.Num(); i++)
 		{
 			const FVector Pos(Positions[i]);
-			AddDrawAsPoint(Pos, Settings.AnchorPointsScale);
+			AddDrawAsPoint(Pos, Settings.AnchorPointsScale, Settings.Color);
 		}	
 	}
 }
@@ -550,7 +550,7 @@ void FPVDebugVisualizationBase::DrawAttributes(FVisualizerDrawContext& InContext
 					GetPivot(Collection, i, OutPos, OutScale);
 
 					const FColor DrawColor = VisualizationSettings.bRandomizeColors ? FColor::MakeRandomColor() : VisualizationSettings.Color;
-					AddDrawAsPoint(FVector(OutPos), AttributeAsScalar[i] * VisualizationSettings.GizmoScale * OutScale);
+					AddDrawAsPoint(FVector(OutPos), AttributeAsScalar[i] * VisualizationSettings.GizmoScale * OutScale, VisualizationSettings.Color);
 				}
 			}
 			else 
@@ -565,7 +565,7 @@ void FPVDebugVisualizationBase::DrawAttributes(FVisualizerDrawContext& InContext
 				{
 					for (int32 j = 0; j < AttributeAsVector[i].Num(); ++j)
 					{
-						AddDrawAsPoint(AttributeAsVector[i][j], VisualizationSettings.GizmoScale);
+						AddDrawAsPoint(AttributeAsVector[i][j], VisualizationSettings.GizmoScale, VisualizationSettings.Color);
 					}
 				}	
 			}
@@ -622,7 +622,7 @@ void FPVDebugVisualizationBase::DrawAttributes(FVisualizerDrawContext& InContext
 			using namespace PVE::AttributeToScalar;
 
 			TArray<double> AttributeAsScalar = ConvertAttributeToScalar(Collection, AttributeName, GroupName);
-			if (AttributeAsScalar.Num() != NumElements)
+			if (AttributeAsScalar.Num() != NumElements || NumElements == 0)
 			{
 				break;
 			}
@@ -630,6 +630,10 @@ void FPVDebugVisualizationBase::DrawAttributes(FVisualizerDrawContext& InContext
 			const double MaxValue = *Algo::MaxElement(AttributeAsScalar);
 			const double MinValue = *Algo::MinElement(AttributeAsScalar);
 			const double Range = MaxValue - MinValue;
+			if (FMath::IsNearlyZero(Range))
+			{
+				break;
+			}
 
 			for (int32 i = 0; i < NumElements; i++)
 			{
@@ -637,7 +641,7 @@ void FPVDebugVisualizationBase::DrawAttributes(FVisualizerDrawContext& InContext
 				float OutScale;
 				GetPivot(Collection, i, OutPos, OutScale);
 
-				AddDrawAsPoint(static_cast<FVector>(OutPos), OutScale, PV::Utilities::GetRandomHueColor(AttributeAsScalar[i] - MinValue) / Range);
+				AddDrawAsPoint(static_cast<FVector>(OutPos), OutScale * VisualizationSettings.GizmoScale, PV::Utilities::GetRandomHueColor(AttributeAsScalar[i] - MinValue) / Range);
 			}
 			break;
 		}
@@ -700,10 +704,7 @@ void FPVDebugVisualizationBase::DrawAsSphere(UPVLineBatchComponent* LineBatchCom
 void FPVDebugVisualizationBase::DrawAsPoints(UPVLineBatchComponent* LineBatchComponent, const TArray<TTuple<FVector, float, FLinearColor>>& InPoints)
 {
 	check(LineBatchComponent);
-	for (const auto& [Position,Size,Color] : InPoints)
-	{
-		LineBatchComponent->AddPoint(Position, Size * 10.0f, Color, SDPG_Foreground);
-	}
+	LineBatchComponent->AddPoints(InPoints, SDPG_Foreground);
 }
 
 void FPVDebugVisualizationBase::DrawAsMeshPoints(UInstancedStaticMeshComponent* InstancedStaticMeshComponent, const TArray<FTransform>& InTransforms, const FLinearColor& InColor)
