[PVE] Add Grafter node and workflow

- This is currently in development and the fact is reflected in the node title.
- Adds support for grafting and distributing a graft skeleton onto an existing skeleton (either grower based or from a preset)
- The settings for the grafter node are currently a subset of the foliage distributor settings and will later be updated to new parametric based configuration that is currently being worked on by product.
- Current implementation supports distribution of the foliage meshes as well but that will be phased out in the upcoming CLs.
- Foliage attributes, branch attributes and most of the point attributes are currently being recomputed and reindexed however some more updates are required here and will be forthcoming in subsequent CLs.

#Synced-CL 50248752

diff --git a/Engine/Plugins/Experimental/ProceduralVegetationEditor/Source/ProceduralVegetation/Private/Facades/PVBranchFacade.cpp b/Engine/Plugins/Experimental/ProceduralVegetationEditor/Source/ProceduralVegetation/Private/Facades/PVBranchFacade.cpp
--- a/Engine/Plugins/Experimental/ProceduralVegetationEditor/Source/ProceduralVegetation/Private/Facades/PVBranchFacade.cpp
+++ b/Engine/Plugins/Experimental/ProceduralVegetationEditor/Source/ProceduralVegetation/Private/Facades/PVBranchFacade.cpp
@@ -156,6 +156,11 @@ namespace PV::Facades
 		}
 	}
 
+	const TManagedArray<int32>& FBranchFacade::GetParentBranchNumbers() const
+	{
+		return BranchParentNumbers.Get();
+	}
+
 	void FBranchFacade::SetParents(int32 Index, TArray<int32> InParents)
 	{
 		if (Parents.IsValid() && Parents.IsValidIndex(Index))
@@ -182,6 +187,25 @@ namespace PV::Facades
 		}
 	}
 
+	const TArray<int32>& FBranchFacade::GetBranchFoliageIDs(int32 Index) const
+	{
+		if (BranchFoliageIDs.IsValid() && BranchFoliageIDs.IsValidIndex(Index))
+		{
+			return BranchFoliageIDs[Index];
+		}
+
+		static const TArray<int32> EmptyChildren;
+		return EmptyChildren;
+	}
+
+	void FBranchFacade::SetBranchFoliageIDs(int32 Index, TArray<int32> InBranchFoliageIDs)
+	{
+		if (BranchFoliageIDs.IsValid() && BranchFoliageIDs.IsValidIndex(Index))
+		{
+			BranchFoliageIDs.Modify()[Index] = MoveTemp(InBranchFoliageIDs);
+		}
+	}
+
 	int32 FBranchFacade::GetBranchPlantNumber(int32 Index) const
 	{
 		if (PlantNumbers.IsValid() && PlantNumbers.IsValidIndex(Index))
@@ -296,6 +320,11 @@ namespace PV::Facades
 		return INDEX_NONE;
 	}
 
+	const TManagedArray<int32>& FBranchFacade::GetBranchHierarchyNumbers() const
+	{
+		return BranchHierarchyNumber.Get();
+	}
+
 	void FBranchFacade::SetBranchHierarchyNumber(int32 Index, int32 InBranchHierarchyNumber)
 	{
 		if (BranchHierarchyNumber.IsValid() && BranchHierarchyNumber.IsValidIndex(Index))
@@ -539,4 +568,18 @@ namespace PV::Facades
 	{
 		GetValueFromAccessor<int32>(PlantNumbers, OutValue, Index, -1);
 	}
+	void FBranchFacade::SetBranchNumbersFromIndex(const TArray<int32>& InputArray, const int32 StartIndex)
+	{
+		FILL_COLLECTION_ATTRIBUTE_FROM_INDEX(BranchNumbers, InputArray, StartIndex);
+	}
+
+	void FBranchFacade::SetParentBranchNumbersFromIndex(const TArray<int32>& InputArray, const int32 StartIndex)
+	{
+		FILL_COLLECTION_ATTRIBUTE_FROM_INDEX(BranchParentNumbers, InputArray, StartIndex);
+	}
+
+	void FBranchFacade::SetBranchHierarchyNumbersFromIndex(const TArray<int32>& InputArray, const int32 StartIndex)
+	{
+		FILL_COLLECTION_ATTRIBUTE_FROM_INDEX(BranchHierarchyNumber, InputArray, StartIndex);
+	}
 }
diff --git a/Engine/Plugins/Experimental/ProceduralVegetationEditor/Source/ProceduralVegetation/Private/Facades/PVFacadeCommon.h b/Engine/Plugins/Experimental/ProceduralVegetationEditor/Source/ProceduralVegetation/Private/Facades/PVFacadeCommon.h
--- a/Engine/Plugins/Experimental/ProceduralVegetationEditor/Source/ProceduralVegetation/Private/Facades/PVFacadeCommon.h
+++ b/Engine/Plugins/Experimental/ProceduralVegetationEditor/Source/ProceduralVegetation/Private/Facades/PVFacadeCommon.h
@@ -20,4 +20,13 @@ namespace PV::Facades
 	Accessor.ModifyAt(i, InputArray[i]); \
 	}
 	
+	#define FILL_COLLECTION_ATTRIBUTE_FROM_INDEX(Accessor, InputArray, StartIndex) \
+	if (Accessor.Num() - StartIndex < InputArray.Num()) \
+	{ \
+		Accessor.AddElements(InputArray.Num() - (Accessor.Num() - StartIndex)); \
+	} \
+	for (int32 i = 0; i < InputArray.Num(); i++) \
+	{ \
+		Accessor.ModifyAt(i + StartIndex, InputArray[i]); \
+	}
 }
diff --git a/Engine/Plugins/Experimental/ProceduralVegetationEditor/Source/ProceduralVegetation/Private/Facades/PVFoliageFacade.cpp b/Engine/Plugins/Experimental/ProceduralVegetationEditor/Source/ProceduralVegetation/Private/Facades/PVFoliageFacade.cpp
--- a/Engine/Plugins/Experimental/ProceduralVegetationEditor/Source/ProceduralVegetation/Private/Facades/PVFoliageFacade.cpp
+++ b/Engine/Plugins/Experimental/ProceduralVegetationEditor/Source/ProceduralVegetation/Private/Facades/PVFoliageFacade.cpp
@@ -3,6 +3,7 @@
 #include "Facades/PVFoliageFacade.h"
  
 #include "ProceduralVegetationModule.h"
+#include "PVFacadeCommon.h"
 #include "PVFoliageAttributesNames.h"
 #include "VisualizeTexture.h"
 #include "Chaos/Deformable/MuscleActivationConstraints.h"
@@ -463,7 +464,12 @@ namespace PV::Facades
  
 		return FoliageNames;
 	}
-	
+
+	void FFoliageFacade::SetFoliageNamesFromIndex(const TArray<FString>& InputArray, const int32 StartIndex)
+	{
+		FILL_COLLECTION_ATTRIBUTE_FROM_INDEX(FoliageInfoNamesAttribute, InputArray, StartIndex);
+	}
+
 	void FFoliageFacade::SetFoliageNames(const TArray<FString>& InputNames)
 	{
 		check(!IsConst());
@@ -588,9 +594,74 @@ namespace PV::Facades
 	{
 		return PivotPointsAttribute.Get();
 	}
- 
+
+	void FFoliageFacade::SetPivotPositionsFromIndex(const TArray<FVector3f>& InputArray, const int32 StartIndex)
+	{
+		FILL_COLLECTION_ATTRIBUTE_FROM_INDEX(PivotPointsAttribute, InputArray, StartIndex);
+	}
+
 	TManagedArray<FVector3f>& FFoliageFacade::ModifyPivotPositions()
 	{
 		return PivotPointsAttribute.Modify();
 	}
-}
\ No newline at end of file
+
+	const TManagedArray<int32>& FFoliageFacade::GetFoliageNameIDs() const
+	{
+		return NameIdsAttribute.Get();
+	}
+
+	void FFoliageFacade::SetFoliageNameIDsFromIndex(const TArray<int32>& InputArray, const int32 StartIndex)
+	{
+		FILL_COLLECTION_ATTRIBUTE_FROM_INDEX(NameIdsAttribute, InputArray, StartIndex);
+	}
+
+	const TManagedArray<int32>& FFoliageFacade::GetFoliageBranchIDs() const
+	{
+		return BranchIdsAttribute.Get();
+	}
+
+	void FFoliageFacade::SetFoliageBranchIDsFromIndex(const TArray<int32>& InputArray, const int32 StartIndex)
+	{
+		FILL_COLLECTION_ATTRIBUTE_FROM_INDEX(BranchIdsAttribute, InputArray, StartIndex);
+	}
+
+	const TManagedArray<FVector3f>& FFoliageFacade::GetUpVectors() const
+	{
+		return UpVectorsAttribute.Get();
+	}
+
+	void FFoliageFacade::SetUpVectorsFromIndex(const TArray<FVector3f>& InputArray, const int32 StartIndex)
+	{
+		FILL_COLLECTION_ATTRIBUTE_FROM_INDEX(UpVectorsAttribute, InputArray, StartIndex);
+	}
+
+	const TManagedArray<FVector3f>& FFoliageFacade::GetNormalVectors() const
+	{
+		return NormalVectorsAttribute.Get();
+	}
+
+	void FFoliageFacade::SetNormalVectorsFromIndex(const TArray<FVector3f>& InputArray, const int32 StartIndex)
+	{
+		FILL_COLLECTION_ATTRIBUTE_FROM_INDEX(NormalVectorsAttribute, InputArray, StartIndex);
+	}
+
+	const TManagedArray<float>& FFoliageFacade::GetScales() const
+	{
+		return ScalesAttribute.Get();
+	}
+
+	void FFoliageFacade::SetScalesFromIndex(const TArray<float>& InputArray, const int32 StartIndex)
+	{
+		FILL_COLLECTION_ATTRIBUTE_FROM_INDEX(ScalesAttribute, InputArray, StartIndex);
+	}
+
+	const TManagedArray<float>& FFoliageFacade::GetLFRs() const
+	{
+		return LengthFromRootAttribute.Get();
+	}
+
+	void FFoliageFacade::SetLFRsFromIndex(const TArray<float>& InputArray, const int32 StartIndex)
+	{
+		FILL_COLLECTION_ATTRIBUTE_FROM_INDEX(LengthFromRootAttribute, InputArray, StartIndex);
+	}
+}
diff --git a/Engine/Plugins/Experimental/ProceduralVegetationEditor/Source/ProceduralVegetation/Private/Facades/PVPointFacade.cpp b/Engine/Plugins/Experimental/ProceduralVegetationEditor/Source/ProceduralVegetation/Private/Facades/PVPointFacade.cpp
--- a/Engine/Plugins/Experimental/ProceduralVegetationEditor/Source/ProceduralVegetation/Private/Facades/PVPointFacade.cpp
+++ b/Engine/Plugins/Experimental/ProceduralVegetationEditor/Source/ProceduralVegetation/Private/Facades/PVPointFacade.cpp
@@ -704,4 +704,33 @@ namespace PV::Facades
 	{
 		GetValueFromAccessor<TArray<float>>(BudLateralMeristem, OutValue, Index, TArray<float>());
 	}
+	const TManagedArray<TArray<float>>& FPointFacade::GetBudLightDetectedArrays() const
+	{
+		return BudLightDetected.Get();
+	}
+
+	void FPointFacade::SetBudLightDetectedArraysFromIndex(const TArray<TArray<float>>& InputArray, const int32 StartIndex)
+	{
+		FILL_COLLECTION_ATTRIBUTE_FROM_INDEX(BudLightDetected, InputArray, StartIndex);
+	}
+	
+	void FPointFacade::SetPositionsFromIndex(const TArray<FVector3f>& InputArray, const int32 StartIndex)
+	{
+		FILL_COLLECTION_ATTRIBUTE_FROM_INDEX(Positions, InputArray, StartIndex);
+	}
+
+	void FPointFacade::SetLFRsFromIndex(const TArray<float>& InputArray, const int32 StartIndex)
+	{
+		FILL_COLLECTION_ATTRIBUTE_FROM_INDEX(LengthFromRoot, InputArray, StartIndex);
+	}
+
+	const TManagedArray<float>& FPointFacade::GetLengthFromSeedsArray() const
+	{
+		return LengthFromSeed.Get();
+	}
+
+	void FPointFacade::SetLengthFromSeedsArrayFromIndex(const TArray<float>& InputArray, const int32 StartIndex)
+	{
+		FILL_COLLECTION_ATTRIBUTE_FROM_INDEX(LengthFromSeed, InputArray, StartIndex);
+	}
 }
diff --git a/Engine/Plugins/Experimental/ProceduralVegetationEditor/Source/ProceduralVegetation/Private/Implementations/PVGrafter.cpp b/Engine/Plugins/Experimental/ProceduralVegetationEditor/Source/ProceduralVegetation/Private/Implementations/PVGrafter.cpp
new file mode 100644
--- /dev/null
+++ b/Engine/Plugins/Experimental/ProceduralVegetationEditor/Source/ProceduralVegetation/Private/Implementations/PVGrafter.cpp
@@ -0,0 +1,646 @@
+// Copyright Epic Games, Inc. All Rights Reserved.
+
+#include "PVGrafter.h"
+
+#include "ProceduralVegetationModule.h"
+#include "PVFloatRamp.h"
+#include "Algo/RandomShuffle.h"
+
+#include "Facades/PVAttributesNames.h"
+#include "Facades/PVBranchFacade.h"
+#include "Facades/PVBudVectorsFacade.h"
+#include "Facades/PVFoliageFacade.h"
+#include "Facades/PVPointFacade.h"
+#include "Helpers/PVFoliageJSONHelper.h"
+
+
+FTransform FPVGrafter::GetPointTransform(const FVector3f AttachmentPoint, const FVector3f UpVector, const FVector3f NormalVector, const float Scale)
+{
+	const FVector UpVectorN = FVector(UpVector).GetSafeNormal();
+	if (UpVectorN.IsNearlyZero())
+	{
+		return FTransform(FQuat::Identity, FVector(AttachmentPoint), FVector::OneVector);
+	}
+
+	const FVector NormalVectorN = FVector(NormalVector).GetSafeNormal();
+
+	FVector OrthogonalizedNormalVectorN =
+		(NormalVectorN - FVector::DotProduct(NormalVectorN, UpVectorN) * UpVectorN).GetSafeNormal();
+	if (OrthogonalizedNormalVectorN.IsNearlyZero())
+	{
+		FVector TempVector;
+		UpVectorN.FindBestAxisVectors(TempVector, OrthogonalizedNormalVectorN);
+	}
+
+	const FQuat RotationQuat = FQuat(FRotationMatrix::MakeFromYZ(OrthogonalizedNormalVectorN, UpVectorN));
+	return FTransform(RotationQuat, FVector(AttachmentPoint), FVector::OneVector);
+}
+
+void FPVGrafter::SetGraft(FManagedArrayCollection& OutCollection, const FManagedArrayCollection& SourceSkeletonCollection,
+                          const FManagedArrayCollection& SourceGraftCollection, const FPVGraftAttributes& InAttributes,
+                          const TMap<int32, int32>& FoliageNamesOldIDsToNewIDs)
+{
+	// Transform and set points
+	PV::Facades::FPointFacade PointFacadeOut(OutCollection);
+	const PV::Facades::FPointFacade PointFacadeSourceGraft(SourceGraftCollection);
+
+	const TManagedArray<FVector3f>& PointPositions = PointFacadeSourceGraft.GetPositions();
+	const FTransform Transform = GetPointTransform(InAttributes.AttachmentPoint, InAttributes.UpVector, InAttributes.NormalVector,
+		InAttributes.Scale);
+	TArray<FVector3f> TransformedPoints;
+	const int32 NumOfPoints = PointPositions.Num();
+	TransformedPoints.SetNum(NumOfPoints);
+	for (int32 i = 0; i < NumOfPoints; ++i)
+	{
+		TransformedPoints[i] = FVector3f(Transform.TransformPosition(FVector(PointPositions[i])));
+	}
+
+	const int32 PointFacadeOutStartIndex = PointFacadeOut.GetPositions().Num();
+	PointFacadeOut.SetPositionsFromIndex(TransformedPoints, PointFacadeOutStartIndex);
+
+	const TManagedArray<float>& LFRs = PointFacadeSourceGraft.GetLengthFromRootsArray();
+	TArray<float> TransformedLFRs;
+	TransformedLFRs.SetNum(NumOfPoints);
+	for (int32 i = 0; i < NumOfPoints; ++i)
+	{
+		TransformedLFRs[i] = LFRs[i] + (InAttributes.LengthFromRoot * InAttributes.Scale);
+	}
+
+	PointFacadeOut.SetLFRsFromIndex(TransformedLFRs, PointFacadeOutStartIndex);
+
+	const TManagedArray<float>& LengthFromSeeds = PointFacadeSourceGraft.GetLengthFromSeedsArray();
+	TArray<float> TransformedLengthFromSeeds;
+	TransformedLengthFromSeeds.SetNum(LengthFromSeeds.Num());
+	for (int32 i = 0; i < LengthFromSeeds.Num(); ++i)
+	{
+		TransformedLengthFromSeeds[i] = LengthFromSeeds[i] + (InAttributes.LengthFromRoot * InAttributes.Scale);
+	}
+
+	PointFacadeOut.SetLengthFromSeedsArrayFromIndex(TransformedLengthFromSeeds, PointFacadeOutStartIndex);
+
+	const TManagedArray<TArray<float>>& BudLightDetectedArrays = PointFacadeSourceGraft.GetBudLightDetectedArrays();
+	PointFacadeOut.SetBudLightDetectedArraysFromIndex(BudLightDetectedArrays.GetConstArray(), PointFacadeOutStartIndex);
+
+	const int32 NumOfPointsInSourceGraft = PointFacadeSourceGraft.GetElementCount();
+	for (int32 i = 0; i < NumOfPointsInSourceGraft; ++i)
+	{
+		const TArray<int>& BudDevelopmentValuesInGraftPoint = PointFacadeSourceGraft.GetBudDevelopment(i);
+		const TArray<int>& BudDevelopmentValuesInAttachmentPoint = PointFacadeOut.GetBudDevelopment(InAttributes.PointIndex);
+		const int32 NumOfValuesInSource = BudDevelopmentValuesInGraftPoint.Num();
+		TArray<int> TransformedValues;
+		TransformedValues.SetNum(NumOfValuesInSource);
+		check(NumOfValuesInSource == 6);
+
+		// 0_Generation
+		TransformedValues[0] = BudDevelopmentValuesInGraftPoint[0] + BudDevelopmentValuesInAttachmentPoint[0];
+
+		// 1_BudAge
+		TransformedValues[1] = BudDevelopmentValuesInGraftPoint[1] + BudDevelopmentValuesInAttachmentPoint[1];
+		
+		// 2_BranchAge
+		TransformedValues[2] = BudDevelopmentValuesInGraftPoint[2] + BudDevelopmentValuesInAttachmentPoint[2];
+		
+		// 3_AgeSenescense
+		TransformedValues[3] = 0;
+		
+		// 4_LightSenescense
+		TransformedValues[4] = 0;
+		
+		// 5_RelativeBudAge
+		// TODO: recompute 0-1 for whole tree
+		TransformedValues[5] = BudDevelopmentValuesInGraftPoint[5];
+		
+		PointFacadeOut.SetBudDevelopment(i + PointFacadeOutStartIndex, TransformedValues);
+	}
+	
+	for (int32 i = 0; i < NumOfPointsInSourceGraft; ++i)
+	{
+		const TArray<float>& BudHormoneLevelsInGraftPoint = PointFacadeSourceGraft.GetBudHormoneLevels(i);
+		// keeping them the same as original graft for now
+		// TODO: recompute and update
+		PointFacadeOut.SetBudHormoneLevels(i + PointFacadeOutStartIndex, BudHormoneLevelsInGraftPoint);
+	}
+	
+	for (int32 i = 0; i < NumOfPointsInSourceGraft; ++i)
+	{
+		const TArray<FVector3f>& BudDirectionsInGraftPoint = PointFacadeSourceGraft.GetBudDirection(i);
+		const int32 NumOfValues = BudDirectionsInGraftPoint.Num();
+		TArray<FVector3f> TransformedValues;
+		TransformedValues.SetNum(NumOfValues);
+		
+		for (int32 j = 0; j < NumOfValues; ++j)
+		{
+			TransformedValues[j] = FVector3f(Transform.TransformVectorNoScale(FVector(BudDirectionsInGraftPoint[j])));
+		}
+		PointFacadeOut.SetBudDirections(i + PointFacadeOutStartIndex, TransformedValues);
+	}
+
+	// Transform and set branches
+	PV::Facades::FBranchFacade BranchFacadeOut(OutCollection);
+	const PV::Facades::FBranchFacade BranchFacadeSourceGraft(SourceGraftCollection);
+
+	const int32 NumOfBranchesInSource = BranchFacadeSourceGraft.GetElementCount();
+	const int32 NumOfBranchesInTarget = BranchFacadeOut.GetElementCount();
+	BranchFacadeOut.AddElements(NumOfBranchesInSource);
+	for (int32 i = 0; i < NumOfBranchesInSource; ++i)
+	{
+		const TArray<int32>& SourceBranchPoints = BranchFacadeSourceGraft.GetPoints(i);
+		const int32 NumOfPointsInSource = SourceBranchPoints.Num();
+		TArray<int32> TransformedBranchPoints;
+		TransformedBranchPoints.SetNum(NumOfPointsInSource);
+
+		for (int32 j = 0; j < NumOfPointsInSource; ++j)
+		{
+			TransformedBranchPoints[j] = SourceBranchPoints[j] + PointFacadeOutStartIndex;
+		}
+
+		BranchFacadeOut.SetPoints(NumOfBranchesInTarget + i, TransformedBranchPoints);
+	}
+
+	check(NumOfBranchesInTarget > 0);
+	const int32 BranchFacadeOutStartIndex = NumOfBranchesInTarget;
+	const int32 MaxBranchNumberInTarget = *Algo::MaxElement(BranchFacadeOut.GetBranchNumbers());
+	const TManagedArray<int32>& BranchNumbers = BranchFacadeSourceGraft.GetBranchNumbers();
+	TArray<int32> TransformedBranchNumbers;
+	TransformedBranchNumbers.SetNum(NumOfBranchesInSource);
+	TMap<int32, int32> SourceOldBranchNumbersToNewBranchNumbers;
+	SourceOldBranchNumbersToNewBranchNumbers.Reserve(NumOfBranchesInSource);
+	for (int32 i = 0; i < NumOfBranchesInSource; ++i)
+	{
+		TransformedBranchNumbers[i] = BranchNumbers[i] + MaxBranchNumberInTarget;
+		SourceOldBranchNumbersToNewBranchNumbers.Add(BranchNumbers[i], TransformedBranchNumbers[i]);
+	}
+
+	BranchFacadeOut.SetBranchNumbersFromIndex(TransformedBranchNumbers, BranchFacadeOutStartIndex);
+
+	const TManagedArray<int32>& ParentBranchNumbers = BranchFacadeSourceGraft.GetParentBranchNumbers();
+	TArray<int32> TransformedParentBranchNumbers;
+	TransformedParentBranchNumbers.SetNum(ParentBranchNumbers.Num());
+	const int32 BranchNumberBeingAttachedToInTarget = BranchFacadeOut.GetBranchNumber(InAttributes.BranchIndex);
+	for (int32 i = 0; i < ParentBranchNumbers.Num(); ++i)
+	{
+		TransformedParentBranchNumbers[i] = ParentBranchNumbers[i] == 0
+			? BranchNumberBeingAttachedToInTarget
+			: SourceOldBranchNumbersToNewBranchNumbers[BranchNumbers[i]];
+	}
+
+	BranchFacadeOut.SetParentBranchNumbersFromIndex(TransformedParentBranchNumbers, BranchFacadeOutStartIndex);
+
+	const TManagedArray<int32>& BranchHierarchyNumbers = BranchFacadeSourceGraft.GetBranchHierarchyNumbers();
+	TArray<int32> TransformedBranchHierarchyNumbers;
+	TransformedBranchHierarchyNumbers.SetNum(BranchHierarchyNumbers.Num());
+	const int32 BranchHeirarchyNumberOfTarget = BranchFacadeOut.GetBranchHierarchyNumber(InAttributes.BranchIndex);
+	for (int32 i = 0; i < BranchHierarchyNumbers.Num(); ++i)
+	{
+		TransformedBranchHierarchyNumbers[i] = BranchHierarchyNumbers[i] + BranchHeirarchyNumberOfTarget;
+	}
+
+	BranchFacadeOut.SetBranchHierarchyNumbersFromIndex(TransformedBranchHierarchyNumbers, BranchFacadeOutStartIndex);
+
+	// TODO: Add code for updating branchSourceBudNumber after adding code for updating budNumbers in points
+
+	for (int32 i = 0; i < NumOfBranchesInSource; ++i)
+	{
+		const TArray<int32>& SourceBranchChildren = BranchFacadeSourceGraft.GetChildren(i);
+		const int32 NumOfChildrenInSource = SourceBranchChildren.Num();
+		TArray<int32> TransformedChildren;
+		TransformedChildren.SetNum(NumOfChildrenInSource);
+
+		for (int32 j = 0; j < NumOfChildrenInSource; ++j)
+		{
+			if (SourceOldBranchNumbersToNewBranchNumbers.Contains(SourceBranchChildren[j]))
+			{
+				TransformedChildren[j] = SourceOldBranchNumbersToNewBranchNumbers[SourceBranchChildren[j]];
+			}
+		}
+
+		BranchFacadeOut.SetChildren(NumOfBranchesInTarget + i, TransformedChildren);
+	}
+
+	for (int32 i = 0; i < NumOfBranchesInSource; ++i)
+	{
+		const TArray<int32>& SourceBranchParents = BranchFacadeSourceGraft.GetParents(i);
+		const int32 NumOfParentsInSource = SourceBranchParents.Num();
+		TArray<int32> TransformedParents;
+		TransformedParents.SetNum(NumOfParentsInSource);
+
+		for (int32 j = 0; j < NumOfParentsInSource; ++j)
+		{
+			if (SourceOldBranchNumbersToNewBranchNumbers.Contains(SourceBranchParents[j]))
+			{
+				TransformedParents[j] = SourceOldBranchNumbersToNewBranchNumbers[SourceBranchParents[j]];
+			}
+		}
+
+		BranchFacadeOut.SetParents(NumOfBranchesInTarget + i, TransformedParents);
+	}
+
+	const int32 PlantNumberOfTargetBranch = BranchFacadeOut.GetBranchPlantNumber(InAttributes.BranchIndex);
+	for (int32 i = 0; i < NumOfBranchesInSource; ++i)
+	{
+		BranchFacadeOut.SetBranchPlantNumber(NumOfBranchesInTarget + i, PlantNumberOfTargetBranch);
+	}
+
+	if (!SourceGraftCollection.HasGroup(PV::GroupNames::FoliageNamesGroup))
+	{
+		return;
+	}
+
+	PV::Facades::FFoliageFacade FoliageFacadeOut(OutCollection);
+
+	const int32 NumOfFoliageEntriesInTarget = FoliageFacadeOut.GetElementCount();
+
+	for (int32 i = 0; i < NumOfBranchesInSource; ++i)
+	{
+		const TArray<int32>& SourceBranchFoliageIDs = BranchFacadeSourceGraft.GetBranchFoliageIDs(i);
+		const int32 NumOfFoliageIDsInSource = SourceBranchFoliageIDs.Num();
+		TArray<int32> TransformedFoliageIDs;
+		TransformedFoliageIDs.SetNum(NumOfFoliageIDsInSource);
+
+		for (int32 j = 0; j < NumOfFoliageIDsInSource; ++j)
+		{
+			TransformedFoliageIDs[j] = SourceBranchFoliageIDs[j] + NumOfFoliageEntriesInTarget;
+		}
+
+		BranchFacadeOut.SetBranchFoliageIDs(NumOfBranchesInTarget + i, TransformedFoliageIDs);
+	}
+
+	// Transform and set foliage
+	const PV::Facades::FFoliageFacade FoliageFacadeSourceGraft(SourceGraftCollection);
+	const TManagedArray<int32>& FoliageNameIDs = FoliageFacadeSourceGraft.GetFoliageNameIDs();
+	TArray<int32> TransformedFoliageNameIDs;
+	TransformedFoliageNameIDs.SetNum(FoliageNameIDs.Num());
+
+	for (int32 j = 0; j < FoliageNameIDs.Num(); ++j)
+	{
+		TransformedFoliageNameIDs[j] = FoliageNamesOldIDsToNewIDs[FoliageNameIDs[j]];
+	}
+
+	FoliageFacadeOut.SetFoliageNameIDsFromIndex(TransformedFoliageNameIDs, NumOfFoliageEntriesInTarget);
+
+	const TManagedArray<int32>& FoliageBranchIDs = FoliageFacadeSourceGraft.GetFoliageBranchIDs();
+	TArray<int32> TransformedFoliageBranchIDs;
+	TransformedFoliageBranchIDs.SetNum(FoliageBranchIDs.Num());
+
+	for (int32 j = 0; j < FoliageBranchIDs.Num(); ++j)
+	{
+		TransformedFoliageBranchIDs[j] = FoliageBranchIDs[j] + NumOfBranchesInTarget;
+	}
+
+	FoliageFacadeOut.SetFoliageBranchIDsFromIndex(TransformedFoliageBranchIDs, NumOfFoliageEntriesInTarget);
+
+	const TManagedArray<FVector3f>& FoliagePivotPositions = FoliageFacadeSourceGraft.GetPivotPositions();
+	TArray<FVector3f> TransformedFoliagePivotPositions;
+	TransformedFoliagePivotPositions.SetNum(FoliagePivotPositions.Num());
+
+	for (int32 j = 0; j < FoliagePivotPositions.Num(); ++j)
+	{
+		TransformedFoliagePivotPositions[j] = FVector3f(Transform.TransformPosition(FVector(FoliagePivotPositions[j])));
+	}
+
+	FoliageFacadeOut.SetPivotPositionsFromIndex(TransformedFoliagePivotPositions, NumOfFoliageEntriesInTarget);
+
+	const TManagedArray<FVector3f>& FoliageUpVectors = FoliageFacadeSourceGraft.GetUpVectors();
+	TArray<FVector3f> TransformedFoliageUpVectors;
+	TransformedFoliageUpVectors.SetNum(FoliageUpVectors.Num());
+
+	for (int32 j = 0; j < FoliageUpVectors.Num(); ++j)
+	{
+		TransformedFoliageUpVectors[j] = FVector3f(Transform.TransformVectorNoScale(FVector(FoliageUpVectors[j])));
+	}
+
+	FoliageFacadeOut.SetUpVectorsFromIndex(TransformedFoliageUpVectors, NumOfFoliageEntriesInTarget);
+
+	const TManagedArray<FVector3f>& FoliageNormalVectors = FoliageFacadeSourceGraft.GetNormalVectors();
+	TArray<FVector3f> TransformedFoliageNormalVectors;
+	TransformedFoliageNormalVectors.SetNum(FoliageNormalVectors.Num());
+
+	for (int32 j = 0; j < FoliageNormalVectors.Num(); ++j)
+	{
+		TransformedFoliageNormalVectors[j] = FVector3f(Transform.TransformVectorNoScale(FVector(FoliageNormalVectors[j])));
+	}
+
+	FoliageFacadeOut.SetNormalVectorsFromIndex(TransformedFoliageNormalVectors, NumOfFoliageEntriesInTarget);
+
+	const TManagedArray<float>& FoliageScales = FoliageFacadeSourceGraft.GetScales();
+	TArray<float> TransformedFoliageScales;
+	TransformedFoliageScales.SetNum(FoliageScales.Num());
+
+	for (int32 j = 0; j < FoliageScales.Num(); ++j)
+	{
+		TransformedFoliageScales[j] = FoliageScales[j] * InAttributes.Scale;
+	}
+
+	FoliageFacadeOut.SetScalesFromIndex(TransformedFoliageScales, NumOfFoliageEntriesInTarget);
+
+	const TManagedArray<float>& FoliageLFRs = FoliageFacadeSourceGraft.GetLFRs();
+	TArray<float> TransformedFoliageLFRs;
+	TransformedFoliageLFRs.SetNum(FoliageLFRs.Num());
+
+	for (int32 j = 0; j < FoliageLFRs.Num(); ++j)
+	{
+		TransformedFoliageLFRs[j] = FoliageLFRs[j] + InAttributes.LengthFromRoot;
+	}
+
+	FoliageFacadeOut.SetLFRsFromIndex(TransformedFoliageLFRs, NumOfFoliageEntriesInTarget);
+}
+
+TMap<int32, int32> FPVGrafter::SetFoliageNames(FManagedArrayCollection& OutCollection, const FManagedArrayCollection& SourceGraftCollection)
+{
+	PV::Facades::FFoliageFacade FoliageFacadeOut(OutCollection);
+	const PV::Facades::FFoliageFacade FoliageFacadeSourceGraft(SourceGraftCollection);
+	const TArray<FString> FoliageNames = FoliageFacadeSourceGraft.GetFoliageNames();
+	const int32 NumOfFoliageNamesInTarget = FoliageFacadeOut.NumFoliageNames();
+	FoliageFacadeOut.SetFoliageNamesFromIndex(FoliageNames, NumOfFoliageNamesInTarget);
+
+	TMap<int32, int32> FoliageNamesOldIDsToNewIDs;
+	FoliageNamesOldIDsToNewIDs.Reserve(FoliageNames.Num());
+	for (int32 i = 0; i < FoliageNames.Num(); ++i)
+	{
+		FoliageNamesOldIDsToNewIDs.Add(i, i + NumOfFoliageNamesInTarget);
+	}
+
+	return FoliageNamesOldIDsToNewIDs;
+}
+
+void FPVGrafter::DistributeGraft(FManagedArrayCollection& OutCollection, const FManagedArrayCollection& SourceSkeletonCollection,
+                                 const FManagedArrayCollection& SourceGraftCollection, const FDistributionSettings& DistributionSettings,
+                                 const FScaleSettings& ScaleSettings, const FVectorSettings& VectorSettings, const FMiscSettings& MiscSettings)
+{
+	TRACE_CPUPROFILER_EVENT_SCOPE(PVGrafter::DistributeGraft);
+
+	FRandomStream RandomStream;
+	RandomStream.Initialize(MiscSettings.RandomSeed);
+
+	PV::Facades::FBranchFacade BranchFacadeSourceSkeleton(SourceSkeletonCollection);
+	PV::Facades::FPointFacade PointFacadeSourceSkeleton(SourceSkeletonCollection);
+
+	TManagedArrayAccessor<TArray<FVector3f>> BudDirectionAttribute(OutCollection, PV::AttributeNames::BudDirection,
+		PV::GroupNames::PointGroup);
+	TManagedArrayAccessor<float> BranchGradientsAttribute(OutCollection, PV::AttributeNames::BranchGradient,
+		PV::GroupNames::PointGroup);
+	TManagedArrayAccessor<TArray<float>> LeafPhyllotaxyAttribute(OutCollection, PV::AttributeNames::LeafPhyllotaxy,
+		PV::GroupNames::DetailsGroup);
+	TManagedArrayAccessor<TArray<float>> BudHormoneLevels(OutCollection, PV::AttributeNames::BudHormoneLevels,
+		PV::GroupNames::PointGroup);
+
+	float PhyllotaxyFormationLeaf = LeafPhyllotaxyAttribute[0][1];
+	float MinBudsLeaf = LeafPhyllotaxyAttribute[0][3];
+	float MaxBudsLeaf = LeafPhyllotaxyAttribute[0][4];
+	int32 ResetPhyllotaxyLeaf = static_cast<int>(LeafPhyllotaxyAttribute[0][6]);
+	float PhyllotaxyOffsetLeaf = LeafPhyllotaxyAttribute[0][7];
+
+	TMap<int32, int32> FoliageNamesOldIDsToNewIDs;
+	if (SourceGraftCollection.HasGroup(PV::GroupNames::FoliageNamesGroup))
+	{
+		FoliageNamesOldIDsToNewIDs = SetFoliageNames(OutCollection, SourceGraftCollection);
+	}
+
+	const int32 NumBranches = BranchFacadeSourceSkeleton.GetElementCount();
+
+	// Compute attachment points in terms of normalized distances
+	for (int32 BranchIndex = 0; BranchIndex < NumBranches; ++BranchIndex)
+	{
+		TArray<float> NormalizedAttachmentPoints;
+		TArray<float> NormalizedAttachmentPointsFinal;
+		TArray<FVector3f> AttachmentPointsPositions;
+		TArray<float> AttachmentPointsLengthsFromRoot;
+		TArray<FVector3f> AttachmentPointsUpVectors;
+		TArray<FVector3f> AttachmentPointsNormalVectors;
+		TArray<float> AttachmentPointsScales;
+		TArray<int32> AttachmentPointsPointIndices;
+		TArray<int32> BranchPoints = BranchFacadeSourceSkeleton.GetPoints(BranchIndex);
+
+		if (BranchPoints.Num() < 2)
+		{
+			continue;
+		}
+
+		const float FirstPointLengthFromRoot = PointFacadeSourceSkeleton.GetLengthFromRoot(BranchPoints[0]);
+		const float LastPointLengthFromRoot = PointFacadeSourceSkeleton.GetLengthFromRoot(BranchPoints.Last());
+		const float BranchLength = LastPointLengthFromRoot - FirstPointLengthFromRoot;
+		float Increment = DistributionSettings.InstanceSpacing / BranchLength;
+		const float AdjustedMaxPerBranch = DistributionSettings.MaxPerBranch == -1
+			? BranchPoints.Num()
+			: DistributionSettings.MaxPerBranch;
+		float LookUp = 0.0f;
+
+		for (int32 i = 0; i < AdjustedMaxPerBranch; ++i)
+		{
+			if (LookUp > 1.0f)
+			{
+				break;
+			}
+
+			float InstanceSpacingRampValue = DistributionSettings.InstanceSpacingRamp
+				? DistributionSettings.InstanceSpacingRamp->GetRichCurveConst()->Eval(1.0f - LookUp)
+				: Increment;
+			// The increment is nudged slightly before blending with the weighted ramp value, the number is chosen based on observation
+			float AdjustedIncrement = Increment + (InstanceSpacingRampValue * (Increment * 10.0f));
+			AdjustedIncrement = FMath::Lerp(Increment, AdjustedIncrement, DistributionSettings.InstanceSpacingRampEffect);
+
+			NormalizedAttachmentPoints.Add(1.0f - LookUp);
+			LookUp += AdjustedIncrement;
+		}
+
+		// Compute Position, Angles, Scale and LengthFromRoot for attachment points
+
+		// Compute based on first point attributes
+		const FVector3f InitialApicalDirection(BudDirectionAttribute[BranchPoints[0]][0]);
+		FVector3f InitialAxillaryDirection(BudDirectionAttribute[BranchPoints[0]][1]);
+		const FVector3f InitialUpVector(BudDirectionAttribute[BranchPoints[0]][5]);
+
+		if (ResetPhyllotaxyLeaf == 1)
+		{
+			InitialAxillaryDirection = FVector3f::CrossProduct(InitialApicalDirection, InitialUpVector);
+		}
+
+		if (PhyllotaxyOffsetLeaf > 0.01f)
+		{
+			const FQuat4f RotationQuat(InitialApicalDirection.GetSafeNormal(), FMath::DegreesToRadians(PhyllotaxyOffsetLeaf));
+			InitialAxillaryDirection = RotationQuat.RotateVector(InitialAxillaryDirection);
+		}
+
+		int32 AxillaryRotationIteration = 0;
+		for (int32 Idx = 0; Idx < NormalizedAttachmentPoints.Num(); ++Idx)
+		{
+			for (int32 BranchPointIndex = 0; BranchPointIndex < BranchPoints.Num() - 1; ++BranchPointIndex)
+			{
+				float AxillaryRotationDegree = FMath::DegreesToRadians(PhyllotaxyFormationLeaf * AxillaryRotationIteration);
+
+				const int32 CurrentPointIndex = BranchPoints[BranchPointIndex];
+				const int32 NextPointIndex = BranchPoints[BranchPointIndex + 1];
+				const FVector3f CurrentPointPosition = PointFacadeSourceSkeleton.GetPosition(CurrentPointIndex);
+				const FVector3f NextPointPosition = PointFacadeSourceSkeleton.GetPosition(NextPointIndex);
+				float EthyleneLevelAtCurrentPoint = BudHormoneLevels[CurrentPointIndex][4];
+				float EthyleneLevelAtNextPoint = BudHormoneLevels[NextPointIndex][4];
+				const float CurrentPointLengthFromRoot = PointFacadeSourceSkeleton.GetLengthFromRoot(CurrentPointIndex);
+				const float NextPointLengthFromRoot = PointFacadeSourceSkeleton.GetLengthFromRoot(NextPointIndex);
+				const float CurrentPointLfrNormalized = FMath::GetMappedRangeValueClamped(
+					FVector2f(FirstPointLengthFromRoot, LastPointLengthFromRoot),
+					FVector2f(0.0f, 1.0f), CurrentPointLengthFromRoot);
+				const float NextPointLfrNormalized = FMath::GetMappedRangeValueClamped(
+					FVector2f(FirstPointLengthFromRoot, LastPointLengthFromRoot),
+					FVector2f(0.0f, 1.0f),
+					NextPointLengthFromRoot);
+
+				if (NormalizedAttachmentPoints[Idx] >= CurrentPointLfrNormalized && NormalizedAttachmentPoints[Idx] <=
+					NextPointLfrNormalized)
+				{
+					const float DistanceAlpha = (NormalizedAttachmentPoints[Idx] -
+						CurrentPointLfrNormalized) / (NextPointLfrNormalized - CurrentPointLfrNormalized);
+
+					const float EthyleneLevel = FMath::Lerp(EthyleneLevelAtCurrentPoint, EthyleneLevelAtNextPoint, DistanceAlpha);
+
+					if ((EthyleneLevel - 0.001f) >= DistributionSettings.EthyleneThreshold)
+					{
+						continue;
+					}
+
+					const FVector3f Position = FMath::Lerp(CurrentPointPosition, NextPointPosition, DistanceAlpha);
+					AttachmentPointsPositions.Add(Position);
+					
+					AttachmentPointsPointIndices.Add(CurrentPointIndex);
+
+					const float AttachmentPointLfr = FMath::GetMappedRangeValueClamped(
+						FVector2f(0.0f, 1.0f),
+						FVector2f(FirstPointLengthFromRoot, LastPointLengthFromRoot),
+						NormalizedAttachmentPoints[Idx]);
+					AttachmentPointsLengthsFromRoot.Add(AttachmentPointLfr);
+
+					// Compute Scale
+					const float CurrentPointBranchGradient = BranchGradientsAttribute[CurrentPointIndex];
+					const float NextPointBranchGradient = BranchGradientsAttribute[NextPointIndex];
+
+					const float BranchGradient = FMath::Clamp(
+						FMath::Lerp(CurrentPointBranchGradient, NextPointBranchGradient, DistanceAlpha), 0.0f, 1.0f);
+					float BranchGradientRampValue = ScaleSettings.ScaleRamp
+						? ScaleSettings.ScaleRamp->GetRichCurveConst()->Eval(1.0f - BranchGradient)
+						: 1.0f;
+
+					const float CurrentPointScale = PointFacadeSourceSkeleton.GetPointScale(CurrentPointIndex);
+					const float NextPointScale = PointFacadeSourceSkeleton.GetPointScale(NextPointIndex);
+					const float BranchScale = FMath::Lerp(CurrentPointScale, NextPointScale, DistanceAlpha);
+
+					const float BranchScaledValue = BranchGradientRampValue * (BranchScale + 0.1f);
+					BranchGradientRampValue = FMath::Lerp(BranchGradientRampValue, BranchScaledValue, ScaleSettings.BranchScaleImpact);
+					BranchGradientRampValue = FMath::Clamp(BranchGradientRampValue, 0.0f, 1.0f);
+
+					float Scale = FMath::GetMappedRangeValueClamped(
+						FVector2f(0.0f, 1.0f),
+						FVector2f(ScaleSettings.MinScale, ScaleSettings.MaxScale),
+						BranchGradientRampValue
+					) * ScaleSettings.BaseScale;
+
+					const float RandomValue = RandomStream.FRandRange(ScaleSettings.RandomScaleMin, ScaleSettings.RandomScaleMax);
+					Scale = Scale * RandomValue;
+
+					AttachmentPointsScales.Add(Scale);
+
+					const FVector3f CurrentPointApicalDirectionVector(BudDirectionAttribute[CurrentPointIndex][0]);
+					const FVector3f NextPointApicalDirectionVector(BudDirectionAttribute[NextPointIndex][0]);
+					const FVector3f BudApicalDirectionVector = FMath::Lerp(CurrentPointApicalDirectionVector,
+						NextPointApicalDirectionVector, DistanceAlpha);
+
+					int32 BudsToAdd = 1;
+
+					// Compute Up Vector
+					// Last point (order of the array is last to first)
+					if (Idx == 0)
+					{
+						AttachmentPointsUpVectors.Add(BudApicalDirectionVector);
+					}
+					else // Points other than last
+					[[likely]]
+					{
+						const FQuat4f RotationQuat(InitialApicalDirection.GetSafeNormal(), AxillaryRotationDegree);
+						FVector3f NewAxillaryDirection = RotationQuat.RotateVector(InitialAxillaryDirection.GetSafeNormal());
+						FQuat4f ApicalCorrection = FQuat4f::FindBetweenNormals(
+							InitialApicalDirection.GetSafeNormal(), BudApicalDirectionVector.GetSafeNormal());
+						NewAxillaryDirection = ApicalCorrection.RotateVector(NewAxillaryDirection).GetSafeNormal();
+
+						int32 NumBuds = RandomStream.RandRange(MinBudsLeaf, MaxBudsLeaf);
+						float RotationPerBud = 360.0f / static_cast<float>(NumBuds);
+
+						for (int32 i = 0; i < NumBuds; ++i)
+						{
+							float RotationAngle = FMath::DegreesToRadians(RotationPerBud * i);
+							const FQuat4f AxillaryRotation(InitialApicalDirection.GetSafeNormal(), RotationAngle);
+							FVector3f AxillaryDirectionModified = AxillaryRotation.RotateVector(NewAxillaryDirection.GetSafeNormal());
+
+							if (VectorSettings.OverrideAxilAngle)
+							{
+								float AxilAngleRampValue = VectorSettings.AxilAngleRamp
+									? FMath::Clamp(
+										VectorSettings.AxilAngleRamp->GetRichCurveConst()->Eval(1.0f - BranchGradient),
+										0.0f,
+										1.0f)
+									: 1.0f;
+
+								AxilAngleRampValue = AxilAngleRampValue * VectorSettings.AxilAngleRampUpperValue;
+								const float AxilAngleBlended = FMath::Lerp(VectorSettings.AxilAngle, AxilAngleRampValue,
+									VectorSettings.AxilAngleRampEffect);
+
+								const FVector3f RotationAxis = FVector3f::CrossProduct(
+									BudApicalDirectionVector.GetSafeNormal(), AxillaryDirectionModified.GetSafeNormal());
+
+								const FQuat4f Rotation(RotationAxis.GetSafeNormal(), FMath::DegreesToRadians(AxilAngleBlended));
+								const FVector3f AxilVector = Rotation.RotateVector(BudApicalDirectionVector.GetSafeNormal());
+
+								AttachmentPointsUpVectors.Add(AxilVector);
+							}
+							else
+							{
+								AttachmentPointsUpVectors.Add(AxillaryDirectionModified);
+							}
+
+							if (i > 0)
+							{
+								FVector3f LastPositionAdded = AttachmentPointsPositions.Last();
+								float LastScaleAdded = AttachmentPointsScales.Last();
+								float LastLengthFromRootAdded = AttachmentPointsLengthsFromRoot.Last();
+								AttachmentPointsPositions.Add(LastPositionAdded);
+								AttachmentPointsPointIndices.Add(CurrentPointIndex);
+								AttachmentPointsScales.Add(LastScaleAdded);
+								AttachmentPointsLengthsFromRoot.Add(LastLengthFromRootAdded);
+							}
+						}
+
+						BudsToAdd = NumBuds;
+						AxillaryRotationIteration++;
+					}
+
+					// Compute Normal vectors based on light optimal direction
+					const FVector3f CurrentPointLightOptimalVector(BudDirectionAttribute[CurrentPointIndex][2]);
+					const FVector3f NextPointLightOptimalVector(BudDirectionAttribute[NextPointIndex][2]);
+					const FVector3f BudLightOptimalVector = FMath::Lerp(CurrentPointLightOptimalVector, NextPointLightOptimalVector,
+						DistanceAlpha);
+
+					FQuat4f QuatRotation = FQuat4f::FindBetweenNormals(BudApicalDirectionVector.GetSafeNormal(),
+						BudLightOptimalVector.GetSafeNormal());
+					const FVector3f ApicalDirectionRotated = QuatRotation.RotateVector(BudApicalDirectionVector);
+
+					for (int32 i = 0; i < BudsToAdd; ++i)
+					{
+						AttachmentPointsNormalVectors.Add(ApicalDirectionRotated);
+					}
+				}
+			}
+		}
+
+		for (int32 Idx = 0; Idx < AttachmentPointsLengthsFromRoot.Num(); Idx++)
+		{
+			FPVGraftAttributes Attributes{
+				.BranchIndex = BranchIndex,
+				.PointIndex = AttachmentPointsPointIndices[Idx],
+				.AttachmentPoint = AttachmentPointsPositions[Idx],
+				.UpVector = AttachmentPointsUpVectors[Idx],
+				.NormalVector = AttachmentPointsNormalVectors[Idx],
+				.Scale = AttachmentPointsScales[Idx],
+				.LengthFromRoot = AttachmentPointsLengthsFromRoot[Idx]
+			};
+
+			SetGraft(OutCollection, SourceSkeletonCollection, SourceGraftCollection, Attributes, FoliageNamesOldIDsToNewIDs);
+		}
+	}
+}
diff --git a/Engine/Plugins/Experimental/ProceduralVegetationEditor/Source/ProceduralVegetation/Private/Implementations/PVGrafter.h b/Engine/Plugins/Experimental/ProceduralVegetationEditor/Source/ProceduralVegetation/Private/Implementations/PVGrafter.h
new file mode 100644
--- /dev/null
+++ b/Engine/Plugins/Experimental/ProceduralVegetationEditor/Source/ProceduralVegetation/Private/Implementations/PVGrafter.h
@@ -0,0 +1,73 @@
+// Copyright Epic Games, Inc. All Rights Reserved.
+
+#pragma once
+
+#include "CoreMinimal.h"
+#include "GeometryCollection/ManagedArrayCollection.h"
+
+// #include "PVGrafter.generated.h"
+
+struct FPVFloatRamp;
+
+	
+
+struct FPVGrafter
+{
+	struct FDistributionSettings
+	{
+		const float EthyleneThreshold;
+		const float InstanceSpacing;
+		const FPVFloatRamp* InstanceSpacingRamp;
+		const float InstanceSpacingRampEffect;
+		const int32 MaxPerBranch;
+	};
+
+	struct FScaleSettings
+	{
+		const float BaseScale;
+		const float BranchScaleImpact;
+		const float MinScale;
+		const float MaxScale;
+		const float RandomScaleMin;
+		const float RandomScaleMax;
+		const FPVFloatRamp* ScaleRamp;
+	};
+
+	struct FVectorSettings
+	{
+		const bool OverrideAxilAngle;
+		const float AxilAngle;
+		const FPVFloatRamp* AxilAngleRamp;
+		const float AxilAngleRampUpperValue;
+		const float AxilAngleRampEffect;
+	};
+
+	struct FMiscSettings
+	{
+		const int32 RandomSeed;
+	};
+	
+	struct FPVGraftAttributes
+	{
+		int32 BranchIndex;
+		int32 PointIndex;
+		FVector3f AttachmentPoint;
+		FVector3f UpVector;
+		FVector3f NormalVector;
+		float Scale;
+		float LengthFromRoot;
+	};
+
+	static FTransform GetPointTransform(const FVector3f AttachmentPoint, const FVector3f UpVector, const FVector3f NormalVector,
+	                                    const float Scale);
+
+	static void SetGraft(FManagedArrayCollection& OutCollection, const FManagedArrayCollection& SourceSkeletonCollection,
+	                     const FManagedArrayCollection& SourceGraftCollection, const FPVGraftAttributes& InAttributes,
+	                     const TMap<int32, int32>& FoliageNamesOldIDsToNewIDs);
+	
+	static TMap<int32, int32> SetFoliageNames(FManagedArrayCollection& OutCollection, const FManagedArrayCollection& SourceGraftCollection);
+	
+	static void DistributeGraft(FManagedArrayCollection& OutCollection, const FManagedArrayCollection& SourceSkeletonCollection,
+	                            const FManagedArrayCollection& SourceGraftCollection, const FDistributionSettings& DistributionSettings,
+	                            const FScaleSettings& ScaleSettings, const FVectorSettings& VectorSettings, const FMiscSettings& MiscSettings);
+};
diff --git a/Engine/Plugins/Experimental/ProceduralVegetationEditor/Source/ProceduralVegetation/Private/Nodes/PVGrafterSettings.cpp b/Engine/Plugins/Experimental/ProceduralVegetationEditor/Source/ProceduralVegetation/Private/Nodes/PVGrafterSettings.cpp
new file mode 100644
--- /dev/null
+++ b/Engine/Plugins/Experimental/ProceduralVegetationEditor/Source/ProceduralVegetation/Private/Nodes/PVGrafterSettings.cpp
@@ -0,0 +1,128 @@
+// Copyright Epic Games, Inc. All Rights Reserved.
+
+#include "PVGrafterSettings.h"
+#include "PCGContext.h"
+#include "ProceduralVegetationModule.h"
+#include "DataTypes/PVGrowthData.h"
+#include "ProceduralVegetationPreset.h"
+#include "Facades/PVBranchFacade.h"
+#include "Facades/PVPointFacade.h"
+#include "Implementations/PVCarve.h"
+
+#define LOCTEXT_NAMESPACE "PVGrafterSettings"
+
+#if WITH_EDITOR
+FText UPVGrafterSettings::GetDefaultNodeTitle() const
+{
+	return LOCTEXT("NodeTitle", "Grafter (In Development)");
+}
+
+FText UPVGrafterSettings::GetNodeTooltipText() const
+{
+	return LOCTEXT("NodeTooltip",
+		"This allows the user to distribute a skeleton onto another one. "
+		"\n\nPress Ctrl + L to lock/unlock node output"
+	);
+}
+#endif
+
+TArray<FPCGPinProperties> UPVGrafterSettings::InputPinProperties() const
+{
+	TArray<FPCGPinProperties> Properties;
+
+	FPCGPinProperties& Pin1 = Properties.Emplace_GetRef(PVGrafterInputPins::Skeleton, GetInputPinTypeIdentifier());
+	Pin1.SetRequiredPin();
+	Pin1.SetAllowMultipleConnections(false);
+	Pin1.bAllowMultipleData = false;
+
+	FPCGPinProperties& Pin2 = Properties.Emplace_GetRef(PVGrafterInputPins::Graft, GetInputPinTypeIdentifier());
+	Pin2.SetRequiredPin();
+	Pin2.SetAllowMultipleConnections(false);
+	Pin2.bAllowMultipleData = false;
+
+	return Properties;
+}
+
+FPCGDataTypeIdentifier UPVGrafterSettings::GetOutputPinTypeIdentifier() const
+{
+	return FPCGDataTypeIdentifier{FPVDataTypeInfoGrowth::AsId()};
+}
+
+FPCGElementPtr UPVGrafterSettings::CreateElement() const
+{
+	return MakeShared<FPVGrafterElement>();
+}
+
+bool FPVGrafterElement::ExecuteInternal(FPCGContext* InContext) const
+{
+	TRACE_CPUPROFILER_EVENT_SCOPE(FPVGrafterElement::ExecuteInternal);
+
+	check(InContext);
+
+	const UPVGrafterSettings* Settings = InContext->GetInputSettings<UPVGrafterSettings>();
+	check(Settings);
+
+	auto SkeletonPinInputs = InContext->InputData.GetInputsByPin(PVGrafterInputPins::Skeleton);
+	auto GraftPinInputs = InContext->InputData.GetInputsByPin(PVGrafterInputPins::Graft);
+
+	if (SkeletonPinInputs.IsEmpty() || GraftPinInputs.IsEmpty())
+	{
+		return true;
+	}
+
+	const UPVGrowthData* SkeletonInputData = Cast<UPVGrowthData>(SkeletonPinInputs[0].Data);
+	const UPVGrowthData* GraftInputData = Cast<UPVGrowthData>(GraftPinInputs[0].Data);
+	if (!SkeletonInputData || !GraftInputData)
+	{
+		PCGLog::InputOutput::LogInvalidInputDataError(InContext);
+		return true;
+	}
+
+	FManagedArrayCollection SourceSkeletonCollection = SkeletonInputData->GetCollection();
+	FManagedArrayCollection SourceGraftCollection = GraftInputData->GetCollection();
+
+	FManagedArrayCollection OutCollection;
+	SourceSkeletonCollection.CopyTo(&OutCollection);
+
+	UPVGrowthData* OutManagedArrayCollectionData = FPCGContext::NewObject_AnyThread<UPVGrowthData>(InContext);
+
+	const FPVGrafter::FDistributionSettings DistributionSettings = {
+		.EthyleneThreshold = Settings->EthyleneThreshold,
+		.InstanceSpacing = Settings->InstanceSpacing,
+		.InstanceSpacingRamp = &Settings->InstanceSpacingRamp,
+		.InstanceSpacingRampEffect = Settings->InstanceSpacingRampEffect,
+		.MaxPerBranch = Settings->MaxPerBranch
+	};
+
+	const FPVGrafter::FScaleSettings ScaleSettings = {
+		.BaseScale = Settings->BaseScale,
+		.BranchScaleImpact = Settings->BranchScaleImpact,
+		.MinScale = Settings->MinScale,
+		.MaxScale = Settings->MaxScale,
+		.RandomScaleMin = Settings->RandomScaleMin,
+		.RandomScaleMax = Settings->RandomScaleMax,
+		.ScaleRamp = &Settings->ScaleRamp
+	};
+
+	const FPVGrafter::FVectorSettings VectorSettings = {
+		.OverrideAxilAngle = Settings->OverrideAxilAngle,
+		.AxilAngle = Settings->AxilAngle,
+		.AxilAngleRamp = &Settings->AxilAngleRamp,
+		.AxilAngleRampUpperValue = Settings->AxilAngleRampUpperValue,
+		.AxilAngleRampEffect = Settings->AxilAngleRampEffect
+	};
+
+	const FPVGrafter::FMiscSettings MiscSettings = {
+		.RandomSeed = Settings->RandomSeed
+	};
+
+	FPVGrafter::DistributeGraft(OutCollection, SourceSkeletonCollection, SourceGraftCollection,
+		DistributionSettings, ScaleSettings, VectorSettings, MiscSettings);
+
+	OutManagedArrayCollectionData->Initialize(MoveTemp(OutCollection));
+	InContext->OutputData.TaggedData.Emplace(OutManagedArrayCollectionData);
+
+	return true;
+}
+
+#undef LOCTEXT_NAMESPACE
diff --git a/Engine/Plugins/Experimental/ProceduralVegetationEditor/Source/ProceduralVegetation/Private/Nodes/PVGrafterSettings.h b/Engine/Plugins/Experimental/ProceduralVegetationEditor/Source/ProceduralVegetation/Private/Nodes/PVGrafterSettings.h
new file mode 100644
--- /dev/null
+++ b/Engine/Plugins/Experimental/ProceduralVegetationEditor/Source/ProceduralVegetation/Private/Nodes/PVGrafterSettings.h
@@ -0,0 +1,153 @@
+// Copyright Epic Games, Inc. All Rights Reserved.
+#pragma once
+
+#include "CoreMinimal.h"
+#include "PVFloatRamp.h"
+#include "Nodes/PVBaseSettings.h"
+#include "Implementations/PVGrafter.h"
+
+#include "PVGrafterSettings.generated.h"
+
+namespace PVGrafterInputPins
+{
+	static const FName Skeleton("Skeleton");
+	static const FName Graft("Graft");
+}
+
+UCLASS(BlueprintType, ClassGroup = (Procedural))
+class UPVGrafterSettings : public UPVBaseSettings
+{
+	GENERATED_BODY()
+	
+	friend class FPVGrafterElement;
+
+public:
+#if WITH_EDITOR
+	virtual FName GetDefaultNodeName() const override
+	{
+		return FName(TEXT("Grafter"));
+	}
+
+	virtual FText GetDefaultNodeTitle() const override;
+	virtual FText GetNodeTooltipText() const override;
+
+	virtual TArray<EPVRenderType> GetDefaultRenderType() const override
+	{
+		return TArray{EPVRenderType::PointData};
+	}
+#endif
+
+protected:
+	virtual TArray<FPCGPinProperties> InputPinProperties() const override;
+	virtual FPCGDataTypeIdentifier GetOutputPinTypeIdentifier() const override;
+	virtual FPCGElementPtr CreateElement() const override;
+
+private:
+	UPROPERTY(EditAnywhere, Category="Distribution Settings",
+		meta=(ClampMin=0.0f, ClampMax=1.0f, UIMin=0.0f, UIMax=1.0f, Tooltip=
+			"Threshold controlling foliage release based on an ethylene-like signal.\n\nHigher Threshold → more buds/branches retained, making the plant denser and bushier.\nLower Threshold → fewer buds/branches retained, resulting in a sparser, pruned structure."
+		))
+	float EthyleneThreshold = 0.85f;
+
+	UPROPERTY(EditAnywhere, Category="Distribution Settings",
+		meta=(ClampMin=0.0f, ClampMax=1.0f, UIMin=0.0f, UIMax=1.0f, Tooltip=
+			"Minimum distance between placed instances.\n\nHigher spacing yields a more open distribution; lower spacing allows denser placement."))
+	float InstanceSpacing = 0.3f;
+
+	UPROPERTY(EditAnywhere, Category="Distribution Settings",
+		meta=(XAxisMin=0.0f, XAxisMax=1.0f, YAxisMin=0.0f, YAxisMax=1.0f, Tooltip=
+			"Ramp that varies instance spacing across the height of the plant\n\nUseful to reduce spacing in some areas based on the height of the plant."
+		))
+	FPVFloatRamp InstanceSpacingRamp;
+
+	UPROPERTY(EditAnywhere, Category="Distribution Settings",
+		meta=(ClampMin=0.0f, ClampMax=1.0f, UIMin=0.0f, UIMax=1.0f, Tooltip=
+			"Strength of the spacing ramp’s influence.\n\nBlends the ramp with the base spacing. Lower values favor the base spacing (subtle effect); higher values apply the ramp more strongly and can amplify spacing differences."
+		))
+	float InstanceSpacingRampEffect = 0.0f;
+
+	UPROPERTY(EditAnywhere, Category="Distribution Settings",
+		meta=(ClampMin=-1, ClampMax=1000, Tooltip=
+			"Caps the number of instances per branch\n\nlimit on how many foliage instances can be placed per branch to avoid overcrowding."))
+	int32 MaxPerBranch = -1;
+
+	UPROPERTY(EditAnywhere, Category="Scale Settings",
+		meta=(ClampMin=0.0f, UIMin=0.0f, UIMax=5.0f, Tooltip=
+			"Default size multiplier for placed foliage\n\nSets the base uniform scale applied to each instance before randomness and ramps are considered. Increase to enlarge all instances proportionally; decrease for smaller foliage. Final size still respects Min/Max limits."
+		))
+	float BaseScale = 1.0f;
+
+	UPROPERTY(EditAnywhere, Category="Scale Settings",
+		meta=(ClampMin=0.0f, ClampMax=1.0f, UIMin=0.0f, UIMax=1.0f, Tooltip=
+			"Impact of branch size on foliage size\n\nThe size of the foliage being scattered will be chosen in relation to the physical size of the branch it is distributed on. The largest branches will be distributed the largest foliage pieces."
+		))
+	float BranchScaleImpact = 0.0f;
+
+	UPROPERTY(EditAnywhere, Category="Scale Settings",
+		meta=(ClampMin=0.0f, UIMin=0.0f, UIMax=5.0f, Tooltip=
+			"Minimum allowed instance size.\n\nClamps the final computed scale so it never goes below this value. Use to prevent tiny, hard-to-see instances."
+		))
+	float MinScale = 0.5f;
+
+	UPROPERTY(EditAnywhere, Category="Scale Settings",
+		meta=(ClampMin=0.0f, UIMin=0.0f, UIMax=5.0f, Tooltip=
+			"Maximum allowed instance size.\n\nClamps the final computed scale so it never exceeds this value. Use to cap unusually large instances for visual consistency."
+		))
+	float MaxScale = 1.0f;
+
+	UPROPERTY(EditAnywhere, Category="Scale Settings",
+		meta=(ClampMin=0.0f, UIMin=0.0f, UIMax=5.0f, Tooltip=
+			"Lower bound for per-instance size variation.\n\nSets the minimum of the random scale range used to vary instance sizes. Each instance samples within Random Scale Min–Max, then combines with Base Scale and any ramps (respecting Min/Max Scale)."
+		))
+	float RandomScaleMin = 1.0f;
+
+	UPROPERTY(EditAnywhere, Category="Scale Settings",
+		meta=(ClampMin=1.0f, UIMin=1.0f, UIMax=5.0f, Tooltip=
+			"Upper bound for per-instance size variation.\n\nSets the maximum of the random scale range used to vary instance sizes. Each instance samples within Random Scale Min–Max, then combines with Base Scale and any ramps (respecting Min/Max Scale)."
+		))
+	float RandomScaleMax = 1.0f;
+
+	UPROPERTY(EditAnywhere, Category="Scale Settings",
+		meta=(XAxisMin=0.0f, XAxisMax=1.0f, YAxisMin=0.0f, YAxisMax=1.0f, Tooltip=
+			"Curve that varies scale across the height of the plant\n\nUse to taper foliage toward tips or emphasize size in targeted regions."))
+	FPVFloatRamp ScaleRamp;
+
+	UPROPERTY(EditAnywhere, Category="Vector Settings",
+		meta=(Tooltip=
+			"Use a custom axil angle instead of the default.\n\nEnables a user-defined orientation. When on, the system uses AxilAngle and the AxilAngleRamp settings (and Effect) to set the tilt of foliage relative to its parent branch direction, overriding any automatic/species defaults."
+		))
+	bool OverrideAxilAngle = true;
+
+	UPROPERTY(EditAnywhere, Category="Vector Settings",
+		meta=(ClampMin=0.0f, ClampMax=90.0f, UIMin=0.0f, UIMax=90.0f, Tooltip=
+			"Base tilt from the parent axis\n\nSets the base axil angle in degrees—the tilt of the instance relative to the parent branch’s direction. Acts as the lower bound when used with a ramp; higher values tilt foliage further away from the branch axis."
+		))
+	float AxilAngle = 45.0f;
+
+	UPROPERTY(EditAnywhere, Category="Vector Settings",
+		meta=(XAxisMin=0.0f, XAxisMax=1.0f, YAxisMin=0.0f, YAxisMax=1.0f, Tooltip=
+			"Ramp that varies the axil angle across the height of the plant\n\nIts output is remapped to an angle range defined by AxilAngle and AxilAngleRampUpperValue, modulated by AxilAngleRampEffect."
+		))
+	FPVFloatRamp AxilAngleRamp;
+
+	UPROPERTY(EditAnywhere, Category="Vector Settings",
+		meta=(ClampMin=0.0f, ClampMax=90.0f, UIMin=0.0f, UIMax=90.0f, Tooltip=
+			"Max angle used when the ramp\n\nProvides detailed control over this parameter to balance visual quality and performance. Valid values: 0.0–90.0."
+		))
+	float AxilAngleRampUpperValue = 45.0f;
+
+	UPROPERTY(EditAnywhere, Category="Vector Settings",
+		meta=(ClampMin=0.0f, ClampMax=1.0f, UIMin=0.0f, UIMax=1.0f, Tooltip=
+			"Adjusts this parameter for finer control.\n\nDefines the upper axil angle reached when AxilAngleRamp evaluates to 1. Together with AxilAngle, it forms the min–max range for ramp-driven orientation."
+		))
+	float AxilAngleRampEffect = 0.0f;
+
+	UPROPERTY(EditAnywhere, Category="Misc Settings")
+	int32 RandomSeed = 123456;
+};
+
+class FPVGrafterElement : public FPVBaseElement
+{
+protected:
+	virtual bool ExecuteInternal(FPCGContext* InContext) const override;
+};
diff --git a/Engine/Plugins/Experimental/ProceduralVegetationEditor/Source/ProceduralVegetation/Public/Facades/PVBranchFacade.h b/Engine/Plugins/Experimental/ProceduralVegetationEditor/Source/ProceduralVegetation/Public/Facades/PVBranchFacade.h
--- a/Engine/Plugins/Experimental/ProceduralVegetationEditor/Source/ProceduralVegetation/Public/Facades/PVBranchFacade.h
+++ b/Engine/Plugins/Experimental/ProceduralVegetationEditor/Source/ProceduralVegetation/Public/Facades/PVBranchFacade.h
@@ -44,12 +44,18 @@ namespace PV::Facades
 		int32 GetParentBranchNumber(int32 BranchIndex) const;
 		
 		void SetParentBranchNumber(int32 BranchIndex, int32 InParentBranchNumber);
+		
+		const TManagedArray<int32>& GetParentBranchNumbers() const;
 
 		void SetParents(int32 Index, TArray<int32> InParents);
 
 		int32 GetBranchNumber(int32 Index) const;
 
 		void SetBranchNumber(int32 Index, int32 InBranchNumber);
+		
+		const TArray<int32>& GetBranchFoliageIDs(int32 Index) const;
+		
+		void SetBranchFoliageIDs(int32 Index, TArray<int32> InBranchFoliageIDs);
 
 		int32 GetBranchPlantNumber(int32 Index) const;
 
@@ -72,6 +78,8 @@ namespace PV::Facades
 		void SetBranchUVMaterial(int32 Index, int32 InMaterial);
 		
 		int32 GetBranchHierarchyNumber(int32 Index) const;
+		
+		const TManagedArray<int32>& GetBranchHierarchyNumbers() const;
 
 		void SetBranchHierarchyNumber(int32 Index, int32 InBranchHierarchyNumber);
 
@@ -112,6 +120,12 @@ namespace PV::Facades
 		void FillBranchParentNumbers(const TArray<int32>& InputArray);
 		
 		void FillPlantNumbers(const TArray<int32>& InputArray);
+		
+		void SetBranchNumbersFromIndex(const TArray<int32>& InputArray, const int32 StartIndex);
+		
+		void SetParentBranchNumbersFromIndex(const TArray<int32>& InputArray, const int32 StartIndex);
+		
+		void SetBranchHierarchyNumbersFromIndex(const TArray<int32>& InputArray, const int32 StartIndex);
 
 		template<typename T>
 		FORCEINLINE bool GetValueFromAccessor(TManagedArrayAccessor<T> Accessor, T& OutValue, int32 Index, const T DefaultValue) const
diff --git a/Engine/Plugins/Experimental/ProceduralVegetationEditor/Source/ProceduralVegetation/Public/Facades/PVFoliageFacade.h b/Engine/Plugins/Experimental/ProceduralVegetationEditor/Source/ProceduralVegetation/Public/Facades/PVFoliageFacade.h
--- a/Engine/Plugins/Experimental/ProceduralVegetationEditor/Source/ProceduralVegetation/Public/Facades/PVFoliageFacade.h
+++ b/Engine/Plugins/Experimental/ProceduralVegetationEditor/Source/ProceduralVegetation/Public/Facades/PVFoliageFacade.h
@@ -66,15 +66,19 @@ namespace PV::Facades
 		int32 GetFoliageBranchId(int32 Index) const;
  
 		int32 NumFoliageNames() const { return FoliageInfoNamesAttribute.Num(); };
+		
 		void SetFoliageBranchId(int32 Index, int32 InputId);
 		
 		const FString GetFoliageName(int32 Index) const;
  
 		TArray<FString> GetFoliageNames() const;
+		
+		void SetFoliageNamesFromIndex(const TArray<FString>& InputArray, const int32 StartIndex);
  
 		int32 NumFoliageInfo() const { return FoliageInfoNamesAttribute.Num(); };
  
 		FPVFoliageInfo GetFoliageInfo(const int32 Index) const;
+		
 		const FVector3f& GetPivotPoint(int32 Index) const;
  
 		TArray<FPVFoliageInfo> GetFoliageInfos() const;
@@ -113,7 +117,33 @@ namespace PV::Facades
  
 		const TManagedArray<FVector3f>& GetPivotPositions() const;
 		
+		void SetPivotPositionsFromIndex(const TArray<FVector3f>& InputArray, const int32 StartIndex);
+		
 		TManagedArray<FVector3f>& ModifyPivotPositions();
+		
+		const TManagedArray<int32>& GetFoliageNameIDs() const;
+		
+		void SetFoliageNameIDsFromIndex(const TArray<int32>& InputArray, const int32 StartIndex);
+		
+		const TManagedArray<int32>& GetFoliageBranchIDs() const;
+		
+		void SetFoliageBranchIDsFromIndex(const TArray<int32>& InputArray, const int32 StartIndex);
+		
+		const TManagedArray<FVector3f>& GetUpVectors() const;
+		
+		void SetUpVectorsFromIndex(const TArray<FVector3f>& InputArray, const int32 StartIndex);
+		
+		const TManagedArray<FVector3f>& GetNormalVectors() const;
+		
+		void SetNormalVectorsFromIndex(const TArray<FVector3f>& InputArray, const int32 StartIndex);
+		
+		const TManagedArray<float>& GetScales() const;
+		
+		void SetScalesFromIndex(const TArray<float>& InputArray, const int32 StartIndex);
+		
+		const TManagedArray<float>& GetLFRs() const;
+		
+		void SetLFRsFromIndex(const TArray<float>& InputArray, const int32 StartIndex);
  
 	protected:
 		void DefineSchema(int32 InitialSize = 0);
diff --git a/Engine/Plugins/Experimental/ProceduralVegetationEditor/Source/ProceduralVegetation/Public/Facades/PVPointFacade.h b/Engine/Plugins/Experimental/ProceduralVegetationEditor/Source/ProceduralVegetation/Public/Facades/PVPointFacade.h
--- a/Engine/Plugins/Experimental/ProceduralVegetationEditor/Source/ProceduralVegetation/Public/Facades/PVPointFacade.h
+++ b/Engine/Plugins/Experimental/ProceduralVegetationEditor/Source/ProceduralVegetation/Public/Facades/PVPointFacade.h
@@ -160,6 +160,14 @@ Accessor.ModifyAt(i, InputArray[i]); \
 		void FillBudDevelopment(const TArray<TArray<int>>& InputArray);
 
 		void FillNjordPixelID(const TArray<float>& InputArray);
+		
+		void SetPositionsFromIndex(const TArray<FVector3f>& InputArray, const int32 StartIndex);
+		
+		void SetLFRsFromIndex(const TArray<float>& InputArray, const int32 StartIndex);
+		
+		const TManagedArray<float>& GetLengthFromSeedsArray() const;
+		
+		void SetLengthFromSeedsArrayFromIndex(const TArray<float>& InputArray, const int32 StartIndex);
 
 		void FillBudStatus(const TArray<TArray<int>>& InputArray);
 
@@ -201,6 +209,10 @@ Accessor.ModifyAt(i, InputArray[i]); \
 		void GetBudNumbers(int32 Index, int32& OutValue) const;
 
 		void GetBudLightDetected(int32 Index, TArray<float>& OutValue) const;
+		
+		const TManagedArray<TArray<float>>& GetBudLightDetectedArrays() const;
+		
+		void SetBudLightDetectedArraysFromIndex(const TArray<TArray<float>>& InputArray, const int32 StartIndex);
 
 		void GetBudDevelopment(int32 Index, TArray<int>& OutValue) const;
 		
