[PVE] Exposing grower settings outside its node
- Exposed pins for growth settings on its node.
- Created grower settings and their data and type-related classes.
- Eventually, these classes will be removed when PCG implements struct overridable pins.
- There are some issues, like they are not derived from PVData because they are not dependent upon ManagedArrayCollection and don't need any visualization. We will tackle these issues separately if needed.

#Synced-CL 50296570

diff --git a/Engine/Plugins/Experimental/ProceduralVegetationEditor/Source/ProceduralVegetation/Private/DataTypes/PVGrowerParams.h b/Engine/Plugins/Experimental/ProceduralVegetationEditor/Source/ProceduralVegetation/Private/DataTypes/PVGrowerParams.h
new file mode 100644
--- /dev/null
+++ b/Engine/Plugins/Experimental/ProceduralVegetationEditor/Source/ProceduralVegetation/Private/DataTypes/PVGrowerParams.h
@@ -0,0 +1,406 @@
+// Copyright Epic Games, Inc. All Rights Reserved.
+
+#pragma once
+
+#include "CoreMinimal.h"
+#include "PVFloatRamp.h"
+#include "PVGrowerParams.generated.h"
+
+USTRUCT()
+struct FPVBranchingParams
+{
+	GENERATED_BODY()
+
+	UPROPERTY(EditAnywhere , Category = "Branching")
+	float AuxinFalloff = 0.0f;
+
+	UPROPERTY(EditAnywhere , Category = "Branching")
+	float ApicalDominance = 0.0f;
+
+	UPROPERTY(EditAnywhere , Category = "Branching")
+	float CytokininBuildup = 0.0f;
+
+	UPROPERTY(EditAnywhere , Category = "Branching")
+	float CytokininDominance = 0.0f;
+
+	UPROPERTY(EditAnywhere , Category = "Branching")
+	float CytokininRandomness = 1.0f;
+
+	UPROPERTY(EditAnywhere , Category = "Branching")
+	float LightRequirement = 0.0f;
+	
+	UPROPERTY(EditAnywhere , Category = "Branching")
+	float RadicalAuxin = 0.0f;
+
+	UPROPERTY(EditAnywhere , Category = "Branching")
+	float MinGravitationalDot = 0.0f;
+};
+
+USTRUCT()
+struct FPVAxialParams
+{
+	GENERATED_BODY()
+
+		UPROPERTY(EditAnywhere , Category = "AxialElongation" , meta=(UIMin=0.0f, UIMax=1.0f))
+	float Elongation = 0.25;
+
+	UPROPERTY(EditAnywhere , Category = "AxialElongation" , meta=(UIMin=0.1f, UIMax=2.0f))
+	float ChildScale = 1.0f;
+
+	UPROPERTY(EditAnywhere , Category = "AxialElongation" , meta=(UIMin=0.0f, UIMax=1.0f))
+	float SeedScaleEffect = 1.0f;
+
+	UPROPERTY(EditAnywhere , Category = "AxialElongation" , meta=(UIMin=0.0f, UIMax=1.0f))
+	float GibberellicAcid = 0.0f;
+
+	UPROPERTY(EditAnywhere , Category = "Gibberellic Acid Controls" , meta=(UIMin=0.0f, UIMax=1.0f))
+	float Random = 0.0f;
+
+	UPROPERTY(EditAnywhere , Category = "Gibberellic Acid Controls", meta=(UIMin=-1.0f, UIMax=1.0f))
+	float LightEffect = 0.0f;
+
+	UPROPERTY(EditAnywhere , Category = "Gibberellic Acid Controls", meta=(UIMin=0.0f, UIMax=1.0f))
+	float AuxinEffect = 0.0f;
+};
+
+USTRUCT()
+struct FPVLateralParams
+{
+	GENERATED_BODY()
+
+	UPROPERTY(EditAnywhere , Category = "LateralElongation", meta=(UIMin=0.0f, UIMax=0.01f))
+	float Elongation = 0.0033f;
+
+	UPROPERTY(EditAnywhere , Category = "LateralElongation", meta=(UIMin=0.0f, UIMax=1.0f))
+	float SeedScaleEffect = 1.0f;
+	
+	UPROPERTY(EditAnywhere , Category = "LateralElongation", meta=(UIMin=0.0f, UIMax=1.0f))
+	float WhorledImpact = 0.5f;
+};
+
+UENUM()
+enum class EPVHitGroundBehaviour : uint8
+{
+	Kill,
+	Deflect,
+	Ignore,
+};
+
+USTRUCT()
+struct FPVGrowerGravityParams
+{
+	GENERATED_BODY()
+
+	UPROPERTY()
+	bool bCustom =	true;
+
+	UPROPERTY(EditAnywhere , Category = "Gravity")
+	FFloatRange CellDensity = FFloatRange(7000, 10000);
+
+	UPROPERTY(EditAnywhere , Category = "Gravity", meta=(UIMin=0.0f, UIMax=100.0f))
+	float CellDevelopmentTime = 5.0f;
+
+	UPROPERTY(EditAnywhere , Category = "Gravity", DisplayName = "Core Reinforcement", meta=(UIMin=0.0f, UIMax=1.0f))
+	float CellDensityCore = 0.5f;
+
+	UPROPERTY(EditAnywhere , Category = "Gravity", DisplayName = "Input Reinforcement", meta=(UIMin=0.0f, UIMax=1.0f))
+	float CellDensityInput = 0.0f;
+
+	UPROPERTY(EditAnywhere , Category = "Gravity" , meta=(UIMin=0.0f, UIMax=10.0f))
+	float CellWeight = 0.5f;
+
+	UPROPERTY(EditAnywhere , Category = "Gravity", meta=(UIMin=0.0f, UIMax=1.0f))
+	float LeafWeight = 0.1f;
+
+	UPROPERTY(EditAnywhere , Category = "Gravity", meta=(UIMin=0.0f, UIMax=10.0f))
+	float GravitationalForce = 9.8f;
+
+	UPROPERTY(EditAnywhere , Category = "Gravity", meta=(UIMin=0.0f, UIMax=1.0f))
+	float BranchWeightBreakThreshold = 0.33f;
+
+	UPROPERTY(EditAnywhere , Category = "Gravity", meta=(UIMin=0.0f, UIMax=10.0f))
+	EPVHitGroundBehaviour HitGroundBehaviour = EPVHitGroundBehaviour::Kill;
+};
+
+USTRUCT()
+struct FPVLeafGrowthParams
+{
+	GENERATED_BODY()
+
+	UPROPERTY(EditAnywhere, Category = "FLeafGrowth")
+	float Ethylene = 0.15f;
+
+	UPROPERTY(EditAnywhere, Category = "FLeafGrowth" , meta=(UIMin=0.1f, UIMax=5.0f))
+	float StartScale = 0.1f;
+
+	UPROPERTY(EditAnywhere, Category = "FLeafGrowth" , meta=(UIMin=0.1f, UIMax=5.0f))
+	float EndScale = 1.0f;
+
+	UPROPERTY(EditAnywhere, Category = "FLeafGrowth" , meta=(UIMin=0, UIMax=10))
+	int Density = 3;
+
+	UPROPERTY(EditAnywhere, Category = "FLeafGrowth" , meta=(UIMin=0.0f, UIMax=1.0f))
+	float AbscisicAcid = 0;
+
+	UPROPERTY(EditAnywhere, Category = "FLeafGrowth" , meta=(UIMin=0.0f, UIMax=1.0f))
+	float EthyleneThreshold = 0.5f;
+
+	UPROPERTY(EditAnywhere, Category = "FLeafGrowth" , meta=(UIMin=0.0f, UIMax=1.0f))
+	float LeafAuxinRetention = 0.0;
+};
+
+USTRUCT()
+struct FPVAbscissionScenesenseParams
+{
+	GENERATED_BODY()
+	
+};
+
+static int PhyllotaxyTypeAngles[] = {180,0,90,90,0};
+
+UENUM()
+enum class EPVGrowthPhyllotaxyType : uint8
+{
+	Alternate,
+	Opposite,
+	Decussate,
+	Whorled,
+	Spiral,
+};
+
+UENUM()
+enum class EPVGrowthPhyllotaxyFormation : uint8
+{
+	Disticious = 180,
+	Tristicious = 120,
+	Pentasticious = 144,
+	Octasticious = 135,
+	Parasticious = 0,
+};
+
+USTRUCT()
+struct FPVPhyllotaxyParams
+{
+	GENERATED_BODY()
+
+	UPROPERTY(EditAnywhere , Category = "Phyllotaxy")
+	EPVGrowthPhyllotaxyType Type = EPVGrowthPhyllotaxyType::Spiral;
+
+	UPROPERTY(EditAnywhere , Category = "Phyllotaxy", meta = (EditCondition = "Type == EPVGrowthPhyllotaxyType::Spiral", EditConditionHides))
+	EPVGrowthPhyllotaxyFormation Formation = EPVGrowthPhyllotaxyFormation::Octasticious;
+
+	UPROPERTY(EditAnywhere , Category = "Phyllotaxy" , meta = (EditCondition = "Type == EPVGrowthPhyllotaxyType::Whorled", EditConditionHides))
+	uint32 Min = 3;
+
+	UPROPERTY(EditAnywhere , Category = "Phyllotaxy", meta = (EditCondition = "Type == EPVGrowthPhyllotaxyType::Whorled", EditConditionHides))
+	uint32 Max = 6;
+
+	UPROPERTY(EditAnywhere , Category = "Phyllotaxy")
+	float AdditionalAngle = 0.0f;
+
+	UPROPERTY(EditAnywhere , Category = "Phyllotaxy")
+	float Offset = 0.0f;
+
+	UPROPERTY(EditAnywhere , Category = "Phyllotaxy", meta=(ClampMin=0.0f, ClampMax=1.0f, UIMin=0.0f, UIMax=1.0f))
+	float Flatten = 0.0f;
+
+	UPROPERTY(EditAnywhere , Category = "Phyllotaxy")
+	float AxilAngle = 35.0f;
+
+	UPROPERTY(EditAnywhere , Category = "Phyllotaxy")
+	bool bAllowCodominantBranching = false;
+
+	UPROPERTY(EditAnywhere , Category = "Phyllotaxy", meta = (EditCondition = "bAllowCodominantBranching == true", EditConditionHides))
+	float CodominantAxilAngle = 60.0f;
+
+	UPROPERTY(EditAnywhere , Category = "Phyllotaxy", meta = (EditCondition = "bAllowCodominantBranching == true", EditConditionHides))
+	float CodominantAxilBias = 0.0f;
+
+	UPROPERTY(EditAnywhere , Category = "Phyllotaxy", meta = (EditCondition = "bAllowCodominantBranching == true", EditConditionHides))
+	uint32 CodominantBudMin = 2;
+
+	UPROPERTY(EditAnywhere , Category = "Phyllotaxy", meta = (EditCondition = "bAllowCodominantBranching == true", EditConditionHides))
+	uint32 CodominantBudMax = 2;
+
+	UPROPERTY(EditAnywhere , Category = "Phyllotaxy")
+	bool bReset = false;
+
+	UPROPERTY(EditAnywhere , Category = "Phyllotaxy", meta=(UIMin=0.0f, UIMax=1.0f))
+	float Stagger = 0.0f;
+};
+
+USTRUCT()
+struct FPVPhototropismParams
+{
+	GENERATED_BODY()
+
+	UPROPERTY(EditAnywhere , Category = "Phototropism", meta=(UIMin=0.0f, UIMax=1.0f))
+	float Apical = 0.25;
+
+	UPROPERTY(EditAnywhere , Category = "Phototropism", meta=(UIMin=0.0f, UIMax=1.0f))
+	float ApicalBias = 0.1f;
+
+	UPROPERTY(EditAnywhere , Category = "Phototropism", meta=(UIMin=0.0f, UIMax=1.0f))
+	float Axillary = 0.1f;
+
+	UPROPERTY(EditAnywhere , Category = "Phototropism", meta=(UIMin=0.0f, UIMax=1.0f))
+	float AxillaryBias = 0.75f;
+};
+
+USTRUCT()
+struct FPVRandomAngleParams
+{
+	GENERATED_BODY()
+
+	UPROPERTY(EditAnywhere , Category = "RandomAngle", meta=(UIMin=0.0f, UIMax=45.0f))
+	float Apical = 0.0f;
+
+	UPROPERTY(EditAnywhere , Category = "RandomAngle", meta=(UIMin=0.0f, UIMax=45.0f))
+	float Axillary = 0.0f;
+
+	UPROPERTY(EditAnywhere , Category = "RandomAngle", meta=(UIMin=0.0f, UIMax=45.0f))
+	float Codominent = 0.0f;
+};
+
+USTRUCT()
+struct FPVTrunkGrowthParams
+{
+	GENERATED_BODY()
+
+	UPROPERTY(EditAnywhere , Category = "TrunkGrowth")
+	float PlantLength = 5.0f;
+
+	UPROPERTY(EditAnywhere , Category = "TrunkGrowth")
+	float ApicalParentGrowth = 1.0f;
+
+	UPROPERTY(EditAnywhere , Category = "TrunkGrowth")
+	float ApicalChildGrowth = 1.0f;
+
+	UPROPERTY(EditAnywhere , Category = "TrunkGrowth")
+	float AxillaryParentGrowth = 1.0f;
+
+	UPROPERTY(EditAnywhere , Category = "TrunkGrowth")
+	float AxillaryChildGrowth = 1.0f;
+
+	UPROPERTY(EditAnywhere , Category = "TrunkGrowth")
+	bool bAxillaryRetry = true;
+
+	UPROPERTY(EditAnywhere , Category = "TrunkGrowth")
+	uint8 MaxGeneration = 25;
+
+	UPROPERTY(EditAnywhere , Category = "TrunkGrowth")
+	float BranchLength = 0.0f;
+
+	UPROPERTY(EditAnywhere , Category = "TrunkGrowth")
+	float Corymb = 0.0f;
+
+	UPROPERTY(EditAnywhere , Category = "TrunkGrowth")
+	bool bAxillaryUseChildGradient = false;
+
+	UPROPERTY(EditAnywhere , Category = "TrunkGrowth")
+	bool bAxillaryUseLFT = false;
+
+	UPROPERTY(EditAnywhere , Category = "TrunkGrowth")
+	bool bApicalUseChildGradient = false;
+
+	UPROPERTY(EditAnywhere , Category = "TrunkGrowth")
+	bool bApicalUseLFT = false;
+
+	UPROPERTY(EditAnywhere, Category="TrunkGrowth", meta=(XAxisMin=0.0f, XAxisMax=1.0f, YAxisMin=0.0f, YAxisMax=1.0f))
+	FPVFloatRamp ApicalPriorityGradient;
+
+	UPROPERTY(EditAnywhere, Category="TrunkGrowth", meta=(XAxisMin=0.0f, XAxisMax=1.0f, YAxisMin=0.0f, YAxisMax=1.0f, EditCondition = "bApicalUseChildGradient == true", EditConditionHides))
+	FPVFloatRamp ApicalPriorityChildGradient;
+
+	UPROPERTY(EditAnywhere, Category="TrunkGrowth", meta=(XAxisMin=0.0f, XAxisMax=1.0f, YAxisMin=0.0f, YAxisMax=1.0f))
+	FPVFloatRamp AxillaryPriorityGradient;
+
+	UPROPERTY(EditAnywhere, Category="TrunkGrowth", meta=(XAxisMin=0.0f, XAxisMax=1.0f, YAxisMin=0.0f, YAxisMax=1.0f, EditCondition = "bAxillaryUseChildGradient == true", EditConditionHides))
+	FPVFloatRamp AxillaryPriorityChildGradient;
+
+	FPVTrunkGrowthParams()
+	{
+		InitializeLinearCurve(ApicalPriorityGradient);
+		InitializeLinearCurve(ApicalPriorityChildGradient);
+		InitializeLinearCurve(AxillaryPriorityGradient);
+		InitializeLinearCurve(AxillaryPriorityChildGradient);
+	}
+
+private:
+	static void InitializeLinearCurve(FPVFloatRamp& Curve)
+	{
+		if (FRichCurve* RichCurve = Curve.GetRichCurve())
+		{
+			RichCurve->Reset();
+			const FKeyHandle Handle0 = RichCurve->AddKey(0.0f, 0.0f);
+			const FKeyHandle Handle1 = RichCurve->AddKey(1.0f, 1.0f);
+
+			RichCurve->SetKeyInterpMode(Handle0, RCIM_Linear);
+			RichCurve->SetKeyInterpMode(Handle1, RCIM_Linear);
+		}
+	}
+};
+
+USTRUCT()
+struct FPVGuideParams
+{
+	GENERATED_BODY()
+
+	UPROPERTY(EditAnywhere , Category = "GuideSettings")
+	float GuidFollowStrength = 0.0;
+
+	UPROPERTY(EditAnywhere , Category = "GuideSettings")
+	float GuideSearchDistance = 0.0;
+
+	UPROPERTY(EditAnywhere , Category = "GuideSettings")
+	float GuideGravityOverride = 0.0;
+
+	UPROPERTY(EditAnywhere , Category = "GuideSettings")
+	bool bUse = false;
+};
+
+UENUM()
+enum class EPVAbscissionMode : uint8
+{
+	Kill,
+	Degrade
+};
+
+USTRUCT()
+struct FPVSenescenceParams
+{
+	GENERATED_BODY()
+
+	UPROPERTY(EditAnywhere , Category = "SenescenceSettings" , meta=(UIMin=1, UIMax=50))
+	int AbscissionMin = 1;
+
+	UPROPERTY(EditAnywhere , Category = "SenescenceSettings" , meta=(UIMin=1, UIMax=50))
+	int AbscissionMax = 5;
+
+	UPROPERTY(EditAnywhere , Category = "SenescenceSettings" , meta=(UIMin=0.0f, UIMax=1.0f))
+	float ScaleRetention = 0;
+};
+
+USTRUCT()
+struct FPVAgeSenescenceParams : public FPVSenescenceParams
+{
+	GENERATED_BODY()
+
+	UPROPERTY(EditAnywhere , Category = "SenescenceSettings" , meta=(UIMin=1, UIMax=50))
+	int Threshold = 20;
+
+	UPROPERTY(EditAnywhere , Category = "SenescenceSettings")
+	EPVAbscissionMode Mode = EPVAbscissionMode::Kill;
+
+	UPROPERTY(EditAnywhere , Category = "SenescenceSettings" , meta=(UIMin=0.0f, UIMax=1.0f))
+	float DegradationAmount = 0;
+};
+
+USTRUCT()
+struct FPVLightSenescenceParams : public FPVSenescenceParams
+{
+	GENERATED_BODY()
+
+	UPROPERTY(EditAnywhere , Category = "SenescenceSettings" , meta=(UIMin=0.0f, UIMax=1.0f))
+	float Threshold = 0.33f;
+};
diff --git a/Engine/Plugins/Experimental/ProceduralVegetationEditor/Source/ProceduralVegetation/Private/DataTypes/PVGrowerSettingsData.cpp b/Engine/Plugins/Experimental/ProceduralVegetationEditor/Source/ProceduralVegetation/Private/DataTypes/PVGrowerSettingsData.cpp
new file mode 100644
--- /dev/null
+++ b/Engine/Plugins/Experimental/ProceduralVegetationEditor/Source/ProceduralVegetation/Private/DataTypes/PVGrowerSettingsData.cpp
@@ -0,0 +1,16 @@
+// Copyright Epic Games, Inc. All Rights Reserved.
+
+#include "DataTypes/PVGrowerSettingsData.h"
+#include "PCGContext.h"
+#include "Data/PCGPointData.h"
+
+PCG_DEFINE_TYPE_INFO(FPVDataTypeInfoGrowerPhyllotaxy, UPVGrowerPhyllotaxyData);
+PCG_DEFINE_TYPE_INFO(FPVDataTypeInfoGrowerAxial, UPVGrowerAxialData);
+PCG_DEFINE_TYPE_INFO(FPVDataTypeInfoGrowerLateral, UPVGrowerLateralData);
+PCG_DEFINE_TYPE_INFO(FPVDataTypeInfoGrowerTrunkGrowth, UPVGrowerTrunkGrowthData);
+PCG_DEFINE_TYPE_INFO(FPVDataTypeInfoGrowerPhototropism, UPVGrowerPhototropismData);
+PCG_DEFINE_TYPE_INFO(FPVDataTypeInfoGrowerAgeSenescence, UPVGrowerAgeSenescenceData);
+PCG_DEFINE_TYPE_INFO(FPVDataTypeInfoGrowerAngle, UPVGrowerAngleData);
+PCG_DEFINE_TYPE_INFO(FPVDataTypeInfoGrowerLeafGrowth, UPVGrowerLeafGrowthData);
+PCG_DEFINE_TYPE_INFO(FPVDataTypeInfoGrowerBranching, UPVGrowerBranchingData);
+PCG_DEFINE_TYPE_INFO(FPVDataTypeInfoGrowerGravity, UPVGrowerGravityData);
\ No newline at end of file
diff --git a/Engine/Plugins/Experimental/ProceduralVegetationEditor/Source/ProceduralVegetation/Private/DataTypes/PVGrowerSettingsData.h b/Engine/Plugins/Experimental/ProceduralVegetationEditor/Source/ProceduralVegetation/Private/DataTypes/PVGrowerSettingsData.h
new file mode 100644
--- /dev/null
+++ b/Engine/Plugins/Experimental/ProceduralVegetationEditor/Source/ProceduralVegetation/Private/DataTypes/PVGrowerSettingsData.h
@@ -0,0 +1,295 @@
+// Copyright Epic Games, Inc. All Rights Reserved.
+
+#pragma once
+
+#include "PCGData.h"
+#include "DataTypes/PVGrowerParams.h"
+#include "Serialization/ArchiveCrc32.h"
+#include "Data/Registry/PCGDataType.h"
+#include "PVGrowerSettingsData.generated.h"
+
+#define PV_GROWER_SETTINGS_OVERRIDE_CRC(Params)\
+virtual void AddToCrc(FArchiveCrc32& Ar, bool bFullDataCrc) const override\
+{\
+	Super::AddToCrc(Ar, bFullDataCrc); \
+	if (bFullDataCrc)\
+	{\
+		FString ClassName = StaticClass()->GetPathName();\
+		Ar << ClassName;\
+		Ar << Params;\
+	}\
+	else\
+	{\
+		AddUIDToCrc(Ar);\
+	}\
+}\
+
+USTRUCT()
+struct FPVDataTypeInfoGrowerPhyllotaxy : public FPCGDataTypeInfo
+{
+	GENERATED_BODY()
+
+	PCG_DECLARE_TYPE_INFO(PROCEDURALVEGETATION_API);
+
+#if WITH_EDITOR
+	// If the type should be hidden to the user.
+	virtual bool Hidden() const override { return true; };
+#endif // WITH_EDITOR
+};
+
+UCLASS(BlueprintType, ClassGroup = (Procedural))
+class UPVGrowerPhyllotaxyData : public UPCGData
+{
+	GENERATED_BODY()
+
+public:
+	FPVPhyllotaxyParams Params;
+	
+	//~ Begin UPCGData interface
+	PCG_ASSIGN_DEFAULT_TYPE_INFO(FPVDataTypeInfoGrowerPhyllotaxy)
+	PV_GROWER_SETTINGS_OVERRIDE_CRC(Params);
+	//~ End UPCGData interface
+};
+
+USTRUCT()
+struct FPVDataTypeInfoGrowerAxial : public FPCGDataTypeInfo
+{
+	GENERATED_BODY()
+
+	PCG_DECLARE_TYPE_INFO(PROCEDURALVEGETATION_API);
+
+#if WITH_EDITOR
+	// If the type should be hidden to the user.
+	virtual bool Hidden() const override { return true; };
+#endif // WITH_EDITOR
+};
+
+UCLASS(BlueprintType, ClassGroup = (Procedural))
+class UPVGrowerAxialData : public UPCGData
+{
+	GENERATED_BODY()
+
+public:
+	FPVAxialParams Params;
+	
+	//~ Begin UPCGData interface
+	PCG_ASSIGN_DEFAULT_TYPE_INFO(FPVDataTypeInfoGrowerAxial)
+	PV_GROWER_SETTINGS_OVERRIDE_CRC(Params);
+	//~ End UPCGData interface
+};
+
+USTRUCT()
+struct FPVDataTypeInfoGrowerLateral : public FPCGDataTypeInfo
+{
+	GENERATED_BODY()
+
+	PCG_DECLARE_TYPE_INFO(PROCEDURALVEGETATION_API);
+
+#if WITH_EDITOR
+	// If the type should be hidden to the user.
+	virtual bool Hidden() const override { return true; };
+#endif // WITH_EDITOR
+};
+
+UCLASS(BlueprintType, ClassGroup = (Procedural))
+class UPVGrowerLateralData : public UPCGData
+{
+	GENERATED_BODY()
+
+public:
+	FPVLateralParams Params;
+	
+	//~ Begin UPCGData interface
+	PCG_ASSIGN_DEFAULT_TYPE_INFO(FPVDataTypeInfoGrowerLateral)
+	PV_GROWER_SETTINGS_OVERRIDE_CRC(Params);
+	//~ End UPCGData interface
+};
+
+USTRUCT()
+struct FPVDataTypeInfoGrowerAgeSenescence : public FPCGDataTypeInfo
+{
+	GENERATED_BODY()
+
+	PCG_DECLARE_TYPE_INFO(PROCEDURALVEGETATION_API);
+
+#if WITH_EDITOR
+	// If the type should be hidden to the user.
+	virtual bool Hidden() const override { return true; };
+#endif // WITH_EDITOR
+};
+
+UCLASS(BlueprintType, ClassGroup = (Procedural))
+class UPVGrowerAgeSenescenceData : public UPCGData
+{
+	GENERATED_BODY()
+
+public:
+	FPVAgeSenescenceParams Params;
+	
+	//~ Begin UPCGData interface
+	PCG_ASSIGN_DEFAULT_TYPE_INFO(FPVDataTypeInfoGrowerAgeSenescence)
+	PV_GROWER_SETTINGS_OVERRIDE_CRC(Params);
+	//~ End UPCGData interface
+};
+
+USTRUCT()
+struct FPVDataTypeInfoGrowerAngle : public FPCGDataTypeInfo
+{
+	GENERATED_BODY()
+
+	PCG_DECLARE_TYPE_INFO(PROCEDURALVEGETATION_API);
+
+#if WITH_EDITOR
+	// If the type should be hidden to the user.
+	virtual bool Hidden() const override { return true; };
+#endif // WITH_EDITOR
+};
+
+UCLASS(BlueprintType, ClassGroup = (Procedural))
+class UPVGrowerAngleData : public UPCGData
+{
+	GENERATED_BODY()
+
+public:
+	FPVRandomAngleParams Params;
+	
+	//~ Begin UPCGData interface
+	PCG_ASSIGN_DEFAULT_TYPE_INFO(FPVDataTypeInfoGrowerAngle)
+	PV_GROWER_SETTINGS_OVERRIDE_CRC(Params);
+	//~ End UPCGData interface
+};
+
+USTRUCT()
+struct FPVDataTypeInfoGrowerBranching : public FPCGDataTypeInfo
+{
+	GENERATED_BODY()
+
+	PCG_DECLARE_TYPE_INFO(PROCEDURALVEGETATION_API);
+
+#if WITH_EDITOR
+	// If the type should be hidden to the user.
+	virtual bool Hidden() const override { return true; };
+#endif // WITH_EDITOR
+};
+
+UCLASS(BlueprintType, ClassGroup = (Procedural))
+class UPVGrowerBranchingData : public UPCGData
+{
+	GENERATED_BODY()
+
+public:
+	FPVBranchingParams Params;
+	
+	//~ Begin UPCGData interface
+	PCG_ASSIGN_DEFAULT_TYPE_INFO(FPVDataTypeInfoGrowerBranching)
+	PV_GROWER_SETTINGS_OVERRIDE_CRC(Params);
+	//~ End UPCGData interface
+};
+
+USTRUCT()
+struct FPVDataTypeInfoGrowerGravity : public FPCGDataTypeInfo
+{
+	GENERATED_BODY()
+
+	PCG_DECLARE_TYPE_INFO(PROCEDURALVEGETATION_API);
+
+#if WITH_EDITOR
+	// If the type should be hidden to the user.
+	virtual bool Hidden() const override { return true; };
+#endif // WITH_EDITOR
+};
+
+UCLASS(BlueprintType, ClassGroup = (Procedural))
+class UPVGrowerGravityData : public UPCGData
+{
+	GENERATED_BODY()
+
+public:
+	FPVGrowerGravityParams Params;
+	
+	//~ Begin UPCGData interface
+	PCG_ASSIGN_DEFAULT_TYPE_INFO(FPVDataTypeInfoGrowerGravity)
+	PV_GROWER_SETTINGS_OVERRIDE_CRC(Params);
+	//~ End UPCGData interface
+};
+
+USTRUCT()
+struct FPVDataTypeInfoGrowerLeafGrowth : public FPCGDataTypeInfo
+{
+	GENERATED_BODY()
+
+	PCG_DECLARE_TYPE_INFO(PROCEDURALVEGETATION_API);
+
+#if WITH_EDITOR
+	// If the type should be hidden to the user.
+	virtual bool Hidden() const override { return true; };
+#endif // WITH_EDITOR
+};
+
+UCLASS(BlueprintType, ClassGroup = (Procedural))
+class UPVGrowerLeafGrowthData : public UPCGData
+{
+	GENERATED_BODY()
+
+public:
+	FPVLeafGrowthParams Params;
+	
+	//~ Begin UPCGData interface
+	PCG_ASSIGN_DEFAULT_TYPE_INFO(FPVDataTypeInfoGrowerLeafGrowth)
+	PV_GROWER_SETTINGS_OVERRIDE_CRC(Params);
+	//~ End UPCGData interface
+};
+
+USTRUCT()
+struct FPVDataTypeInfoGrowerPhototropism : public FPCGDataTypeInfo
+{
+	GENERATED_BODY()
+
+	PCG_DECLARE_TYPE_INFO(PROCEDURALVEGETATION_API);
+
+#if WITH_EDITOR
+	// If the type should be hidden to the user.
+	virtual bool Hidden() const override { return true; };
+#endif // WITH_EDITOR
+};
+
+UCLASS(BlueprintType, ClassGroup = (Procedural))
+class UPVGrowerPhototropismData : public UPCGData
+{
+	GENERATED_BODY()
+
+public:
+	FPVPhototropismParams Params;
+	
+	//~ Begin UPCGData interface
+	PCG_ASSIGN_DEFAULT_TYPE_INFO(FPVDataTypeInfoGrowerPhototropism)
+	PV_GROWER_SETTINGS_OVERRIDE_CRC(Params);
+	//~ End UPCGData interface
+};
+
+USTRUCT()
+struct FPVDataTypeInfoGrowerTrunkGrowth : public FPCGDataTypeInfo
+{
+	GENERATED_BODY()
+
+	PCG_DECLARE_TYPE_INFO(PROCEDURALVEGETATION_API);
+
+#if WITH_EDITOR
+	// If the type should be hidden to the user.
+	virtual bool Hidden() const override { return true; };
+#endif // WITH_EDITOR
+};
+
+UCLASS(BlueprintType, ClassGroup = (Procedural))
+class UPVGrowerTrunkGrowthData : public UPCGData
+{
+	GENERATED_BODY()
+
+public:
+	FPVTrunkGrowthParams Params;
+	
+	//~ Begin UPCGData interface
+	PCG_ASSIGN_DEFAULT_TYPE_INFO(FPVDataTypeInfoGrowerTrunkGrowth)
+	PV_GROWER_SETTINGS_OVERRIDE_CRC(Params);
+	//~ End UPCGData interface
+};
\ No newline at end of file
diff --git a/Engine/Plugins/Experimental/ProceduralVegetationEditor/Source/ProceduralVegetation/Private/DataTypes/PVGrowthData.cpp b/Engine/Plugins/Experimental/ProceduralVegetationEditor/Source/ProceduralVegetation/Private/DataTypes/PVGrowthData.cpp
--- a/Engine/Plugins/Experimental/ProceduralVegetationEditor/Source/ProceduralVegetation/Private/DataTypes/PVGrowthData.cpp
+++ b/Engine/Plugins/Experimental/ProceduralVegetationEditor/Source/ProceduralVegetation/Private/DataTypes/PVGrowthData.cpp
@@ -1,5 +1,7 @@
 // Copyright Epic Games, Inc. All Rights Reserved.
 
 #include "DataTypes/PVGrowthData.h"
+#include "PCGContext.h"
+#include "Data/PCGPointData.h"
 
 PCG_DEFINE_TYPE_INFO(FPVDataTypeInfoGrowth, UPVGrowthData);
\ No newline at end of file
diff --git a/Engine/Plugins/Experimental/ProceduralVegetationEditor/Source/ProceduralVegetation/Private/Implementations/PVGrower.cpp b/Engine/Plugins/Experimental/ProceduralVegetationEditor/Source/ProceduralVegetation/Private/Implementations/PVGrower.cpp
--- a/Engine/Plugins/Experimental/ProceduralVegetationEditor/Source/ProceduralVegetation/Private/Implementations/PVGrower.cpp
+++ b/Engine/Plugins/Experimental/ProceduralVegetationEditor/Source/ProceduralVegetation/Private/Implementations/PVGrower.cpp
@@ -1769,7 +1769,7 @@ void GravityOverride(const FPVGrowerGravityParams& GravityParams, UPVGrowerPoint
 	GravityAttribute.GravityOverride = GravityOverride;
 }
 
-void FPVGrower::ApplyGravityWeight(UPVGrowerData* Skeleton, const FPVGrowerGravityParams& InGravityParams, const FPVLeafGrowthParams& InLeafGrowth, const FPVAxialParams& InAxialParams, const FPVAgeSenescenceSettings& InAgeSenescence)
+void FPVGrower::ApplyGravityWeight(UPVGrowerData* Skeleton, const FPVGrowerGravityParams& InGravityParams, const FPVLeafGrowthParams& InLeafGrowth, const FPVAxialParams& InAxialParams, const FPVAgeSenescenceParams& InAgeSenescence)
 {
 	TArray<FPVPointGravityAttribute> GravityAttributes;
 	//int PointIndex = 0;
@@ -2161,7 +2161,7 @@ void FPVGrower::RemoveBrokenBranches(UPVGrowerData* Skeleton,const TArray<int32>
 	}
 }
 
-void FPVGrower::BranchBreak(UPVGrowerData* Skeleton, const FPVAgeSenescenceSettings& InAgeSenescence, const FPVGrowerGravityParams& InGravityParams)
+void FPVGrower::BranchBreak(UPVGrowerData* Skeleton, const FPVAgeSenescenceParams& InAgeSenescence, const FPVGrowerGravityParams& InGravityParams)
 {
 	TArray<int32> UnbrokenPoints;
 	TArray<int32> BrokenPoints;
@@ -2265,7 +2265,7 @@ void FPVGrower::RemoveBuds(UPVGrowerData* Skeleton, UPVGrowerPrimitive* Primitiv
 //This function sets the last point of the input as broken tip
 //for all the other unbroken points it sets the Ethylene to 1 so it cannot grow further
 //also apical is set to zero for no further apical growth.
-void FPVGrower::SetBrokenTip(UPVGrowerData* Skeleton, TArray<int32> UnbrokenPoints, const FPVAgeSenescenceSettings& InAgeSenescence)
+void FPVGrower::SetBrokenTip(UPVGrowerData* Skeleton, TArray<int32> UnbrokenPoints, const FPVAgeSenescenceParams& InAgeSenescence)
 {
 	if(UnbrokenPoints.IsEmpty())
 	{
@@ -2318,7 +2318,7 @@ TArray<int> FPVGrower::ExtractBroken(const TArray<int32>& Points, const TArray<i
 
 //Checks each point in the branch points and return a list of unbroken points
 //if the gravity pressure is < branch break threshold it is unbroken
-TArray<int32> FPVGrower::FindNonBrokenBuds(UPVGrowerData* Skeleton, TArray<int32> PrimPoints, const FPVAgeSenescenceSettings& InAgeSenescence, const FPVGrowerGravityParams& InGravityParams)
+TArray<int32> FPVGrower::FindNonBrokenBuds(UPVGrowerData* Skeleton, TArray<int32> PrimPoints, const FPVAgeSenescenceParams& InAgeSenescence, const FPVGrowerGravityParams& InGravityParams)
 {
 	TArray<int32> UnbrokenPTs;
 
diff --git a/Engine/Plugins/Experimental/ProceduralVegetationEditor/Source/ProceduralVegetation/Private/Implementations/PVGrower.h b/Engine/Plugins/Experimental/ProceduralVegetationEditor/Source/ProceduralVegetation/Private/Implementations/PVGrower.h
--- a/Engine/Plugins/Experimental/ProceduralVegetationEditor/Source/ProceduralVegetation/Private/Implementations/PVGrower.h
+++ b/Engine/Plugins/Experimental/ProceduralVegetationEditor/Source/ProceduralVegetation/Private/Implementations/PVGrower.h
@@ -2,8 +2,8 @@
 
 #pragma once
 
+#include "DataTypes/PVGrowerParams.h"
 #include "CoreMinimal.h"
-#include "PVGrowerParams.h"
 #include "PVGrowerData.h"
 #include "PVLightDetection.h"
 #include "GeometryCollection/ManagedArrayCollection.h"
@@ -78,10 +78,10 @@ public:
 	FPVRandomAngleParams RandomAngleChild;
 
 	UPROPERTY(EditAnywhere , Category = "Grower")
-	FPVAgeSenescenceSettings AgeSenescence;
+	FPVAgeSenescenceParams AgeSenescence;
 
 	UPROPERTY(EditAnywhere , Category = "Light Senescence")
-	FPVLightSenescenceSettings LightSenescence;
+	FPVLightSenescenceParams LightSenescence;
 
 	UPROPERTY()
 	FPVGuideParams GuideSettings;
@@ -175,7 +175,7 @@ struct FPVGrower
 	
 	static void ApplyDavinciWeight(UPVGrowerData* Skeleton, const FPVGrowerGravityParams& InGravityParams, const FPVLateralParams& InLateralParams);
 	
-	static void ApplyGravityWeight(UPVGrowerData* Skeleton, const FPVGrowerGravityParams& InGravityParams, const FPVLeafGrowthParams& InLeafGrowth, const FPVAxialParams& InAxialParams, const FPVAgeSenescenceSettings& InAgeSenescence);
+	static void ApplyGravityWeight(UPVGrowerData* Skeleton, const FPVGrowerGravityParams& InGravityParams, const FPVLeafGrowthParams& InLeafGrowth, const FPVAxialParams& InAxialParams, const FPVAgeSenescenceParams& InAgeSenescence);
 
 	static void AttributeAccumulation(UPVGrowerData* Skeleton, TArray<FPVPointGravityAttribute>& GravityAttributes, const FPVGrowerGravityParams& InGravityParams, const FPVAxialParams& InAxialParams);
 	
@@ -185,7 +185,7 @@ struct FPVGrower
 
 	static void RemoveBrokenBranches(UPVGrowerData* Skeleton,const TArray<int32>& BrokenPrimitives);
 	
-	static void BranchBreak(UPVGrowerData* Skeleton, const FPVAgeSenescenceSettings& InAgeSenescence, const FPVGrowerGravityParams& InGravityParams);
+	static void BranchBreak(UPVGrowerData* Skeleton, const FPVAgeSenescenceParams& InAgeSenescence, const FPVGrowerGravityParams& InGravityParams);
 
 	static void RemoveBuds(UPVGrowerData* Skeleton, UPVGrowerPrimitive* Primitive, TArray<int32> BrokenBuds, TArray<UPVGrowerPrimitive*>& PrimitivesToRemove);
 
@@ -195,11 +195,11 @@ struct FPVGrower
 
 	static void CleanParentsForBrokenBuds(UPVGrowerData* Skeleton, TArray<UPVGrowerPrimitive*>& RemoveBranches, TArray<int32>& CleanParents);
 
-	static void SetBrokenTip(UPVGrowerData* Skeleton, TArray<int32> UnbrokenPoints, const FPVAgeSenescenceSettings& InAgeSenescence);
+	static void SetBrokenTip(UPVGrowerData* Skeleton, TArray<int32> UnbrokenPoints, const FPVAgeSenescenceParams& InAgeSenescence);
 
 	static TArray<int> ExtractBroken(const TArray<int32>& Points, const TArray<int32>& UnBrokenPoints);
 
-	static TArray<int32> FindNonBrokenBuds(UPVGrowerData* Skeleton, TArray<int32> PrimPoints, const FPVAgeSenescenceSettings& InAgeSenescence, const FPVGrowerGravityParams& InGravityParams);
+	static TArray<int32> FindNonBrokenBuds(UPVGrowerData* Skeleton, TArray<int32> PrimPoints, const FPVAgeSenescenceParams& InAgeSenescence, const FPVGrowerGravityParams& InGravityParams);
 
 	static void FindAdopters(UPVGrowerData* Skeleton, TArray<UPVGrowerPrimitive*>& Adopters);
 	
diff --git a/Engine/Plugins/Experimental/ProceduralVegetationEditor/Source/ProceduralVegetation/Private/Implementations/PVGrowerParams.h b/Engine/Plugins/Experimental/ProceduralVegetationEditor/Source/ProceduralVegetation/Private/Implementations/PVGrowerParams.h
deleted file mode 100644
diff --git a/Engine/Plugins/Experimental/ProceduralVegetationEditor/Source/ProceduralVegetation/Private/Nodes/GrowerSettings/PVGrowerAgeSenescenceSettings.cpp b/Engine/Plugins/Experimental/ProceduralVegetationEditor/Source/ProceduralVegetation/Private/Nodes/GrowerSettings/PVGrowerAgeSenescenceSettings.cpp
new file mode 100644
--- /dev/null
+++ b/Engine/Plugins/Experimental/ProceduralVegetationEditor/Source/ProceduralVegetation/Private/Nodes/GrowerSettings/PVGrowerAgeSenescenceSettings.cpp
@@ -0,0 +1,45 @@
+﻿// Copyright Epic Games, Inc. All Rights Reserved.
+
+#include "PVGrowerAgeSenescenceSettings.h"
+#include "DataTypes/PVGrowthData.h"
+#include "Data/PCGBasePointData.h"
+#include "DataTypes/PVGrowerSettingsData.h"
+#include "Helpers/PCGSettingsHelpers.h"
+
+#define LOCTEXT_NAMESPACE "PVGrowerAgeSenescenceSettings"
+
+#if WITH_EDITOR
+FText UPVGrowerAgeSenescenceSettings::GetDefaultNodeTitle() const
+{ 
+	return LOCTEXT("NodeTitle", "Grower Age Senescence Settings"); 
+}
+#endif
+
+FPCGDataTypeIdentifier UPVGrowerAgeSenescenceSettings::GetOutputPinTypeIdentifier() const
+{
+	return FPCGDataTypeIdentifier{ FPVDataTypeInfoGrowerAgeSenescence::AsId() };
+}
+
+FPCGElementPtr UPVGrowerAgeSenescenceSettings::CreateElement() const
+{
+	return MakeShared<FPVAgeSenescenceElement>();
+}
+
+bool FPVAgeSenescenceElement::ExecuteInternal(FPCGContext* InContext) const
+{
+	TRACE_CPUPROFILER_EVENT_SCOPE(FPVAgeSenescenceElement::Execute);
+
+	check(InContext);
+
+	const UPVGrowerAgeSenescenceSettings* Settings = InContext->GetInputSettings<UPVGrowerAgeSenescenceSettings>();
+	check(Settings);
+	
+	UPVGrowerAgeSenescenceData* OutData = FPCGContext::NewObject_AnyThread<UPVGrowerAgeSenescenceData>(InContext);
+	OutData->Params =  Settings->Params;
+
+	InContext->OutputData.TaggedData.Emplace(OutData);
+	
+	return true;
+}
+
+#undef LOCTEXT_NAMESPACE
\ No newline at end of file
diff --git a/Engine/Plugins/Experimental/ProceduralVegetationEditor/Source/ProceduralVegetation/Private/Nodes/GrowerSettings/PVGrowerAgeSenescenceSettings.h b/Engine/Plugins/Experimental/ProceduralVegetationEditor/Source/ProceduralVegetation/Private/Nodes/GrowerSettings/PVGrowerAgeSenescenceSettings.h
new file mode 100644
--- /dev/null
+++ b/Engine/Plugins/Experimental/ProceduralVegetationEditor/Source/ProceduralVegetation/Private/Nodes/GrowerSettings/PVGrowerAgeSenescenceSettings.h
@@ -0,0 +1,34 @@
+﻿// Copyright Epic Games, Inc. All Rights Reserved.
+#pragma once
+
+#include "CoreMinimal.h"
+#include "PVGrowerBaseSettings.h"
+#include "DataTypes/PVGrowerParams.h"
+#include "Nodes/PVBaseSettings.h"
+#include "PVGrowerAgeSenescenceSettings.generated.h"
+
+
+UCLASS(BlueprintType, ClassGroup = (Procedural))
+class UPVGrowerAgeSenescenceSettings : public UPVGrowerBaseSettings
+{
+	GENERATED_BODY()
+
+public:
+#if WITH_EDITOR
+	virtual FName GetDefaultNodeName() const override { return FName(TEXT("GrowerAgeSenescence")); }
+	virtual FText GetDefaultNodeTitle() const override;
+#endif
+
+protected:
+	virtual FPCGDataTypeIdentifier GetOutputPinTypeIdentifier() const override;
+	virtual FPCGElementPtr CreateElement() const override;
+public:
+	UPROPERTY(EditAnywhere, Category="AgeSenescence", meta=(ShowOnlyInnerProperties, PCG_Overridable))
+	FPVAgeSenescenceParams Params;
+};
+
+class FPVAgeSenescenceElement : public FPVBaseElement
+{
+protected:
+	virtual bool ExecuteInternal(FPCGContext* InContext) const override;
+};
diff --git a/Engine/Plugins/Experimental/ProceduralVegetationEditor/Source/ProceduralVegetation/Private/Nodes/GrowerSettings/PVGrowerAngleSettings.cpp b/Engine/Plugins/Experimental/ProceduralVegetationEditor/Source/ProceduralVegetation/Private/Nodes/GrowerSettings/PVGrowerAngleSettings.cpp
new file mode 100644
--- /dev/null
+++ b/Engine/Plugins/Experimental/ProceduralVegetationEditor/Source/ProceduralVegetation/Private/Nodes/GrowerSettings/PVGrowerAngleSettings.cpp
@@ -0,0 +1,45 @@
+﻿// Copyright Epic Games, Inc. All Rights Reserved.
+
+#include "PVGrowerAngleSettings.h"
+#include "DataTypes/PVGrowthData.h"
+#include "Data/PCGBasePointData.h"
+#include "DataTypes/PVGrowerSettingsData.h"
+#include "Helpers/PCGSettingsHelpers.h"
+
+#define LOCTEXT_NAMESPACE "PVGrowerAngleSettings"
+
+#if WITH_EDITOR
+FText UPVGrowerAngleSettings::GetDefaultNodeTitle() const
+{ 
+	return LOCTEXT("NodeTitle", "Grower Angle Settings"); 
+}
+#endif
+
+FPCGDataTypeIdentifier UPVGrowerAngleSettings::GetOutputPinTypeIdentifier() const
+{
+	return FPCGDataTypeIdentifier{ FPVDataTypeInfoGrowerAngle::AsId() };
+}
+
+FPCGElementPtr UPVGrowerAngleSettings::CreateElement() const
+{
+	return MakeShared<FPVGrowerAngleElement>();
+}
+
+bool FPVGrowerAngleElement::ExecuteInternal(FPCGContext* InContext) const
+{
+	TRACE_CPUPROFILER_EVENT_SCOPE(FPVGrowerAngleElement::Execute);
+
+	check(InContext);
+
+	const UPVGrowerAngleSettings* Settings = InContext->GetInputSettings<UPVGrowerAngleSettings>();
+	check(Settings);
+	
+	UPVGrowerAngleData* OutData = FPCGContext::NewObject_AnyThread<UPVGrowerAngleData>(InContext);
+	OutData->Params =  Settings->Params;
+
+	InContext->OutputData.TaggedData.Emplace(OutData);
+	
+	return true;
+}
+
+#undef LOCTEXT_NAMESPACE
\ No newline at end of file
diff --git a/Engine/Plugins/Experimental/ProceduralVegetationEditor/Source/ProceduralVegetation/Private/Nodes/GrowerSettings/PVGrowerAngleSettings.h b/Engine/Plugins/Experimental/ProceduralVegetationEditor/Source/ProceduralVegetation/Private/Nodes/GrowerSettings/PVGrowerAngleSettings.h
new file mode 100644
--- /dev/null
+++ b/Engine/Plugins/Experimental/ProceduralVegetationEditor/Source/ProceduralVegetation/Private/Nodes/GrowerSettings/PVGrowerAngleSettings.h
@@ -0,0 +1,35 @@
+﻿// Copyright Epic Games, Inc. All Rights Reserved.
+#pragma once
+
+#include "CoreMinimal.h"
+#include "PVGrowerBaseSettings.h"
+#include "DataTypes/PVGrowerParams.h"
+#include "PVGrowerAngleSettings.generated.h"
+
+
+UCLASS(BlueprintType, ClassGroup = (Procedural))
+class UPVGrowerAngleSettings : public UPVGrowerBaseSettings
+{
+	GENERATED_BODY()
+
+public:
+#if WITH_EDITOR
+	virtual FName GetDefaultNodeName() const override { return FName(TEXT("GrowerAngleSettings")); }
+	virtual FText GetDefaultNodeTitle() const override;
+#endif
+
+protected:
+	virtual FPCGDataTypeIdentifier GetOutputPinTypeIdentifier() const override;
+	virtual FPCGElementPtr CreateElement() const override;
+
+public:
+
+	UPROPERTY(EditAnywhere, Category="AngleSettings", meta=(ShowOnlyInnerProperties, PCG_Overridable))
+	FPVRandomAngleParams Params;
+};
+
+class FPVGrowerAngleElement : public FPVBaseElement
+{
+protected:
+	virtual bool ExecuteInternal(FPCGContext* InContext) const override;
+};
diff --git a/Engine/Plugins/Experimental/ProceduralVegetationEditor/Source/ProceduralVegetation/Private/Nodes/GrowerSettings/PVGrowerAxialSettings.cpp b/Engine/Plugins/Experimental/ProceduralVegetationEditor/Source/ProceduralVegetation/Private/Nodes/GrowerSettings/PVGrowerAxialSettings.cpp
new file mode 100644
--- /dev/null
+++ b/Engine/Plugins/Experimental/ProceduralVegetationEditor/Source/ProceduralVegetation/Private/Nodes/GrowerSettings/PVGrowerAxialSettings.cpp
@@ -0,0 +1,45 @@
+﻿// Copyright Epic Games, Inc. All Rights Reserved.
+
+#include "PVGrowerAxialSettings.h"
+#include "DataTypes/PVGrowthData.h"
+#include "Data/PCGBasePointData.h"
+#include "DataTypes/PVGrowerSettingsData.h"
+#include "Helpers/PCGSettingsHelpers.h"
+
+#define LOCTEXT_NAMESPACE "PVGrowerAxialSettings"
+
+#if WITH_EDITOR
+FText UPVGrowerAxialSettings::GetDefaultNodeTitle() const
+{ 
+	return LOCTEXT("NodeTitle", "Grower Axial Settings"); 
+}
+#endif
+
+FPCGDataTypeIdentifier UPVGrowerAxialSettings::GetOutputPinTypeIdentifier() const
+{
+	return FPCGDataTypeIdentifier{ FPVDataTypeInfoGrowerAxial::AsId() };
+}
+
+FPCGElementPtr UPVGrowerAxialSettings::CreateElement() const
+{
+	return MakeShared<FPVGrowerAxialElement>();
+}
+
+bool FPVGrowerAxialElement::ExecuteInternal(FPCGContext* InContext) const
+{
+	TRACE_CPUPROFILER_EVENT_SCOPE(FPVGrowerAxialElement::Execute);
+
+	check(InContext);
+
+	const UPVGrowerAxialSettings* Settings = InContext->GetInputSettings<UPVGrowerAxialSettings>();
+	check(Settings);
+	
+	UPVGrowerAxialData* OutData = FPCGContext::NewObject_AnyThread<UPVGrowerAxialData>(InContext);
+	OutData->Params =  Settings->Params;
+
+	InContext->OutputData.TaggedData.Emplace(OutData);
+	
+	return true;
+}
+
+#undef LOCTEXT_NAMESPACE
\ No newline at end of file
diff --git a/Engine/Plugins/Experimental/ProceduralVegetationEditor/Source/ProceduralVegetation/Private/Nodes/GrowerSettings/PVGrowerAxialSettings.h b/Engine/Plugins/Experimental/ProceduralVegetationEditor/Source/ProceduralVegetation/Private/Nodes/GrowerSettings/PVGrowerAxialSettings.h
new file mode 100644
--- /dev/null
+++ b/Engine/Plugins/Experimental/ProceduralVegetationEditor/Source/ProceduralVegetation/Private/Nodes/GrowerSettings/PVGrowerAxialSettings.h
@@ -0,0 +1,34 @@
+﻿// Copyright Epic Games, Inc. All Rights Reserved.
+#pragma once
+
+#include "CoreMinimal.h"
+#include "PVGrowerBaseSettings.h"
+#include "DataTypes/PVGrowerParams.h"
+#include "PVGrowerAxialSettings.generated.h"
+
+
+UCLASS(BlueprintType, ClassGroup = (Procedural))
+class UPVGrowerAxialSettings : public UPVGrowerBaseSettings
+{
+	GENERATED_BODY()
+
+public:
+#if WITH_EDITOR
+	virtual FName GetDefaultNodeName() const override { return FName(TEXT("GrowerAxialSettings")); }
+	virtual FText GetDefaultNodeTitle() const override;
+#endif
+
+protected:
+	virtual FPCGDataTypeIdentifier GetOutputPinTypeIdentifier() const override;
+	virtual FPCGElementPtr CreateElement() const override;
+
+public:
+	UPROPERTY(EditAnywhere, Category="AxialSettings", meta=(ShowOnlyInnerProperties, PCG_Overridable))
+	FPVAxialParams Params;
+};
+
+class FPVGrowerAxialElement : public FPVBaseElement
+{
+protected:
+	virtual bool ExecuteInternal(FPCGContext* InContext) const override;
+};
diff --git a/Engine/Plugins/Experimental/ProceduralVegetationEditor/Source/ProceduralVegetation/Private/Nodes/GrowerSettings/PVGrowerBaseSettings.cpp b/Engine/Plugins/Experimental/ProceduralVegetationEditor/Source/ProceduralVegetation/Private/Nodes/GrowerSettings/PVGrowerBaseSettings.cpp
new file mode 100644
--- /dev/null
+++ b/Engine/Plugins/Experimental/ProceduralVegetationEditor/Source/ProceduralVegetation/Private/Nodes/GrowerSettings/PVGrowerBaseSettings.cpp
@@ -0,0 +1,22 @@
+﻿// Copyright Epic Games, Inc. All Rights Reserved.
+
+#include "PVGrowerBaseSettings.h"
+
+#define LOCTEXT_NAMESPACE "PVGrowerBaseSettings"
+
+#if WITH_EDITOR
+
+FText UPVGrowerBaseSettings::GetNodeTooltipText() const
+{ 
+	return LOCTEXT("NodeTooltip", 
+		"\n\nPress Ctrl + L to lock/unlock node output"
+	); 
+}
+#endif
+
+TArray<FPCGPinProperties> UPVGrowerBaseSettings::InputPinProperties() const
+{
+	return TArray<FPCGPinProperties>();
+}
+
+#undef LOCTEXT_NAMESPACE
\ No newline at end of file
diff --git a/Engine/Plugins/Experimental/ProceduralVegetationEditor/Source/ProceduralVegetation/Private/Nodes/GrowerSettings/PVGrowerBaseSettings.h b/Engine/Plugins/Experimental/ProceduralVegetationEditor/Source/ProceduralVegetation/Private/Nodes/GrowerSettings/PVGrowerBaseSettings.h
new file mode 100644
--- /dev/null
+++ b/Engine/Plugins/Experimental/ProceduralVegetationEditor/Source/ProceduralVegetation/Private/Nodes/GrowerSettings/PVGrowerBaseSettings.h
@@ -0,0 +1,23 @@
+﻿// Copyright Epic Games, Inc. All Rights Reserved.
+#pragma once
+
+#include "CoreMinimal.h"
+#include "Nodes/PVBaseSettings.h"
+#include "PVGrowerBaseSettings.generated.h"
+
+UCLASS(BlueprintType, Abstract, ClassGroup = (Procedural))
+class UPVGrowerBaseSettings : public UPVBaseSettings
+{
+	GENERATED_BODY()
+
+public:
+#if WITH_EDITOR
+	virtual FName GetDefaultNodeName() const override { return FName(TEXT("UPVGrowerBaseSettings")); }
+	virtual FText GetNodeTooltipText() const override;
+	virtual FLinearColor GetNodeTitleColor() const override { return FLinearColor::Yellow; }
+	virtual TArray<EPVRenderType> GetDefaultRenderType() const override { return TArray{ EPVRenderType::None }; }
+#endif
+
+protected:
+	virtual TArray<FPCGPinProperties> InputPinProperties() const override;
+};
diff --git a/Engine/Plugins/Experimental/ProceduralVegetationEditor/Source/ProceduralVegetation/Private/Nodes/GrowerSettings/PVGrowerBranchingSettings.cpp b/Engine/Plugins/Experimental/ProceduralVegetationEditor/Source/ProceduralVegetation/Private/Nodes/GrowerSettings/PVGrowerBranchingSettings.cpp
new file mode 100644
--- /dev/null
+++ b/Engine/Plugins/Experimental/ProceduralVegetationEditor/Source/ProceduralVegetation/Private/Nodes/GrowerSettings/PVGrowerBranchingSettings.cpp
@@ -0,0 +1,45 @@
+﻿// Copyright Epic Games, Inc. All Rights Reserved.
+
+#include "PVGrowerBranchingSettings.h"
+#include "DataTypes/PVGrowthData.h"
+#include "Data/PCGBasePointData.h"
+#include "DataTypes/PVGrowerSettingsData.h"
+#include "Helpers/PCGSettingsHelpers.h"
+
+#define LOCTEXT_NAMESPACE "PVGrowerBranchingSettings"
+
+#if WITH_EDITOR
+FText UPVGrowerBranchingSettings::GetDefaultNodeTitle() const
+{ 
+	return LOCTEXT("NodeTitle", "Grower Branching Settings"); 
+}
+#endif
+
+FPCGDataTypeIdentifier UPVGrowerBranchingSettings::GetOutputPinTypeIdentifier() const
+{
+	return FPCGDataTypeIdentifier{ FPVDataTypeInfoGrowerBranching::AsId() };
+}
+
+FPCGElementPtr UPVGrowerBranchingSettings::CreateElement() const
+{
+	return MakeShared<FPVGrowerBranchingElement>();
+}
+
+bool FPVGrowerBranchingElement::ExecuteInternal(FPCGContext* InContext) const
+{
+	TRACE_CPUPROFILER_EVENT_SCOPE(FPVGrowerBranchingElement::Execute);
+
+	check(InContext);
+
+	const UPVGrowerBranchingSettings* Settings = InContext->GetInputSettings<UPVGrowerBranchingSettings>();
+	check(Settings);
+	
+	UPVGrowerBranchingData* OutData = FPCGContext::NewObject_AnyThread<UPVGrowerBranchingData>(InContext);
+	OutData->Params =  Settings->Params;
+
+	InContext->OutputData.TaggedData.Emplace(OutData);
+	
+	return true;
+}
+
+#undef LOCTEXT_NAMESPACE
\ No newline at end of file
diff --git a/Engine/Plugins/Experimental/ProceduralVegetationEditor/Source/ProceduralVegetation/Private/Nodes/GrowerSettings/PVGrowerBranchingSettings.h b/Engine/Plugins/Experimental/ProceduralVegetationEditor/Source/ProceduralVegetation/Private/Nodes/GrowerSettings/PVGrowerBranchingSettings.h
new file mode 100644
--- /dev/null
+++ b/Engine/Plugins/Experimental/ProceduralVegetationEditor/Source/ProceduralVegetation/Private/Nodes/GrowerSettings/PVGrowerBranchingSettings.h
@@ -0,0 +1,32 @@
+﻿// Copyright Epic Games, Inc. All Rights Reserved.
+#pragma once
+
+#include "CoreMinimal.h"
+#include "PVGrowerBaseSettings.h"
+#include "DataTypes/PVGrowerParams.h"
+#include "PVGrowerBranchingSettings.generated.h"
+
+UCLASS(BlueprintType, ClassGroup = (Procedural))
+class UPVGrowerBranchingSettings : public UPVGrowerBaseSettings
+{
+	GENERATED_BODY()
+
+public:
+#if WITH_EDITOR
+	virtual FName GetDefaultNodeName() const override { return FName(TEXT("GrowerBranchingSettings")); }
+	virtual FText GetDefaultNodeTitle() const override;
+#endif
+ 
+protected:
+	virtual FPCGDataTypeIdentifier GetOutputPinTypeIdentifier() const override;
+	virtual FPCGElementPtr CreateElement() const override;
+public:
+	UPROPERTY(EditAnywhere, Category="BranchingSettings", meta=(ShowOnlyInnerProperties, PCG_Overridable))
+	FPVBranchingParams Params;
+};
+
+class FPVGrowerBranchingElement : public FPVBaseElement
+{
+protected:
+	virtual bool ExecuteInternal(FPCGContext* InContext) const override;
+};
diff --git a/Engine/Plugins/Experimental/ProceduralVegetationEditor/Source/ProceduralVegetation/Private/Nodes/GrowerSettings/PVGrowerGravitySettings.cpp b/Engine/Plugins/Experimental/ProceduralVegetationEditor/Source/ProceduralVegetation/Private/Nodes/GrowerSettings/PVGrowerGravitySettings.cpp
new file mode 100644
--- /dev/null
+++ b/Engine/Plugins/Experimental/ProceduralVegetationEditor/Source/ProceduralVegetation/Private/Nodes/GrowerSettings/PVGrowerGravitySettings.cpp
@@ -0,0 +1,45 @@
+﻿// Copyright Epic Games, Inc. All Rights Reserved.
+
+#include "PVGrowerGravitySettings.h"
+#include "DataTypes/PVGrowthData.h"
+#include "Data/PCGBasePointData.h"
+#include "DataTypes/PVGrowerSettingsData.h"
+#include "Helpers/PCGSettingsHelpers.h"
+
+#define LOCTEXT_NAMESPACE "PVGrowerGravitySettings"
+
+#if WITH_EDITOR
+FText UPVGrowerGravitySettings::GetDefaultNodeTitle() const
+{ 
+	return LOCTEXT("NodeTitle", "Grower Gravity Settings"); 
+}
+#endif
+
+FPCGDataTypeIdentifier UPVGrowerGravitySettings::GetOutputPinTypeIdentifier() const
+{
+	return FPCGDataTypeIdentifier{ FPVDataTypeInfoGrowerGravity::AsId() };
+}
+
+FPCGElementPtr UPVGrowerGravitySettings::CreateElement() const
+{
+	return MakeShared<FPVGrowerGravityElement>();
+}
+
+bool FPVGrowerGravityElement::ExecuteInternal(FPCGContext* InContext) const
+{
+	TRACE_CPUPROFILER_EVENT_SCOPE(FPVGrowerGravityElement::Execute);
+
+	check(InContext);
+
+	const UPVGrowerGravitySettings* Settings = InContext->GetInputSettings<UPVGrowerGravitySettings>();
+	check(Settings);
+	
+	UPVGrowerGravityData* OutData = FPCGContext::NewObject_AnyThread<UPVGrowerGravityData>(InContext);
+	OutData->Params =  Settings->Params;
+
+	InContext->OutputData.TaggedData.Emplace(OutData);
+	
+	return true;
+}
+
+#undef LOCTEXT_NAMESPACE
\ No newline at end of file
diff --git a/Engine/Plugins/Experimental/ProceduralVegetationEditor/Source/ProceduralVegetation/Private/Nodes/GrowerSettings/PVGrowerGravitySettings.h b/Engine/Plugins/Experimental/ProceduralVegetationEditor/Source/ProceduralVegetation/Private/Nodes/GrowerSettings/PVGrowerGravitySettings.h
new file mode 100644
--- /dev/null
+++ b/Engine/Plugins/Experimental/ProceduralVegetationEditor/Source/ProceduralVegetation/Private/Nodes/GrowerSettings/PVGrowerGravitySettings.h
@@ -0,0 +1,34 @@
+﻿// Copyright Epic Games, Inc. All Rights Reserved.
+#pragma once
+
+#include "CoreMinimal.h"
+#include "PVGrowerBaseSettings.h"
+#include "DataTypes/PVGrowerParams.h"
+#include "PVGrowerGravitySettings.generated.h"
+
+
+UCLASS(BlueprintType, ClassGroup = (Procedural))
+class UPVGrowerGravitySettings : public UPVGrowerBaseSettings
+{
+	GENERATED_BODY()
+
+public:
+#if WITH_EDITOR
+	virtual FName GetDefaultNodeName() const override { return FName(TEXT("GrowerGravitySettings")); }
+	virtual FText GetDefaultNodeTitle() const override;
+#endif
+
+protected:
+	virtual FPCGDataTypeIdentifier GetOutputPinTypeIdentifier() const override;
+	virtual FPCGElementPtr CreateElement() const override;
+public:
+
+	UPROPERTY(EditAnywhere, Category="GravitySetttings", meta=(ShowOnlyInnerProperties, PCG_Overridable))
+	FPVGrowerGravityParams Params;
+};
+
+class FPVGrowerGravityElement : public FPVBaseElement
+{
+protected:
+	virtual bool ExecuteInternal(FPCGContext* InContext) const override;
+};
diff --git a/Engine/Plugins/Experimental/ProceduralVegetationEditor/Source/ProceduralVegetation/Private/Nodes/GrowerSettings/PVGrowerLateralSettings.cpp b/Engine/Plugins/Experimental/ProceduralVegetationEditor/Source/ProceduralVegetation/Private/Nodes/GrowerSettings/PVGrowerLateralSettings.cpp
new file mode 100644
--- /dev/null
+++ b/Engine/Plugins/Experimental/ProceduralVegetationEditor/Source/ProceduralVegetation/Private/Nodes/GrowerSettings/PVGrowerLateralSettings.cpp
@@ -0,0 +1,45 @@
+﻿// Copyright Epic Games, Inc. All Rights Reserved.
+
+#include "PVGrowerLateralSettings.h"
+#include "DataTypes/PVGrowthData.h"
+#include "Data/PCGBasePointData.h"
+#include "DataTypes/PVGrowerSettingsData.h"
+#include "Helpers/PCGSettingsHelpers.h"
+
+#define LOCTEXT_NAMESPACE "PVGrowerLateralSettings"
+
+#if WITH_EDITOR
+FText UPVGrowerLateralSettings::GetDefaultNodeTitle() const
+{ 
+	return LOCTEXT("NodeTitle", "Grower Lateral Settings"); 
+}
+#endif
+
+FPCGDataTypeIdentifier UPVGrowerLateralSettings::GetOutputPinTypeIdentifier() const
+{
+	return FPCGDataTypeIdentifier{ FPVDataTypeInfoGrowerLateral::AsId() };
+}
+
+FPCGElementPtr UPVGrowerLateralSettings::CreateElement() const
+{
+	return MakeShared<FPVGrowerLateralElement>();
+}
+
+bool FPVGrowerLateralElement::ExecuteInternal(FPCGContext* InContext) const
+{
+	TRACE_CPUPROFILER_EVENT_SCOPE(FPVGrowerLateralElement::Execute);
+
+	check(InContext);
+
+	const UPVGrowerLateralSettings* Settings = InContext->GetInputSettings<UPVGrowerLateralSettings>();
+	check(Settings);
+	
+	UPVGrowerLateralData* OutData = FPCGContext::NewObject_AnyThread<UPVGrowerLateralData>(InContext);
+	OutData->Params =  Settings->Params;
+
+	InContext->OutputData.TaggedData.Emplace(OutData);
+	
+	return true;
+}
+
+#undef LOCTEXT_NAMESPACE
\ No newline at end of file
diff --git a/Engine/Plugins/Experimental/ProceduralVegetationEditor/Source/ProceduralVegetation/Private/Nodes/GrowerSettings/PVGrowerLateralSettings.h b/Engine/Plugins/Experimental/ProceduralVegetationEditor/Source/ProceduralVegetation/Private/Nodes/GrowerSettings/PVGrowerLateralSettings.h
new file mode 100644
--- /dev/null
+++ b/Engine/Plugins/Experimental/ProceduralVegetationEditor/Source/ProceduralVegetation/Private/Nodes/GrowerSettings/PVGrowerLateralSettings.h
@@ -0,0 +1,33 @@
+﻿// Copyright Epic Games, Inc. All Rights Reserved.
+#pragma once
+
+#include "CoreMinimal.h"
+#include "PVGrowerBaseSettings.h"
+#include "DataTypes/PVGrowerParams.h"
+#include "PVGrowerLateralSettings.generated.h"
+
+UCLASS(BlueprintType, ClassGroup = (Procedural))
+class UPVGrowerLateralSettings : public UPVGrowerBaseSettings
+{
+	GENERATED_BODY()
+
+public:
+#if WITH_EDITOR
+	virtual FName GetDefaultNodeName() const override { return FName(TEXT("GrowerLateralSettings")); }
+	virtual FText GetDefaultNodeTitle() const override;
+#endif
+
+protected:
+	virtual FPCGDataTypeIdentifier GetOutputPinTypeIdentifier() const override;
+	virtual FPCGElementPtr CreateElement() const override;
+public:
+
+	UPROPERTY(EditAnywhere, Category="LateralSettings", meta=(ShowOnlyInnerProperties, PCG_Overridable))
+	FPVLateralParams Params;
+};
+
+class FPVGrowerLateralElement : public FPVBaseElement
+{
+protected:
+	virtual bool ExecuteInternal(FPCGContext* InContext) const override;
+};
diff --git a/Engine/Plugins/Experimental/ProceduralVegetationEditor/Source/ProceduralVegetation/Private/Nodes/GrowerSettings/PVGrowerLeafGrowthSettings.cpp b/Engine/Plugins/Experimental/ProceduralVegetationEditor/Source/ProceduralVegetation/Private/Nodes/GrowerSettings/PVGrowerLeafGrowthSettings.cpp
new file mode 100644
--- /dev/null
+++ b/Engine/Plugins/Experimental/ProceduralVegetationEditor/Source/ProceduralVegetation/Private/Nodes/GrowerSettings/PVGrowerLeafGrowthSettings.cpp
@@ -0,0 +1,59 @@
+﻿// Copyright Epic Games, Inc. All Rights Reserved.
+
+#include "PVGrowerLeafGrowthSettings.h"
+#include "DataTypes/PVData.h"
+#include "DataTypes/PVGrowthData.h"
+#include "Data/PCGBasePointData.h"
+#include "DataTypes/PVGrowerSettingsData.h"
+#include "GeometryCollection/ManagedArrayCollection.h"
+#include "Helpers/PCGSettingsHelpers.h"
+
+#define LOCTEXT_NAMESPACE "PVGrowerLeafGrowthSettings"
+
+#if WITH_EDITOR
+FText UPVGrowerLeafGrowthSettings::GetDefaultNodeTitle() const
+{ 
+	return LOCTEXT("NodeTitle", "Grower Leaf Growth Settings"); 
+}
+
+FText UPVGrowerLeafGrowthSettings::GetNodeTooltipText() const
+{ 
+	return LOCTEXT("NodeTooltip", 
+		"\n\nPress Ctrl + L to lock/unlock node output"
+	); 
+}
+#endif
+
+TArray<FPCGPinProperties> UPVGrowerLeafGrowthSettings::InputPinProperties() const
+{
+	return TArray<FPCGPinProperties>();
+}
+
+FPCGDataTypeIdentifier UPVGrowerLeafGrowthSettings::GetOutputPinTypeIdentifier() const
+{
+	return FPCGDataTypeIdentifier{ FPVDataTypeInfoGrowerLeafGrowth::AsId() };
+}
+
+FPCGElementPtr UPVGrowerLeafGrowthSettings::CreateElement() const
+{
+	return MakeShared<FPVGrowerLeafGrowthElement>();
+}
+
+bool FPVGrowerLeafGrowthElement::ExecuteInternal(FPCGContext* InContext) const
+{
+	TRACE_CPUPROFILER_EVENT_SCOPE(FPVGrowerLeafGrowthElement::Execute);
+
+	check(InContext);
+
+	const UPVGrowerLeafGrowthSettings* Settings = InContext->GetInputSettings<UPVGrowerLeafGrowthSettings>();
+	check(Settings);
+	
+	UPVGrowerLeafGrowthData* OutData = FPCGContext::NewObject_AnyThread<UPVGrowerLeafGrowthData>(InContext);
+	OutData->Params =  Settings->Params;
+
+	InContext->OutputData.TaggedData.Emplace(OutData);
+	
+	return true;
+}
+
+#undef LOCTEXT_NAMESPACE
\ No newline at end of file
diff --git a/Engine/Plugins/Experimental/ProceduralVegetationEditor/Source/ProceduralVegetation/Private/Nodes/GrowerSettings/PVGrowerLeafGrowthSettings.h b/Engine/Plugins/Experimental/ProceduralVegetationEditor/Source/ProceduralVegetation/Private/Nodes/GrowerSettings/PVGrowerLeafGrowthSettings.h
new file mode 100644
--- /dev/null
+++ b/Engine/Plugins/Experimental/ProceduralVegetationEditor/Source/ProceduralVegetation/Private/Nodes/GrowerSettings/PVGrowerLeafGrowthSettings.h
@@ -0,0 +1,42 @@
+﻿// Copyright Epic Games, Inc. All Rights Reserved.
+#pragma once
+
+#include "CoreMinimal.h"
+#include "DataTypes/PVGrowerParams.h"
+#include "Nodes/PVBaseSettings.h"
+#include "Implementations/PVGravity.h"
+#include "PVGrowerLeafGrowthSettings.generated.h"
+
+
+UCLASS(BlueprintType, ClassGroup = (Procedural))
+class UPVGrowerLeafGrowthSettings : public UPVBaseSettings
+{
+	GENERATED_BODY()
+
+public:
+#if WITH_EDITOR
+	virtual FName GetDefaultNodeName() const override { return FName(TEXT("GrowerLeafGrowthSettings")); }
+	virtual FText GetDefaultNodeTitle() const override;
+	virtual FText GetNodeTooltipText() const override;
+	virtual FLinearColor GetNodeTitleColor() const override { return FLinearColor::Yellow; }
+	
+	virtual TArray<EPVRenderType> GetDefaultRenderType() const override { return TArray{ EPVRenderType::None }; }
+#endif
+
+protected:
+	virtual TArray<FPCGPinProperties> InputPinProperties() const override;
+	virtual FPCGDataTypeIdentifier GetOutputPinTypeIdentifier() const override;
+
+	virtual FPCGElementPtr CreateElement() const override;
+
+public:
+
+	UPROPERTY(EditAnywhere, Category="LeafSettings", meta=(ShowOnlyInnerProperties, PCG_Overridable))
+	FPVLeafGrowthParams Params;
+};
+
+class FPVGrowerLeafGrowthElement : public FPVBaseElement
+{
+protected:
+	virtual bool ExecuteInternal(FPCGContext* InContext) const override;
+};
diff --git a/Engine/Plugins/Experimental/ProceduralVegetationEditor/Source/ProceduralVegetation/Private/Nodes/GrowerSettings/PVGrowerPhototropismSettings.cpp b/Engine/Plugins/Experimental/ProceduralVegetationEditor/Source/ProceduralVegetation/Private/Nodes/GrowerSettings/PVGrowerPhototropismSettings.cpp
new file mode 100644
--- /dev/null
+++ b/Engine/Plugins/Experimental/ProceduralVegetationEditor/Source/ProceduralVegetation/Private/Nodes/GrowerSettings/PVGrowerPhototropismSettings.cpp
@@ -0,0 +1,45 @@
+﻿// Copyright Epic Games, Inc. All Rights Reserved.
+
+#include "PVGrowerPhototropismSettings.h"
+#include "DataTypes/PVGrowthData.h"
+#include "Data/PCGBasePointData.h"
+#include "DataTypes/PVGrowerSettingsData.h"
+#include "Helpers/PCGSettingsHelpers.h"
+
+#define LOCTEXT_NAMESPACE "PVGrowerPhototropismSettings"
+
+#if WITH_EDITOR
+FText UPVGrowerPhototropismSettings::GetDefaultNodeTitle() const
+{ 
+	return LOCTEXT("NodeTitle", "Grower Phototropism Settings"); 
+}
+#endif
+
+FPCGDataTypeIdentifier UPVGrowerPhototropismSettings::GetOutputPinTypeIdentifier() const
+{
+	return FPCGDataTypeIdentifier{ FPVDataTypeInfoGrowerPhototropism::AsId() };
+}
+
+FPCGElementPtr UPVGrowerPhototropismSettings::CreateElement() const
+{
+	return MakeShared<FPVGrowerPhototropismElement>();
+}
+
+bool FPVGrowerPhototropismElement::ExecuteInternal(FPCGContext* InContext) const
+{
+	TRACE_CPUPROFILER_EVENT_SCOPE(FPVGrowerPhototropismElement::Execute);
+
+	check(InContext);
+
+	const UPVGrowerPhototropismSettings* Settings = InContext->GetInputSettings<UPVGrowerPhototropismSettings>();
+	check(Settings);
+	
+	UPVGrowerPhototropismData* OutData = FPCGContext::NewObject_AnyThread<UPVGrowerPhototropismData>(InContext);
+	OutData->Params =  Settings->Params;
+
+	InContext->OutputData.TaggedData.Emplace(OutData);
+	
+	return true;
+}
+
+#undef LOCTEXT_NAMESPACE
\ No newline at end of file
diff --git a/Engine/Plugins/Experimental/ProceduralVegetationEditor/Source/ProceduralVegetation/Private/Nodes/GrowerSettings/PVGrowerPhototropismSettings.h b/Engine/Plugins/Experimental/ProceduralVegetationEditor/Source/ProceduralVegetation/Private/Nodes/GrowerSettings/PVGrowerPhototropismSettings.h
new file mode 100644
--- /dev/null
+++ b/Engine/Plugins/Experimental/ProceduralVegetationEditor/Source/ProceduralVegetation/Private/Nodes/GrowerSettings/PVGrowerPhototropismSettings.h
@@ -0,0 +1,34 @@
+﻿// Copyright Epic Games, Inc. All Rights Reserved.
+#pragma once
+
+#include "CoreMinimal.h"
+#include "PVGrowerBaseSettings.h"
+#include "DataTypes/PVGrowerParams.h"
+#include "PVGrowerPhototropismSettings.generated.h"
+
+UCLASS(BlueprintType, ClassGroup = (Procedural))
+class UPVGrowerPhototropismSettings : public UPVGrowerBaseSettings
+{
+	GENERATED_BODY()
+
+public:
+#if WITH_EDITOR
+	virtual FName GetDefaultNodeName() const override { return FName(TEXT("GrowerPhototropismSettings")); }
+	virtual FText GetDefaultNodeTitle() const override;
+#endif
+
+protected:
+	virtual FPCGDataTypeIdentifier GetOutputPinTypeIdentifier() const override;
+	virtual FPCGElementPtr CreateElement() const override;
+
+public:
+
+	UPROPERTY(EditAnywhere, Category="PhototropismSettings", meta=(ShowOnlyInnerProperties, PCG_Overridable))
+	FPVPhototropismParams Params;
+};
+
+class FPVGrowerPhototropismElement : public FPVBaseElement
+{
+protected:
+	virtual bool ExecuteInternal(FPCGContext* InContext) const override;
+};
diff --git a/Engine/Plugins/Experimental/ProceduralVegetationEditor/Source/ProceduralVegetation/Private/Nodes/GrowerSettings/PVGrowerPhyllotaxySettings.cpp b/Engine/Plugins/Experimental/ProceduralVegetationEditor/Source/ProceduralVegetation/Private/Nodes/GrowerSettings/PVGrowerPhyllotaxySettings.cpp
new file mode 100644
--- /dev/null
+++ b/Engine/Plugins/Experimental/ProceduralVegetationEditor/Source/ProceduralVegetation/Private/Nodes/GrowerSettings/PVGrowerPhyllotaxySettings.cpp
@@ -0,0 +1,45 @@
+﻿// Copyright Epic Games, Inc. All Rights Reserved.
+
+#include "PVGrowerPhyllotaxySettings.h"
+#include "DataTypes/PVGrowthData.h"
+#include "Data/PCGBasePointData.h"
+#include "DataTypes/PVGrowerSettingsData.h"
+#include "Helpers/PCGSettingsHelpers.h"
+
+#define LOCTEXT_NAMESPACE "PVPhyllotaxySettings"
+
+#if WITH_EDITOR
+FText UPVGrowerPhyllotaxySettings::GetDefaultNodeTitle() const
+{ 
+	return LOCTEXT("NodeTitle", "Grower Phyllotaxy Settings"); 
+}
+#endif
+
+FPCGDataTypeIdentifier UPVGrowerPhyllotaxySettings::GetOutputPinTypeIdentifier() const
+{
+	return FPCGDataTypeIdentifier{ FPVDataTypeInfoGrowerPhyllotaxy::AsId() };
+}
+
+FPCGElementPtr UPVGrowerPhyllotaxySettings::CreateElement() const
+{
+	return MakeShared<FPVGrowerPhyllotaxyElement>();
+}
+
+bool FPVGrowerPhyllotaxyElement::ExecuteInternal(FPCGContext* InContext) const
+{
+	TRACE_CPUPROFILER_EVENT_SCOPE(FPVPhylotaxyElement::Execute);
+
+	check(InContext);
+
+	const UPVGrowerPhyllotaxySettings* Settings = InContext->GetInputSettings<UPVGrowerPhyllotaxySettings>();
+	check(Settings);
+	
+	UPVGrowerPhyllotaxyData* OutData = FPCGContext::NewObject_AnyThread<UPVGrowerPhyllotaxyData>(InContext);
+	OutData->Params =  Settings->PhyllotaxyParams;
+
+	InContext->OutputData.TaggedData.Emplace(OutData);
+	
+	return true;
+}
+
+#undef LOCTEXT_NAMESPACE
\ No newline at end of file
diff --git a/Engine/Plugins/Experimental/ProceduralVegetationEditor/Source/ProceduralVegetation/Private/Nodes/GrowerSettings/PVGrowerPhyllotaxySettings.h b/Engine/Plugins/Experimental/ProceduralVegetationEditor/Source/ProceduralVegetation/Private/Nodes/GrowerSettings/PVGrowerPhyllotaxySettings.h
new file mode 100644
--- /dev/null
+++ b/Engine/Plugins/Experimental/ProceduralVegetationEditor/Source/ProceduralVegetation/Private/Nodes/GrowerSettings/PVGrowerPhyllotaxySettings.h
@@ -0,0 +1,33 @@
+﻿// Copyright Epic Games, Inc. All Rights Reserved.
+#pragma once
+
+#include "CoreMinimal.h"
+#include "PVGrowerBaseSettings.h"
+#include "DataTypes/PVGrowerParams.h"
+#include "PVGrowerPhyllotaxySettings.generated.h"
+
+
+UCLASS(BlueprintType, ClassGroup = (Procedural))
+class UPVGrowerPhyllotaxySettings : public UPVGrowerBaseSettings
+{
+	GENERATED_BODY()
+
+public:
+#if WITH_EDITOR
+	virtual FName GetDefaultNodeName() const override { return FName(TEXT("GrowerPhyllotaxySettings")); }
+	virtual FText GetDefaultNodeTitle() const override;
+#endif
+
+protected:
+	virtual FPCGDataTypeIdentifier GetOutputPinTypeIdentifier() const override;
+	virtual FPCGElementPtr CreateElement() const override;
+public:
+	UPROPERTY(EditAnywhere, Category="PhyllotaxySettings", meta=(ShowOnlyInnerProperties, PCG_Overridable))
+	FPVPhyllotaxyParams PhyllotaxyParams;
+};
+
+class FPVGrowerPhyllotaxyElement : public FPVBaseElement
+{
+protected:
+	virtual bool ExecuteInternal(FPCGContext* InContext) const override;
+};
diff --git a/Engine/Plugins/Experimental/ProceduralVegetationEditor/Source/ProceduralVegetation/Private/Nodes/GrowerSettings/PVGrowerTrunkGrowthSettings.cpp b/Engine/Plugins/Experimental/ProceduralVegetationEditor/Source/ProceduralVegetation/Private/Nodes/GrowerSettings/PVGrowerTrunkGrowthSettings.cpp
new file mode 100644
--- /dev/null
+++ b/Engine/Plugins/Experimental/ProceduralVegetationEditor/Source/ProceduralVegetation/Private/Nodes/GrowerSettings/PVGrowerTrunkGrowthSettings.cpp
@@ -0,0 +1,45 @@
+﻿// Copyright Epic Games, Inc. All Rights Reserved.
+
+#include "PVGrowerTrunkGrowthSettings.h"
+#include "DataTypes/PVGrowthData.h"
+#include "Data/PCGBasePointData.h"
+#include "DataTypes/PVGrowerSettingsData.h"
+#include "Helpers/PCGSettingsHelpers.h"
+
+#define LOCTEXT_NAMESPACE "PVGrowerTrunkGrowthSettings"
+
+#if WITH_EDITOR
+FText UPVGrowerTrunkGrowthSettings::GetDefaultNodeTitle() const
+{ 
+	return LOCTEXT("NodeTitle", "Grower Trunk Growth Settings"); 
+}
+#endif
+
+FPCGDataTypeIdentifier UPVGrowerTrunkGrowthSettings::GetOutputPinTypeIdentifier() const
+{
+	return FPCGDataTypeIdentifier{ FPVDataTypeInfoGrowerTrunkGrowth::AsId() };
+}
+
+FPCGElementPtr UPVGrowerTrunkGrowthSettings::CreateElement() const
+{
+	return MakeShared<FPVGrowerTrunkGrowthElement>();
+}
+
+bool FPVGrowerTrunkGrowthElement::ExecuteInternal(FPCGContext* InContext) const
+{
+	TRACE_CPUPROFILER_EVENT_SCOPE(FPVGrowerTrunkGrowthElement::Execute);
+
+	check(InContext);
+
+	const UPVGrowerTrunkGrowthSettings* Settings = InContext->GetInputSettings<UPVGrowerTrunkGrowthSettings>();
+	check(Settings);
+	
+	UPVGrowerTrunkGrowthData* OutData = FPCGContext::NewObject_AnyThread<UPVGrowerTrunkGrowthData>(InContext);
+	OutData->Params =  Settings->Params;
+
+	InContext->OutputData.TaggedData.Emplace(OutData);
+	
+	return true;
+}
+
+#undef LOCTEXT_NAMESPACE
\ No newline at end of file
diff --git a/Engine/Plugins/Experimental/ProceduralVegetationEditor/Source/ProceduralVegetation/Private/Nodes/GrowerSettings/PVGrowerTrunkGrowthSettings.h b/Engine/Plugins/Experimental/ProceduralVegetationEditor/Source/ProceduralVegetation/Private/Nodes/GrowerSettings/PVGrowerTrunkGrowthSettings.h
new file mode 100644
--- /dev/null
+++ b/Engine/Plugins/Experimental/ProceduralVegetationEditor/Source/ProceduralVegetation/Private/Nodes/GrowerSettings/PVGrowerTrunkGrowthSettings.h
@@ -0,0 +1,34 @@
+﻿// Copyright Epic Games, Inc. All Rights Reserved.
+#pragma once
+
+#include "CoreMinimal.h"
+#include "PVGrowerBaseSettings.h"
+#include "DataTypes/PVGrowerParams.h"
+#include "PVGrowerTrunkGrowthSettings.generated.h"
+
+UCLASS(BlueprintType, ClassGroup = (Procedural))
+class UPVGrowerTrunkGrowthSettings : public UPVGrowerBaseSettings
+{
+	GENERATED_BODY()
+
+public:
+#if WITH_EDITOR
+	virtual FName GetDefaultNodeName() const override { return FName(TEXT("GrowerTrunkGrowthSettings")); }
+	virtual FText GetDefaultNodeTitle() const override;
+#endif
+
+protected:
+	virtual FPCGDataTypeIdentifier GetOutputPinTypeIdentifier() const override;
+	virtual FPCGElementPtr CreateElement() const override;
+
+public:
+
+	UPROPERTY(EditAnywhere, Category="TrunkGrowthSettings", meta=(ShowOnlyInnerProperties, PCG_Overridable))
+	FPVTrunkGrowthParams Params;
+};
+
+class FPVGrowerTrunkGrowthElement : public FPVBaseElement
+{
+protected:
+	virtual bool ExecuteInternal(FPCGContext* InContext) const override;
+};
diff --git a/Engine/Plugins/Experimental/ProceduralVegetationEditor/Source/ProceduralVegetation/Private/Nodes/PVGrowerSettings.cpp b/Engine/Plugins/Experimental/ProceduralVegetationEditor/Source/ProceduralVegetation/Private/Nodes/PVGrowerSettings.cpp
--- a/Engine/Plugins/Experimental/ProceduralVegetationEditor/Source/ProceduralVegetation/Private/Nodes/PVGrowerSettings.cpp
+++ b/Engine/Plugins/Experimental/ProceduralVegetationEditor/Source/ProceduralVegetation/Private/Nodes/PVGrowerSettings.cpp
@@ -3,9 +3,11 @@
 #include "PVGrowerSettings.h"
 
 #include "ProceduralVegetationModule.h"
+#include "GrowerSettings/PVGrowerPhyllotaxySettings.h"
 #include "DataTypes/PVData.h"
 #include "DataTypes/PVGrowthData.h"
 #include "Data/PCGBasePointData.h"
+#include "DataTypes/PVGrowerSettingsData.h"
 #include "GeometryCollection/ManagedArrayCollection.h"
 #include "Implementations/PVSeedGenerator.h"
 #include "Helpers/PCGSettingsHelpers.h"
@@ -13,6 +15,22 @@
 
 #define LOCTEXT_NAMESPACE "PVGrowerSettings"
 
+namespace GrowerPinsLabels
+{
+	const FName PhyllotaxyInputLabel = TEXT("Phyllotaxy");
+	const FName AxialInputLabel = TEXT("Axial");
+	const FName LateralInputLabel = TEXT("Lateral");
+	const FName PhototropismInputLabel = TEXT("Phototropism");
+	const FName AgeSenescenceInputLabel = TEXT("AgeSenescence");
+	const FName GuideInputLabel = TEXT("Guide");
+	const FName GravityInputLabel = TEXT("Gravity");
+	const FName LeafGrowthInputLabel = TEXT("Leaf Growth");
+	const FName TrunkGrowthInputLabel = TEXT("Trunk Growth");
+	const FName BranchingInputLabel = TEXT("Branching");
+	const FName AngleInputLabel = TEXT("Angle");
+}
+
+
 #if WITH_EDITOR
 FText UPVGrowerSettings::GetDefaultNodeTitle() const
 { 
@@ -48,6 +66,47 @@ TArray<FPCGPinProperties> UPVGrowerSettings::InputPinProperties() const
 	Pin.SetOverrideOrUserParamPin();
 	Pin.bAllowMultipleData = false;
 
+
+	FPCGPinProperties& PhyllotaxyPin = Properties.Emplace_GetRef(GrowerPinsLabels::PhyllotaxyInputLabel, FPCGDataTypeIdentifier{ FPVDataTypeInfoGrowerPhyllotaxy::AsId()});
+	PhyllotaxyPin.SetAllowMultipleConnections(false);
+	PhyllotaxyPin.SetAdvancedPin();
+
+	FPCGPinProperties& AxialPin = Properties.Emplace_GetRef(GrowerPinsLabels::AxialInputLabel, FPCGDataTypeIdentifier{ FPVDataTypeInfoGrowerAxial::AsId()});
+	AxialPin.SetAllowMultipleConnections(false);
+	AxialPin.SetAdvancedPin();
+
+	FPCGPinProperties& PhototropismPin = Properties.Emplace_GetRef(GrowerPinsLabels::PhototropismInputLabel, FPCGDataTypeIdentifier{ FPVDataTypeInfoGrowerPhototropism::AsId()});
+	PhototropismPin.SetAllowMultipleConnections(false);
+	PhototropismPin.SetAdvancedPin();
+
+	FPCGPinProperties& LateralPin = Properties.Emplace_GetRef(GrowerPinsLabels::LateralInputLabel, FPCGDataTypeIdentifier{ FPVDataTypeInfoGrowerLateral::AsId()});
+	LateralPin.SetAllowMultipleConnections(false);
+	LateralPin.SetAdvancedPin();
+
+	FPCGPinProperties& BranchingPin = Properties.Emplace_GetRef(GrowerPinsLabels::BranchingInputLabel, FPCGDataTypeIdentifier{ FPVDataTypeInfoGrowerBranching::AsId()});
+	BranchingPin.SetAllowMultipleConnections(false);
+	BranchingPin.SetAdvancedPin();
+
+	FPCGPinProperties& AgeSenescencePin = Properties.Emplace_GetRef(GrowerPinsLabels::AgeSenescenceInputLabel, FPCGDataTypeIdentifier{ FPVDataTypeInfoGrowerAgeSenescence::AsId()});
+	AgeSenescencePin.SetAllowMultipleConnections(false);
+	AgeSenescencePin.SetAdvancedPin();
+
+	FPCGPinProperties& AngePin = Properties.Emplace_GetRef(GrowerPinsLabels::AngleInputLabel, FPCGDataTypeIdentifier{ FPVDataTypeInfoGrowerAngle::AsId()});
+	AngePin.SetAllowMultipleConnections(false);
+	AngePin.SetAdvancedPin();
+
+	FPCGPinProperties& LeafGrowthPin = Properties.Emplace_GetRef(GrowerPinsLabels::LeafGrowthInputLabel, FPCGDataTypeIdentifier{ FPVDataTypeInfoGrowerLeafGrowth::AsId()});
+	LeafGrowthPin.SetAllowMultipleConnections(false);
+	LeafGrowthPin.SetAdvancedPin();
+
+	FPCGPinProperties& GravityPin = Properties.Emplace_GetRef(GrowerPinsLabels::GravityInputLabel, FPCGDataTypeIdentifier{ FPVDataTypeInfoGrowerGravity::AsId()});
+	GravityPin.SetAllowMultipleConnections(false);
+	GravityPin.SetAdvancedPin();
+	
+	FPCGPinProperties& TrunkGrowthPin = Properties.Emplace_GetRef(GrowerPinsLabels::TrunkGrowthInputLabel, FPCGDataTypeIdentifier{ FPVDataTypeInfoGrowerTrunkGrowth::AsId()});
+	TrunkGrowthPin.SetAllowMultipleConnections(false);
+	TrunkGrowthPin.SetAdvancedPin();
+	
 	return Properties;
 }
 
@@ -56,6 +115,7 @@ FPCGDataTypeIdentifier UPVGrowerSettings::GetInputPinTypeIdentifier() const
 	return FPCGDataTypeIdentifier{ FPVDataTypeInfoGrowth::AsId() };
 }
 
+
 FPCGDataTypeIdentifier UPVGrowerSettings::GetOutputPinTypeIdentifier() const
 {
 	return FPCGDataTypeIdentifier{ FPVDataTypeInfoGrowth::AsId() };
@@ -72,10 +132,90 @@ bool FPVGrowerElement::ExecuteInternal(FPCGContext* InContext) const
 
 	check(InContext);
 
-	const UPVGrowerSettings* Settings = InContext->GetInputSettings<UPVGrowerSettings>();
-	
-	if (Settings)
+	if (const UPVGrowerSettings* Settings = InContext->GetInputSettings<UPVGrowerSettings>())
 	{
+		FPVGrowerParams GrowerParams = Settings->GrowerParams;
+		
+		for (const FPCGTaggedData& TaggedData : InContext->InputData.GetInputsByPin(GrowerPinsLabels::PhyllotaxyInputLabel))
+		{
+			if (const UPVGrowerPhyllotaxyData* PhyllotaxyData = Cast<UPVGrowerPhyllotaxyData>(TaggedData.Data))
+			{
+				GrowerParams.Phyllotaxy = PhyllotaxyData->Params;
+			}
+		}
+
+		for (const FPCGTaggedData& TaggedData : InContext->InputData.GetInputsByPin(GrowerPinsLabels::BranchingInputLabel))
+		{
+			if (const UPVGrowerBranchingData* Data = Cast<UPVGrowerBranchingData>(TaggedData.Data))
+			{
+				GrowerParams.Branching = Data->Params;
+			}
+		}
+
+		for (const FPCGTaggedData& TaggedData : InContext->InputData.GetInputsByPin(GrowerPinsLabels::LateralInputLabel))
+		{
+			if (const UPVGrowerLateralData* Data = Cast<UPVGrowerLateralData>(TaggedData.Data))
+			{
+				GrowerParams.LateralParams = Data->Params;
+			}
+		}
+
+		for (const FPCGTaggedData& TaggedData : InContext->InputData.GetInputsByPin(GrowerPinsLabels::TrunkGrowthInputLabel))
+		{
+			if (const UPVGrowerTrunkGrowthData* Data = Cast<UPVGrowerTrunkGrowthData>(TaggedData.Data))
+			{
+				GrowerParams.TrunkGrowth = Data->Params;
+			}
+		}
+
+		for (const FPCGTaggedData& TaggedData : InContext->InputData.GetInputsByPin(GrowerPinsLabels::AgeSenescenceInputLabel))
+		{
+			if (const UPVGrowerAgeSenescenceData* Data = Cast<UPVGrowerAgeSenescenceData>(TaggedData.Data))
+			{
+				GrowerParams.AgeSenescence = Data->Params;
+			}
+		}
+
+		for (const FPCGTaggedData& TaggedData : InContext->InputData.GetInputsByPin(GrowerPinsLabels::AngleInputLabel))
+		{
+			if (const UPVGrowerAngleData* Data = Cast<UPVGrowerAngleData>(TaggedData.Data))
+			{
+				GrowerParams.RandomAngle = Data->Params;
+			}
+		}
+
+		for (const FPCGTaggedData& TaggedData : InContext->InputData.GetInputsByPin(GrowerPinsLabels::GravityInputLabel))
+		{
+			if (const UPVGrowerGravityData* Data = Cast<UPVGrowerGravityData>(TaggedData.Data))
+			{
+				GrowerParams.GravityParams = Data->Params;
+			}
+		}
+
+		for (const FPCGTaggedData& TaggedData : InContext->InputData.GetInputsByPin(GrowerPinsLabels::AxialInputLabel))
+		{
+			if (const UPVGrowerAxialData* Data = Cast<UPVGrowerAxialData>(TaggedData.Data))
+			{
+				GrowerParams.AxialParams = Data->Params;
+			}
+		}
+
+		for (const FPCGTaggedData& TaggedData : InContext->InputData.GetInputsByPin(GrowerPinsLabels::PhototropismInputLabel))
+		{
+			if (const UPVGrowerPhototropismData* Data = Cast<UPVGrowerPhototropismData>(TaggedData.Data))
+			{
+				GrowerParams.Phototropism = Data->Params;
+			}
+		}
+
+		for (const FPCGTaggedData& TaggedData : InContext->InputData.GetInputsByPin(GrowerPinsLabels::LeafGrowthInputLabel))
+		{
+			if (const UPVGrowerLeafGrowthData* Data = Cast<UPVGrowerLeafGrowthData>(TaggedData.Data))
+			{
+				GrowerParams.LeafGrowth = Data->Params;
+			}
+		}
+
 		const TArray<FPCGTaggedData>& Inputs = InContext->InputData.GetInputsByPin(PCGPinConstants::DefaultInputLabel);
 		if (Inputs.IsValidIndex(0))
 		{
@@ -85,7 +225,7 @@ bool FPVGrowerElement::ExecuteInternal(FPCGContext* InContext) const
 
 				FManagedArrayCollection OutCollection;
 		
-				FPVGrower::Grow(SkeletonCollection, Settings->GrowerParams, OutCollection);
+				FPVGrower::Grow(SkeletonCollection, GrowerParams, OutCollection);
 			
 				UPVGrowthData* OutManagedArrayCollectionData = FPCGContext::NewObject_AnyThread<UPVGrowthData>(InContext);
 				OutManagedArrayCollectionData->Initialize(MoveTemp(OutCollection));
@@ -102,7 +242,7 @@ bool FPVGrowerElement::ExecuteInternal(FPCGContext* InContext) const
 			FManagedArrayCollection OutCollection;
 		
 			FPVSeedPoint SeedPoint;
-			FPVGrower::Grow(SeedPoint, Settings->GrowerParams, OutCollection);
+			FPVGrower::Grow(SeedPoint, GrowerParams, OutCollection);
 			
 			UPVGrowthData* OutManagedArrayCollectionData = FPCGContext::NewObject_AnyThread<UPVGrowthData>(InContext);
 			OutManagedArrayCollectionData->Initialize(MoveTemp(OutCollection));
diff --git a/Engine/Plugins/Experimental/ProceduralVegetationEditor/Source/ProceduralVegetation/Private/Nodes/PVGrowerSettings.h b/Engine/Plugins/Experimental/ProceduralVegetationEditor/Source/ProceduralVegetation/Private/Nodes/PVGrowerSettings.h
--- a/Engine/Plugins/Experimental/ProceduralVegetationEditor/Source/ProceduralVegetation/Private/Nodes/PVGrowerSettings.h
+++ b/Engine/Plugins/Experimental/ProceduralVegetationEditor/Source/ProceduralVegetation/Private/Nodes/PVGrowerSettings.h
@@ -1,6 +1,8 @@
 ﻿// Copyright Epic Games, Inc. All Rights Reserved.
 #pragma once
 
+#include <DataTypes/PVGrowerParams.h>
+
 #include "CoreMinimal.h"
 #include "Nodes/PVBaseSettings.h"
 #include "Implementations/PVGrower.h"
@@ -16,7 +18,7 @@ public:
 	virtual FName GetDefaultNodeName() const override { return FName(TEXT("ProceduralVegetationGrower")); }
 	virtual FText GetDefaultNodeTitle() const override;
 	virtual FText GetNodeTooltipText() const override;
-	virtual EPCGSettingsType GetType() const override { return EPCGSettingsType::InputOutput; }
+	virtual EPCGSettingsType GetType() const override { return EPCGSettingsType::Spatial; }
 	virtual FLinearColor GetNodeTitleColor() const override { return FLinearColor::Yellow; }
 
 	virtual TArray<EPVRenderType> GetDefaultRenderType() const override { return TArray{ EPVRenderType::PointData }; }
@@ -32,7 +34,6 @@ protected:
 public:
 	UPROPERTY(EditAnywhere, Category = Settings, meta=( ShowOnlyInnerProperties ))
 	FPVGrowerParams GrowerParams;
-
 };
 
 class FPVGrowerElement : public FPVBaseElement
diff --git a/Engine/Plugins/Experimental/ProceduralVegetationEditor/Source/ProceduralVegetation/Public/DataTypes/PVGrowthData.h b/Engine/Plugins/Experimental/ProceduralVegetationEditor/Source/ProceduralVegetation/Public/DataTypes/PVGrowthData.h
--- a/Engine/Plugins/Experimental/ProceduralVegetationEditor/Source/ProceduralVegetation/Public/DataTypes/PVGrowthData.h
+++ b/Engine/Plugins/Experimental/ProceduralVegetationEditor/Source/ProceduralVegetation/Public/DataTypes/PVGrowthData.h
@@ -3,7 +3,8 @@
 #pragma once
 
 #include "PVData.h"
-
+#include "Nodes/PVObjectInteractionSettings.h"
+#include "Serialization/ArchiveCrc32.h"
 #include "PVGrowthData.generated.h"
 
 USTRUCT()
