[PVE] Assigned categories for grower settings.

Made Growth cycles overridable. By doing that, we can now expose it from the subgraph. We have already exposed pins for other settings, so we had to make them NonOverridable.

Renamed Angle to Random Angle

#Synced-CL 50296570

diff --git a/Engine/Plugins/Experimental/ProceduralVegetationEditor/Source/ProceduralVegetation/Private/DataTypes/PVGrowerParams.h b/Engine/Plugins/Experimental/ProceduralVegetationEditor/Source/ProceduralVegetation/Private/DataTypes/PVGrowerParams.h
--- a/Engine/Plugins/Experimental/ProceduralVegetationEditor/Source/ProceduralVegetation/Private/DataTypes/PVGrowerParams.h
+++ b/Engine/Plugins/Experimental/ProceduralVegetationEditor/Source/ProceduralVegetation/Private/DataTypes/PVGrowerParams.h
@@ -127,25 +127,25 @@ struct FPVLeafGrowthParams
 {
 	GENERATED_BODY()
 
-	UPROPERTY(EditAnywhere, Category = "FLeafGrowth")
+	UPROPERTY(EditAnywhere, Category = "Leaf Growth")
 	float Ethylene = 0.15f;
 
-	UPROPERTY(EditAnywhere, Category = "FLeafGrowth" , meta=(UIMin=0.1f, UIMax=5.0f))
+	UPROPERTY(EditAnywhere, Category = "Leaf Growth" , meta=(UIMin=0.1f, UIMax=5.0f))
 	float StartScale = 0.1f;
 
-	UPROPERTY(EditAnywhere, Category = "FLeafGrowth" , meta=(UIMin=0.1f, UIMax=5.0f))
+	UPROPERTY(EditAnywhere, Category = "Leaf Growth" , meta=(UIMin=0.1f, UIMax=5.0f))
 	float EndScale = 1.0f;
 
-	UPROPERTY(EditAnywhere, Category = "FLeafGrowth" , meta=(UIMin=0, UIMax=10))
+	UPROPERTY(EditAnywhere, Category = "Leaf Growth" , meta=(UIMin=0, UIMax=10))
 	int Density = 3;
 
-	UPROPERTY(EditAnywhere, Category = "FLeafGrowth" , meta=(UIMin=0.0f, UIMax=1.0f))
+	UPROPERTY(EditAnywhere, Category = "Leaf Growth" , meta=(UIMin=0.0f, UIMax=1.0f))
 	float AbscisicAcid = 0;
 
-	UPROPERTY(EditAnywhere, Category = "FLeafGrowth" , meta=(UIMin=0.0f, UIMax=1.0f))
+	UPROPERTY(EditAnywhere, Category = "Leaf Growth" , meta=(UIMin=0.0f, UIMax=1.0f))
 	float EthyleneThreshold = 0.5f;
 
-	UPROPERTY(EditAnywhere, Category = "FLeafGrowth" , meta=(UIMin=0.0f, UIMax=1.0f))
+	UPROPERTY(EditAnywhere, Category = "Leaf Growth" , meta=(UIMin=0.0f, UIMax=1.0f))
 	float LeafAuxinRetention = 0.0;
 };
 
@@ -371,13 +371,13 @@ struct FPVSenescenceParams
 {
 	GENERATED_BODY()
 
-	UPROPERTY(EditAnywhere , Category = "SenescenceSettings" , meta=(UIMin=1, UIMax=50))
+	UPROPERTY(EditAnywhere , Category = "Senescence" , meta=(UIMin=1, UIMax=50))
 	int AbscissionMin = 1;
 
-	UPROPERTY(EditAnywhere , Category = "SenescenceSettings" , meta=(UIMin=1, UIMax=50))
+	UPROPERTY(EditAnywhere , Category = "Senescence" , meta=(UIMin=1, UIMax=50))
 	int AbscissionMax = 5;
 
-	UPROPERTY(EditAnywhere , Category = "SenescenceSettings" , meta=(UIMin=0.0f, UIMax=1.0f))
+	UPROPERTY(EditAnywhere , Category = "Senescence" , meta=(UIMin=0.0f, UIMax=1.0f))
 	float ScaleRetention = 0;
 };
 
@@ -386,10 +386,10 @@ struct FPVAgeSenescenceParams : public FPVSenescenceParams
 {
 	GENERATED_BODY()
 
-	UPROPERTY(EditAnywhere , Category = "SenescenceSettings" , meta=(UIMin=1, UIMax=50))
+	UPROPERTY(EditAnywhere , Category = "Senescence" , meta=(UIMin=1, UIMax=50))
 	int Threshold = 20;
 
-	UPROPERTY(EditAnywhere , Category = "SenescenceSettings")
+	UPROPERTY(EditAnywhere , Category = "Senescence")
 	EPVAbscissionMode Mode = EPVAbscissionMode::Kill;
 
 	UPROPERTY(EditAnywhere , Category = "SenescenceSettings" , meta=(UIMin=0.0f, UIMax=1.0f))
@@ -401,6 +401,6 @@ struct FPVLightSenescenceParams : public FPVSenescenceParams
 {
 	GENERATED_BODY()
 
-	UPROPERTY(EditAnywhere , Category = "SenescenceSettings" , meta=(UIMin=0.0f, UIMax=1.0f))
+	UPROPERTY(EditAnywhere , Category = "Senescence" , meta=(UIMin=0.0f, UIMax=1.0f))
 	float Threshold = 0.33f;
 };
diff --git a/Engine/Plugins/Experimental/ProceduralVegetationEditor/Source/ProceduralVegetation/Private/DataTypes/PVGrowerSettingsData.cpp b/Engine/Plugins/Experimental/ProceduralVegetationEditor/Source/ProceduralVegetation/Private/DataTypes/PVGrowerSettingsData.cpp
--- a/Engine/Plugins/Experimental/ProceduralVegetationEditor/Source/ProceduralVegetation/Private/DataTypes/PVGrowerSettingsData.cpp
+++ b/Engine/Plugins/Experimental/ProceduralVegetationEditor/Source/ProceduralVegetation/Private/DataTypes/PVGrowerSettingsData.cpp
@@ -10,7 +10,7 @@ PCG_DEFINE_TYPE_INFO(FPVDataTypeInfoGrowerLateral, UPVGrowerLateralData);
 PCG_DEFINE_TYPE_INFO(FPVDataTypeInfoGrowerTrunkGrowth, UPVGrowerTrunkGrowthData);
 PCG_DEFINE_TYPE_INFO(FPVDataTypeInfoGrowerPhototropism, UPVGrowerPhototropismData);
 PCG_DEFINE_TYPE_INFO(FPVDataTypeInfoGrowerAgeSenescence, UPVGrowerAgeSenescenceData);
-PCG_DEFINE_TYPE_INFO(FPVDataTypeInfoGrowerAngle, UPVGrowerAngleData);
+PCG_DEFINE_TYPE_INFO(FPVDataTypeInfoGrowerRandomAngle, UPVGrowerRandomAngleData);
 PCG_DEFINE_TYPE_INFO(FPVDataTypeInfoGrowerLeafGrowth, UPVGrowerLeafGrowthData);
 PCG_DEFINE_TYPE_INFO(FPVDataTypeInfoGrowerBranching, UPVGrowerBranchingData);
 PCG_DEFINE_TYPE_INFO(FPVDataTypeInfoGrowerGravity, UPVGrowerGravityData);
\ No newline at end of file
diff --git a/Engine/Plugins/Experimental/ProceduralVegetationEditor/Source/ProceduralVegetation/Private/DataTypes/PVGrowerSettingsData.h b/Engine/Plugins/Experimental/ProceduralVegetationEditor/Source/ProceduralVegetation/Private/DataTypes/PVGrowerSettingsData.h
--- a/Engine/Plugins/Experimental/ProceduralVegetationEditor/Source/ProceduralVegetation/Private/DataTypes/PVGrowerSettingsData.h
+++ b/Engine/Plugins/Experimental/ProceduralVegetationEditor/Source/ProceduralVegetation/Private/DataTypes/PVGrowerSettingsData.h
@@ -133,7 +133,7 @@ public:
 };
 
 USTRUCT()
-struct FPVDataTypeInfoGrowerAngle : public FPCGDataTypeInfo
+struct FPVDataTypeInfoGrowerRandomAngle : public FPCGDataTypeInfo
 {
 	GENERATED_BODY()
 
@@ -146,7 +146,7 @@ struct FPVDataTypeInfoGrowerAngle : public FPCGDataTypeInfo
 };
 
 UCLASS(BlueprintType, ClassGroup = (Procedural))
-class UPVGrowerAngleData : public UPCGData
+class UPVGrowerRandomAngleData : public UPCGData
 {
 	GENERATED_BODY()
 
@@ -154,7 +154,7 @@ public:
 	FPVRandomAngleParams Params;
 	
 	//~ Begin UPCGData interface
-	PCG_ASSIGN_DEFAULT_TYPE_INFO(FPVDataTypeInfoGrowerAngle)
+	PCG_ASSIGN_DEFAULT_TYPE_INFO(FPVDataTypeInfoGrowerRandomAngle)
 	PV_GROWER_SETTINGS_OVERRIDE_CRC(Params);
 	//~ End UPCGData interface
 };
diff --git a/Engine/Plugins/Experimental/ProceduralVegetationEditor/Source/ProceduralVegetation/Private/Implementations/PVGrower.h b/Engine/Plugins/Experimental/ProceduralVegetationEditor/Source/ProceduralVegetation/Private/Implementations/PVGrower.h
--- a/Engine/Plugins/Experimental/ProceduralVegetationEditor/Source/ProceduralVegetation/Private/Implementations/PVGrower.h
+++ b/Engine/Plugins/Experimental/ProceduralVegetationEditor/Source/ProceduralVegetation/Private/Implementations/PVGrower.h
@@ -20,76 +20,76 @@ struct FPVGrowerParams
 	GENERATED_BODY()
 
 public:
-	UPROPERTY(EditAnywhere , Category = "Grower" , meta=(UIMin=1.0f, UIMax=45.0f))
+	UPROPERTY(EditAnywhere , Category = "Basic" , meta=(UIMin=1.0f, UIMax=45.0f, PCG_Overridable))
 	uint32 GrowthCycles = 1;
 	
-	UPROPERTY(EditAnywhere , Category = "Grower")
+	UPROPERTY(EditAnywhere , Category = "Phyllotaxy", meta =(ShowOnlyInnerProperties, PCG_NotOverridable))
 	FPVPhyllotaxyParams Phyllotaxy;
 
-	UPROPERTY(EditAnywhere , Category = "Grower")
+	UPROPERTY(EditAnywhere , Category = "Phyllotaxy", meta =(PCG_NotOverridable))
 	bool bChildPhyllotaxySameAsParent = true;
 
-	UPROPERTY(EditAnywhere , Category = "Grower")
+	UPROPERTY(EditAnywhere , Category = "Phyllotaxy", meta =(PCG_NotOverridable))
 	bool bLeafPhyllotaxySameAsChild = true;
 
-	UPROPERTY(EditAnywhere , Category = "Grower" , meta = (EditCondition = "bChildPhyllotaxySameAsParent == false" , EditConditionHides))
+	UPROPERTY(EditAnywhere , Category = "Phyllotaxy" , meta = (EditCondition = "bChildPhyllotaxySameAsParent == false" , EditConditionHides, PCG_NotOverridable))
 	FPVPhyllotaxyParams PhyllotaxyChild;
 
-	UPROPERTY(EditAnywhere , Category = "Grower")
+	UPROPERTY(EditAnywhere , Category = "Axial", meta =(ShowOnlyInnerProperties, PCG_NotOverridable))
 	FPVAxialParams AxialParams;
 
-	UPROPERTY(EditAnywhere , Category = "Grower")
+	UPROPERTY(EditAnywhere , Category = "Lateral", meta =(ShowOnlyInnerProperties, PCG_NotOverridable))
 	FPVLateralParams LateralParams;
 	
-	UPROPERTY(EditAnywhere , Category = "Grower")
+	UPROPERTY(EditAnywhere , Category = "Trunk Growth", meta =(ShowOnlyInnerProperties, PCG_NotOverridable))
 	FPVTrunkGrowthParams TrunkGrowth;
 
-	UPROPERTY(EditAnywhere , Category = "Grower")
+	UPROPERTY(EditAnywhere , Category = "Branching", meta =(ShowOnlyInnerProperties, PCG_NotOverridable))
 	FPVBranchingParams Branching;
 
-	UPROPERTY(EditAnywhere , Category = "Grower")
+	UPROPERTY(EditAnywhere , Category = "Branching", meta =(PCG_NotOverridable))
 	bool bChildBranchingConditionSameAsParent = true;
 
-	UPROPERTY(EditAnywhere , Category = "Grower" , meta = (EditCondition = "bChildBranchingConditionSameAsParent == false" , EditConditionHides))
+	UPROPERTY(EditAnywhere , Category = "Branching" , meta = (EditCondition = "bChildBranchingConditionSameAsParent == false" , EditConditionHides, PCG_NotOverridable))
 	FPVBranchingParams BranchingChild;
 
-	UPROPERTY(EditAnywhere , Category = "Grower")
+	UPROPERTY(EditAnywhere , Category = "Phototropism", meta =(ShowOnlyInnerProperties, PCG_NotOverridable))
 	FPVPhototropismParams Phototropism;
 
-	UPROPERTY(EditAnywhere , Category = "Grower")
+	UPROPERTY(EditAnywhere , Category = "Phototropism", meta =(PCG_NotOverridable))
 	bool bChildPhototropismSameAsParent = true;
 
-	UPROPERTY(EditAnywhere , Category = "Grower" , meta = (EditCondition = "bChildPhototropismSameAsParent == false" , EditConditionHides))
+	UPROPERTY(EditAnywhere , Category = "Phototropism" , meta = (EditCondition = "bChildPhototropismSameAsParent == false" , EditConditionHides, PCG_NotOverridable))
 	FPVPhototropismParams PhototropismChild;
 
-	UPROPERTY(EditAnywhere , Category = "Grower")
+	UPROPERTY(EditAnywhere , Category = "Gravity", meta =(ShowOnlyInnerProperties, PCG_NotOverridable))
 	FPVGrowerGravityParams GravityParams;
 
-	UPROPERTY(EditAnywhere, Category = "Grower")
+	UPROPERTY(EditAnywhere, Category = "Leaf Growth", meta =(ShowOnlyInnerProperties, PCG_NotOverridable))
 	FPVLeafGrowthParams LeafGrowth;
 
-	UPROPERTY(EditAnywhere , Category = "Grower")
+	UPROPERTY(EditAnywhere , Category = "RandomAngle", meta =(ShowOnlyInnerProperties, PCG_NotOverridable))
 	FPVRandomAngleParams RandomAngle;
 
-	UPROPERTY(EditAnywhere , Category = "Grower")
+	UPROPERTY(EditAnywhere , Category = "RandomAngle", meta =(PCG_NotOverridable))
 	bool bChildRandomAngleSameAsParent = true;
 
-	UPROPERTY(EditAnywhere , Category = "Grower", meta = (EditCondition = "bChildRandomAngleSameAsParent == false" , EditConditionHides))
+	UPROPERTY(EditAnywhere , Category = "RandomAngle", meta = (EditCondition = "bChildRandomAngleSameAsParent == false" , EditConditionHides, PCG_NotOverridable))
 	FPVRandomAngleParams RandomAngleChild;
 
-	UPROPERTY(EditAnywhere , Category = "Grower")
+	UPROPERTY(EditAnywhere , Category = "Age Senescence", meta =(ShowOnlyInnerProperties, PCG_NotOverridable))
 	FPVAgeSenescenceParams AgeSenescence;
 
-	UPROPERTY(EditAnywhere , Category = "Light Senescence")
+	UPROPERTY(EditAnywhere , Category = "Light Senescence", meta =(ShowOnlyInnerProperties, PCG_NotOverridable))
 	FPVLightSenescenceParams LightSenescence;
 
 	UPROPERTY()
 	FPVGuideParams GuideSettings;
 
-	UPROPERTY(EditAnywhere , Category = "Grower")
+	UPROPERTY(EditAnywhere , Category = "Collider", meta =(PCG_NotOverridable))
 	TArray<FPVColliderParams> ColliderSettings; 
 
-	UPROPERTY(EditAnywhere , Category = "Grower")
+	UPROPERTY(EditAnywhere , Category = "Basic", meta =(PCG_NotOverridable))
 	uint32 RandomSeed = 0;
 
 	void GetPhyllotaxy(FPVPhyllotaxyParams& OutPhyllotaxy, int Generation) const;
diff --git a/Engine/Plugins/Experimental/ProceduralVegetationEditor/Source/ProceduralVegetation/Private/Nodes/GrowerSettings/PVGrowerAngleSettings.cpp b/Engine/Plugins/Experimental/ProceduralVegetationEditor/Source/ProceduralVegetation/Private/Nodes/GrowerSettings/PVGrowerAngleSettings.cpp
deleted file mode 100644
diff --git a/Engine/Plugins/Experimental/ProceduralVegetationEditor/Source/ProceduralVegetation/Private/Nodes/GrowerSettings/PVGrowerAngleSettings.h b/Engine/Plugins/Experimental/ProceduralVegetationEditor/Source/ProceduralVegetation/Private/Nodes/GrowerSettings/PVGrowerAngleSettings.h
deleted file mode 100644
diff --git a/Engine/Plugins/Experimental/ProceduralVegetationEditor/Source/ProceduralVegetation/Private/Nodes/GrowerSettings/PVGrowerRandomAngleSettings.cpp b/Engine/Plugins/Experimental/ProceduralVegetationEditor/Source/ProceduralVegetation/Private/Nodes/GrowerSettings/PVGrowerRandomAngleSettings.cpp
new file mode 100644
--- /dev/null
+++ b/Engine/Plugins/Experimental/ProceduralVegetationEditor/Source/ProceduralVegetation/Private/Nodes/GrowerSettings/PVGrowerRandomAngleSettings.cpp
@@ -0,0 +1,45 @@
+﻿// Copyright Epic Games, Inc. All Rights Reserved.
+
+#include "PVGrowerRandomAngleSettings.h"
+#include "DataTypes/PVGrowthData.h"
+#include "Data/PCGBasePointData.h"
+#include "DataTypes/PVGrowerSettingsData.h"
+#include "Helpers/PCGSettingsHelpers.h"
+
+#define LOCTEXT_NAMESPACE "PVGrowerRandomAngleSettings"
+
+#if WITH_EDITOR
+FText UPVGrowerRandomAngleSettings::GetDefaultNodeTitle() const
+{ 
+	return LOCTEXT("NodeTitle", "Grower Random Angle Settings"); 
+}
+#endif
+
+FPCGDataTypeIdentifier UPVGrowerRandomAngleSettings::GetOutputPinTypeIdentifier() const
+{
+	return FPCGDataTypeIdentifier{ FPVDataTypeInfoGrowerRandomAngle::AsId() };
+}
+
+FPCGElementPtr UPVGrowerRandomAngleSettings::CreateElement() const
+{
+	return MakeShared<FPVGrowerAngleElement>();
+}
+
+bool FPVGrowerAngleElement::ExecuteInternal(FPCGContext* InContext) const
+{
+	TRACE_CPUPROFILER_EVENT_SCOPE(FPVGrowerAngleElement::Execute);
+
+	check(InContext);
+
+	const UPVGrowerRandomAngleSettings* Settings = InContext->GetInputSettings<UPVGrowerRandomAngleSettings>();
+	check(Settings);
+	
+	UPVGrowerRandomAngleData* OutData = FPCGContext::NewObject_AnyThread<UPVGrowerRandomAngleData>(InContext);
+	OutData->Params =  Settings->Params;
+
+	InContext->OutputData.TaggedData.Emplace(OutData);
+	
+	return true;
+}
+
+#undef LOCTEXT_NAMESPACE
\ No newline at end of file
diff --git a/Engine/Plugins/Experimental/ProceduralVegetationEditor/Source/ProceduralVegetation/Private/Nodes/GrowerSettings/PVGrowerRandomAngleSettings.h b/Engine/Plugins/Experimental/ProceduralVegetationEditor/Source/ProceduralVegetation/Private/Nodes/GrowerSettings/PVGrowerRandomAngleSettings.h
new file mode 100644
--- /dev/null
+++ b/Engine/Plugins/Experimental/ProceduralVegetationEditor/Source/ProceduralVegetation/Private/Nodes/GrowerSettings/PVGrowerRandomAngleSettings.h
@@ -0,0 +1,35 @@
+﻿// Copyright Epic Games, Inc. All Rights Reserved.
+#pragma once
+
+#include "CoreMinimal.h"
+#include "PVGrowerBaseSettings.h"
+#include "DataTypes/PVGrowerParams.h"
+#include "PVGrowerRandomAngleSettings.generated.h"
+
+
+UCLASS(BlueprintType, ClassGroup = (Procedural))
+class UPVGrowerRandomAngleSettings : public UPVGrowerBaseSettings
+{
+	GENERATED_BODY()
+
+public:
+#if WITH_EDITOR
+	virtual FName GetDefaultNodeName() const override { return FName(TEXT("GrowerRandomAngleSettings")); }
+	virtual FText GetDefaultNodeTitle() const override;
+#endif
+
+protected:
+	virtual FPCGDataTypeIdentifier GetOutputPinTypeIdentifier() const override;
+	virtual FPCGElementPtr CreateElement() const override;
+
+public:
+
+	UPROPERTY(EditAnywhere, Category="RandomAngleSettings", meta=(ShowOnlyInnerProperties, PCG_Overridable))
+	FPVRandomAngleParams Params;
+};
+
+class FPVGrowerAngleElement : public FPVBaseElement
+{
+protected:
+	virtual bool ExecuteInternal(FPCGContext* InContext) const override;
+};
diff --git a/Engine/Plugins/Experimental/ProceduralVegetationEditor/Source/ProceduralVegetation/Private/Nodes/PVGrowerSettings.cpp b/Engine/Plugins/Experimental/ProceduralVegetationEditor/Source/ProceduralVegetation/Private/Nodes/PVGrowerSettings.cpp
--- a/Engine/Plugins/Experimental/ProceduralVegetationEditor/Source/ProceduralVegetation/Private/Nodes/PVGrowerSettings.cpp
+++ b/Engine/Plugins/Experimental/ProceduralVegetationEditor/Source/ProceduralVegetation/Private/Nodes/PVGrowerSettings.cpp
@@ -24,10 +24,10 @@ namespace GrowerPinsLabels
 	const FName AgeSenescenceInputLabel = TEXT("AgeSenescence");
 	const FName GuideInputLabel = TEXT("Guide");
 	const FName GravityInputLabel = TEXT("Gravity");
-	const FName LeafGrowthInputLabel = TEXT("Leaf Growth");
-	const FName TrunkGrowthInputLabel = TEXT("Trunk Growth");
+	const FName LeafGrowthInputLabel = TEXT("LeafGrowth");
+	const FName TrunkGrowthInputLabel = TEXT("TrunkGrowth");
 	const FName BranchingInputLabel = TEXT("Branching");
-	const FName AngleInputLabel = TEXT("Angle");
+	const FName RandomAngleInputLabel = TEXT("RandomAngle");
 }
 
 
@@ -91,9 +91,9 @@ TArray<FPCGPinProperties> UPVGrowerSettings::InputPinProperties() const
 	AgeSenescencePin.SetAllowMultipleConnections(false);
 	AgeSenescencePin.SetAdvancedPin();
 
-	FPCGPinProperties& AngePin = Properties.Emplace_GetRef(GrowerPinsLabels::AngleInputLabel, FPCGDataTypeIdentifier{ FPVDataTypeInfoGrowerAngle::AsId()});
-	AngePin.SetAllowMultipleConnections(false);
-	AngePin.SetAdvancedPin();
+	FPCGPinProperties& RandomAnglePin = Properties.Emplace_GetRef(GrowerPinsLabels::RandomAngleInputLabel, FPCGDataTypeIdentifier{ FPVDataTypeInfoGrowerRandomAngle::AsId()});
+	RandomAnglePin.SetAllowMultipleConnections(false);
+	RandomAnglePin.SetAdvancedPin();
 
 	FPCGPinProperties& LeafGrowthPin = Properties.Emplace_GetRef(GrowerPinsLabels::LeafGrowthInputLabel, FPCGDataTypeIdentifier{ FPVDataTypeInfoGrowerLeafGrowth::AsId()});
 	LeafGrowthPin.SetAllowMultipleConnections(false);
@@ -176,9 +176,9 @@ bool FPVGrowerElement::ExecuteInternal(FPCGContext* InContext) const
 			}
 		}
 
-		for (const FPCGTaggedData& TaggedData : InContext->InputData.GetInputsByPin(GrowerPinsLabels::AngleInputLabel))
+		for (const FPCGTaggedData& TaggedData : InContext->InputData.GetInputsByPin(GrowerPinsLabels::RandomAngleInputLabel))
 		{
-			if (const UPVGrowerAngleData* Data = Cast<UPVGrowerAngleData>(TaggedData.Data))
+			if (const UPVGrowerRandomAngleData* Data = Cast<UPVGrowerRandomAngleData>(TaggedData.Data))
 			{
 				GrowerParams.RandomAngle = Data->Params;
 			}
diff --git a/Engine/Plugins/Experimental/ProceduralVegetationEditor/Source/ProceduralVegetation/Private/Nodes/PVGrowerSettings.h b/Engine/Plugins/Experimental/ProceduralVegetationEditor/Source/ProceduralVegetation/Private/Nodes/PVGrowerSettings.h
--- a/Engine/Plugins/Experimental/ProceduralVegetationEditor/Source/ProceduralVegetation/Private/Nodes/PVGrowerSettings.h
+++ b/Engine/Plugins/Experimental/ProceduralVegetationEditor/Source/ProceduralVegetation/Private/Nodes/PVGrowerSettings.h
@@ -32,7 +32,7 @@ protected:
 	virtual FPCGElementPtr CreateElement() const override;
 
 public:
-	UPROPERTY(EditAnywhere, Category = Settings, meta=( ShowOnlyInnerProperties ))
+	UPROPERTY(EditAnywhere, Category = Settings, meta=( ShowOnlyInnerProperties, PCG_Overridable))
 	FPVGrowerParams GrowerParams;
 };
 
