Split 2D importer output into one pin per plant

This CL changes the behavior of the 2D Importer such that it outputs one collection per plant, instead of combining them all into one big collection.

Since this CL required larger changes in PVImportHelper, I took the opportunity to perform some minor cleanup of this file (re-name namespace, move relevant contents of PVPlantDescription to PVImportHelper).

PVImportHelper changes:
- Split "ComputeBranchHierarchies" into "ComputeDynamicMeshBranchHierarchies" & "ConvertDynamicMeshBranchHierarchiesToBranchDescription" to allow the importers to map the output to correct part of the mesh
- Rename namespace "Import::Helper" to "ImportHelper"
- Remove FPVFoliageDescription as it will not be used for 5.8
- Move FPVBranchHierarchyDescription into PVImportHelper (also remove USTRUCT and PROPERTIES from FPVBranchHierarchyDescription)
- Change name and signature of GenerateGrowthDataFromPlantDescription to allow passing of an array view (to avoid copying the input branch descriptions at the call-site)

Importer_Texture2D changes:
- Remove debug states which are no longer needed due to better visualization
- Output multiple collections
- VisualizationCollection: (Holds the visualization which is shown to the user to help with tweaking parameters)
- DebugCollection: (used for development, should probably be removed before 5.8 release)
- CombinedOutputCollection: Multi-plant collection which contains all plants which were successfully extracted from the image
- PlantCollections: An array of collections holding one plant each, maps 1-1 to the PlantSettings property in the FPVImportTexture2DParams

#Synced-CL 50294763

diff --git a/Engine/Plugins/Experimental/ProceduralVegetationEditor/Source/ProceduralVegetation/Private/Nodes/PVTexture2DImporterSettings.cpp b/Engine/Plugins/Experimental/ProceduralVegetationEditor/Source/ProceduralVegetation/Private/Nodes/PVTexture2DImporterSettings.cpp
--- a/Engine/Plugins/Experimental/ProceduralVegetationEditor/Source/ProceduralVegetation/Private/Nodes/PVTexture2DImporterSettings.cpp
+++ b/Engine/Plugins/Experimental/ProceduralVegetationEditor/Source/ProceduralVegetation/Private/Nodes/PVTexture2DImporterSettings.cpp
@@ -16,37 +16,22 @@
 
 #define LOCTEXT_NAMESPACE "PVTexture2DImporterSettings"
 
+namespace PVTexture2DImporterSettings
+{
+	const static FName CombinedOutputPinName = TEXT("CombinedOutput");
+
+	FName GetPlantPinName(int32 PlantIndex)
+	{
+		return *FString::Printf(TEXT("Plant %d"), PlantIndex);
+	}
+}
+
 UPVTexture2DImporterSettings::UPVTexture2DImporterSettings()
 {
 #if WITH_EDITORONLY_DATA
 	static ConstructorHelpers::FObjectFinder<UMaterialInterface> TextureVisualizationMaterialRef(TEXT("/Script/Engine.Material'/Engine/Engine_MI_Shaders/M_Shader_SimpleTranslucent.M_Shader_SimpleTranslucent'"));
 	TextureVisualizationMaterial = TextureVisualizationMaterialRef.Object;
 
-	FPVVisualizationSettings& LineSegmentVisualization = DebugVisualizationSettings.VisualizationSettings.AddDefaulted_GetRef();
-	LineSegmentVisualization.DebugType = EPVDebugType::Custom;
-	LineSegmentVisualization.AttributeToFilter = TEXT("LineSegment");
-	LineSegmentVisualization.CustomGroupToFilter = TEXT("LineSegments");
-	LineSegmentVisualization.VisualizationMode = EPVDebugValueVisualizationMode::Curve;
-
-	FPVVisualizationSettings& PerimeterCurveVisualization = DebugVisualizationSettings.VisualizationSettings.AddDefaulted_GetRef();
-	PerimeterCurveVisualization.DebugType = EPVDebugType::Custom;
-	PerimeterCurveVisualization.bShowAnchorPoints = false;
-	PerimeterCurveVisualization.AttributeToFilter = TEXT("PerimeterCurve");
-	PerimeterCurveVisualization.CustomGroupToFilter = TEXT("PerimeterCurves");
-	PerimeterCurveVisualization.VisualizationMode = EPVDebugValueVisualizationMode::Curve;
-
-	FPVVisualizationSettings& EndPointVisualization = DebugVisualizationSettings.VisualizationSettings.AddDefaulted_GetRef();
-	EndPointVisualization.DebugType = EPVDebugType::Custom;
-	EndPointVisualization.bShowAnchorPoints = false;
-	EndPointVisualization.bDrawPointAsMesh = false;
-	EndPointVisualization.bUsePivotAsPosition = true;
-	EndPointVisualization.AttributeToFilter = TEXT("EndPoint");
-	EndPointVisualization.CustomGroupToFilter = TEXT("VertexAttributes");
-	EndPointVisualization.CustomPivotPositionAttributeName = "Vertex";
-	EndPointVisualization.VisualizationMode = EPVDebugValueVisualizationMode::Point;
-	EndPointVisualization.Color = FColor(39, 0, 255);
-	EndPointVisualization.GizmoScale = 8.000000;
-
 	FPVVisualizationSettings& BranchCurvesVisualization = DebugVisualizationSettings.VisualizationSettings.AddDefaulted_GetRef();
 	BranchCurvesVisualization.DebugType = EPVDebugType::Custom;
 	BranchCurvesVisualization.bShowAnchorPoints = false;
@@ -55,15 +40,6 @@ UPVTexture2DImporterSettings::UPVTexture2DImporterSettings()
 	BranchCurvesVisualization.CustomGroupToFilter = TEXT("BranchCurves");
 	BranchCurvesVisualization.VisualizationMode = EPVDebugValueVisualizationMode::Curve;
 
-	FPVVisualizationSettings& RootPointsVisualization = DebugVisualizationSettings.VisualizationSettings.AddDefaulted_GetRef();
-	RootPointsVisualization.DebugType = EPVDebugType::Custom;
-	RootPointsVisualization.bShowAnchorPoints = false;
-	RootPointsVisualization.bDrawPointAsMesh = false;
-	RootPointsVisualization.AttributeToFilter = TEXT("Point");
-	RootPointsVisualization.CustomGroupToFilter = TEXT("InputRootPoints");
-	RootPointsVisualization.VisualizationMode = EPVDebugValueVisualizationMode::Point;
-	RootPointsVisualization.GizmoScale = 10.000000;
-
 	FPVVisualizationSettings& PixelSamplingVisualization = DebugVisualizationSettings.VisualizationSettings.AddDefaulted_GetRef();
 	PixelSamplingVisualization.DebugType = EPVDebugType::Custom;
 	PixelSamplingVisualization.bShowAnchorPoints = false;
@@ -94,6 +70,21 @@ TArray<FPCGPinProperties> UPVTexture2DImporterSettings::InputPinProperties() con
 	return TArray<FPCGPinProperties>();
 }
 
+TArray<FPCGPinProperties> UPVTexture2DImporterSettings::OutputPinProperties() const
+{
+	TArray<FPCGPinProperties> Properties;
+	
+	FPCGPinProperties& CombinedOutputPin = Properties.Emplace_GetRef(PVTexture2DImporterSettings::CombinedOutputPinName, GetOutputPinTypeIdentifier());
+	CombinedOutputPin.bAllowMultipleData = false;
+
+	for (int32 i = 0; i < Params.PlantSettings.Num(); i++)
+	{
+		FPCGPinProperties& Pin = Properties.Emplace_GetRef(PVTexture2DImporterSettings::GetPlantPinName(i), GetOutputPinTypeIdentifier());
+		Pin.bAllowMultipleData = false;
+	}
+	return Properties;
+}
+
 FPCGDataTypeIdentifier UPVTexture2DImporterSettings::GetOutputPinTypeIdentifier() const
 {
 	return FPCGDataTypeIdentifier{ FPVDataTypeInfoGrowth::AsId() };
@@ -119,44 +110,30 @@ bool FPVTexture2DImporterElement::ExecuteInternal(FPCGContext* InContext) const
 		return true;
 	}
 
-#if WITH_EDITOR
-	if (!PV::GPVEditorInterface)
+	if (Settings->Params.PlantSettings.Num() == 0)
 	{
+		PCGLog::LogWarningOnGraph(LOCTEXT("NoPlantSettings", "No plant settings added"), InContext);
 		return true;
 	}
 
-	UMaterialInstanceDynamic* TextureVisualizationMaterialInstance = nullptr;
-
-	if (Settings->TextureVisualizationMaterial)
-	{
-		TextureVisualizationMaterialInstance = UMaterialInstanceDynamic::Create(Settings->TextureVisualizationMaterial, nullptr);
-		TextureVisualizationMaterialInstance->SetFlags(RF_Transient);
-		
-		const FTaskTagScope Scope(ETaskTag::EParallelGameThread); // Needed for SetTextureParameterValue
-		TextureVisualizationMaterialInstance->SetTextureParameterValue(TEXT("Diffuse"), Settings->Params.Texture2DAsset);
-		TextureVisualizationMaterialInstance->SetTextureParameterValue(TEXT("OpacityMask"), Settings->Params.Texture2DAsset);
-	}
-
-	FPVImportTexture2DOutput ImportTexture2DOutput;
-	if (!PV::GPVEditorInterface->ImportGrowthDataFromTexture2D(Settings->Params, TextureVisualizationMaterialInstance, ImportTexture2DOutput))
+#if WITH_EDITOR
+	if (!PV::GPVEditorInterface)
 	{
-		return false;
+		return true;
 	}
 
-	UPVGrowthData* OutVariantData = FPCGContext::NewObject_AnyThread<UPVGrowthData>(InContext);
-	OutVariantData->Initialize(MoveTemp(ImportTexture2DOutput.OutputCollection));
-	
+	const static TArray<FPVVisualizationSettings> VisualizationSettings = Invoke([]()
 	{
-		TArray<FPVVisualizationSettings> VisualizationSettings;
+		TArray<FPVVisualizationSettings> OutVisualizationSettings;
 
-		FPVVisualizationSettings& PerimeterCurveVisualization = VisualizationSettings.AddDefaulted_GetRef();
+		FPVVisualizationSettings& PerimeterCurveVisualization = OutVisualizationSettings.AddDefaulted_GetRef();
 		PerimeterCurveVisualization.DebugType = EPVDebugType::Custom;
 		PerimeterCurveVisualization.bShowAnchorPoints = false;
 		PerimeterCurveVisualization.AttributeToFilter = TEXT("PerimeterCurve");
 		PerimeterCurveVisualization.CustomGroupToFilter = TEXT("PerimeterCurves");
 		PerimeterCurveVisualization.VisualizationMode = EPVDebugValueVisualizationMode::Curve;
 
-		FPVVisualizationSettings& EndPointVisualization = VisualizationSettings.AddDefaulted_GetRef();
+		FPVVisualizationSettings& EndPointVisualization = OutVisualizationSettings.AddDefaulted_GetRef();
 		EndPointVisualization.DebugType = EPVDebugType::Custom;
 		EndPointVisualization.bShowAnchorPoints = false;
 		EndPointVisualization.bDrawPointAsMesh = false;
@@ -167,7 +144,7 @@ bool FPVTexture2DImporterElement::ExecuteInternal(FPCGContext* InContext) const
 		EndPointVisualization.Color = FColor::Red;
 		EndPointVisualization.GizmoScale = 6.5f;
 
-		FPVVisualizationSettings& RootPointsVisualization = VisualizationSettings.AddDefaulted_GetRef();
+		FPVVisualizationSettings& RootPointsVisualization = OutVisualizationSettings.AddDefaulted_GetRef();
 		RootPointsVisualization.DebugType = EPVDebugType::Custom;
 		RootPointsVisualization.bShowAnchorPoints = false;
 		RootPointsVisualization.bDrawPointAsMesh = false;
@@ -177,17 +154,78 @@ bool FPVTexture2DImporterElement::ExecuteInternal(FPCGContext* InContext) const
 		RootPointsVisualization.Color = FColor::Blue;
 		RootPointsVisualization.GizmoScale = 10;
 
-		OutVariantData->AddVisualizationCollection(MoveTemp(ImportTexture2DOutput.VisualizationCollection), MoveTemp(VisualizationSettings));
+		FPVVisualizationSettings& LabelVisualization = OutVisualizationSettings.AddDefaulted_GetRef();
+		LabelVisualization.DebugType = EPVDebugType::Custom;
+		LabelVisualization.bShowAnchorPoints = false;
+		LabelVisualization.CustomGroupToFilter = TEXT("Labels");
+		LabelVisualization.AttributeToFilter = TEXT("Text");
+		LabelVisualization.CustomPivotPositionAttributeName = TEXT("Position");
+		LabelVisualization.CustomPivotScaleAttributeName = TEXT("Scale");
+		LabelVisualization.VisualizationMode = EPVDebugValueVisualizationMode::Text;
+		LabelVisualization.Color = FColor::Black;
+		LabelVisualization.TextScale = 10.f;
+
+		return OutVisualizationSettings;
+	});
+
+	const UMaterialInstanceDynamic* TextureVisualizationMaterialInstance = Invoke([&]()
+	{
+		UMaterialInstanceDynamic* OutMaterialInstance = nullptr;
+
+		if (Settings->TextureVisualizationMaterial)
+		{
+			OutMaterialInstance = UMaterialInstanceDynamic::Create(Settings->TextureVisualizationMaterial, nullptr);
+			OutMaterialInstance->SetFlags(RF_Transient);
+
+			const FTaskTagScope Scope(ETaskTag::EParallelGameThread); // Needed for SetTextureParameterValue
+			OutMaterialInstance->SetTextureParameterValue(TEXT("Diffuse"), Settings->Params.Texture2DAsset);
+			OutMaterialInstance->SetTextureParameterValue(TEXT("OpacityMask"), Settings->Params.Texture2DAsset);
+		}
+
+		return OutMaterialInstance;
+	});
+
+	FPVImportTexture2DOutput ImportTexture2DOutput;
+	if (!PV::GPVEditorInterface->ImportGrowthDataFromTexture2D(Settings->Params, TextureVisualizationMaterialInstance, ImportTexture2DOutput))
+	{
+		return false;
 	}
 
-	if (TextureVisualizationMaterialInstance)
+	// Add Combined output
 	{
+		UPVGrowthData* OutVariantData = FPCGContext::NewObject_AnyThread<UPVGrowthData>(InContext);
+
 		OutVariantData->AddManagedResource(TextureVisualizationMaterialInstance);
+		OutVariantData->AddVisualizationCollection(ImportTexture2DOutput.VisualizationCollection, VisualizationSettings);
+
+		if (ImportTexture2DOutput.DebugCollection.IsSet())
+		{
+			OutVariantData->Initialize(MoveTemp(ImportTexture2DOutput.DebugCollection.GetValue()));
+		}
+		else
+		{
+			OutVariantData->Initialize(MoveTemp(ImportTexture2DOutput.CombinedPlantCollection));
+		}
+
+		FPCGTaggedData& CollectionOutput = InContext->OutputData.TaggedData.Emplace_GetRef();
+		CollectionOutput.Data = OutVariantData;
+		CollectionOutput.Pin = PVTexture2DImporterSettings::CombinedOutputPinName;
 	}
 
-	FPCGTaggedData& CollectionOutput = InContext->OutputData.TaggedData.Emplace_GetRef();
-	CollectionOutput.Data = OutVariantData;
-	CollectionOutput.Pin = TEXT("Output");
+	// Add per-plant pin output
+	for (int32 i = 0; i < ImportTexture2DOutput.PlantCollections.Num(); ++i)
+	{
+		UPVGrowthData* OutVariantData = FPCGContext::NewObject_AnyThread<UPVGrowthData>(InContext);
+		OutVariantData->Initialize(MoveTemp(ImportTexture2DOutput.PlantCollections[i]));
+
+		// TODO: For now we add the visualization to all outputs as we cannot control which pin gets inspected
+		OutVariantData->AddManagedResource(TextureVisualizationMaterialInstance);
+		OutVariantData->AddVisualizationCollection(ImportTexture2DOutput.VisualizationCollection, VisualizationSettings);
+
+		FPCGTaggedData& CollectionOutput = InContext->OutputData.TaggedData.Emplace_GetRef();
+		CollectionOutput.Data = OutVariantData;
+		CollectionOutput.Pin = PVTexture2DImporterSettings::GetPlantPinName(i);
+	}
 #endif
 
 	return true;
diff --git a/Engine/Plugins/Experimental/ProceduralVegetationEditor/Source/ProceduralVegetation/Private/Nodes/PVTexture2DImporterSettings.h b/Engine/Plugins/Experimental/ProceduralVegetationEditor/Source/ProceduralVegetation/Private/Nodes/PVTexture2DImporterSettings.h
--- a/Engine/Plugins/Experimental/ProceduralVegetationEditor/Source/ProceduralVegetation/Private/Nodes/PVTexture2DImporterSettings.h
+++ b/Engine/Plugins/Experimental/ProceduralVegetationEditor/Source/ProceduralVegetation/Private/Nodes/PVTexture2DImporterSettings.h
@@ -29,6 +29,7 @@ public:
 
 protected:
 	virtual TArray<FPCGPinProperties> InputPinProperties() const override;
+	virtual TArray<FPCGPinProperties> OutputPinProperties() const override;
 	virtual FPCGDataTypeIdentifier GetOutputPinTypeIdentifier() const override;
 	virtual FPCGElementPtr CreateElement() const override;
 
diff --git a/Engine/Plugins/Experimental/ProceduralVegetationEditor/Source/ProceduralVegetation/Public/DataTypes/PVData.h b/Engine/Plugins/Experimental/ProceduralVegetationEditor/Source/ProceduralVegetation/Public/DataTypes/PVData.h
--- a/Engine/Plugins/Experimental/ProceduralVegetationEditor/Source/ProceduralVegetation/Public/DataTypes/PVData.h
+++ b/Engine/Plugins/Experimental/ProceduralVegetationEditor/Source/ProceduralVegetation/Public/DataTypes/PVData.h
@@ -198,7 +198,7 @@ public:
 	const FManagedArrayCollection& GetCollection() const { return Collection; }
 	FManagedArrayCollection& GetCollection() { return Collection; }
 
-	void AddManagedResource(UObject* InResource) 
+	void AddManagedResource(const UObject* InResource) 
 	{ 
 		if (InResource)
 		{
@@ -236,7 +236,7 @@ protected:
 	FManagedArrayCollection Collection;
 
 	UPROPERTY(Transient)
-	TArray<TObjectPtr<UObject>> ManagedResources;
+	TArray<TObjectPtr<const UObject>> ManagedResources;
 
 #if WITH_EDITOR
 	mutable FPVDebugSettings DebugSettings;
diff --git a/Engine/Plugins/Experimental/ProceduralVegetationEditor/Source/ProceduralVegetation/Public/PVImportTexture2DParams.h b/Engine/Plugins/Experimental/ProceduralVegetationEditor/Source/ProceduralVegetation/Public/PVImportTexture2DParams.h
--- a/Engine/Plugins/Experimental/ProceduralVegetationEditor/Source/ProceduralVegetation/Public/PVImportTexture2DParams.h
+++ b/Engine/Plugins/Experimental/ProceduralVegetationEditor/Source/ProceduralVegetation/Public/PVImportTexture2DParams.h
@@ -11,29 +11,22 @@ UENUM(BlueprintType)
 enum class EPVImportTexture2DDebugState : uint8
 {
 	None,
-	Texture,
 	SampleImage,
-	TraceImage,
-	RemoveSmallCurves,
-	SmoothCurves,
-	SimplifyCurves,
 	TriangulateCurves,
 	ComputeVertexScale,
-	FindTips,
-	ApplyRootPoints,
 	ComputeBranchHierarchies,
 	MAX UMETA(Hidden)
 };
 
 USTRUCT(BlueprintType)
-struct FPVImportTexture2DRootPoint
+struct FPVImportTexture2DPlantSettings
 {
 	GENERATED_BODY()
 
-	UPROPERTY(EditAnywhere, Category = "Root Point")
-	FVector2f Position = FVector2f::ZeroVector;
+	UPROPERTY(EditAnywhere, Category = "Plant Settings")
+	FVector2f RootPosition = FVector2f::ZeroVector;
 
-	UPROPERTY(EditAnywhere, Category = "Root Point")
+	UPROPERTY(EditAnywhere, Category = "Plant Settings")
 	float Rotation = 0;
 };
 
@@ -48,8 +41,8 @@ struct FPVImportTexture2DParams
 	UPROPERTY(EditAnywhere, Category = "Image Trace", meta=(ClampMin=0, UIMin=0))
     int32 SampleHeight = 512;
     
-    // Any branches with a bounds area lower than this value will be ignored.
-    // Value is normalized meaning a value of 0.1 will filter out any branch whose bounds cover less than 10% of the image. 
+    // Any plants with a bounds area lower than this value will be ignored.
+    // Value is normalized meaning a value of 0.1 will filter out any plants whose bounds cover less than 10% of the image. 
 	UPROPERTY(EditAnywhere, Category = "Image Trace", meta=(ClampMin=0, UIMin=0, UIMax=0.2))
 	float MinBoundsArea = 0.05f;
 	
@@ -68,16 +61,16 @@ struct FPVImportTexture2DParams
 	UPROPERTY(EditAnywhere, Category = "Texture", meta=(ClampMin=0, ClampMax=1, UIMin=0, UIMax=1))
 	float WhiteLevel = 0.01;
 
-	UPROPERTY(EditAnywhere, Category = "Branch Geneneration")
+	UPROPERTY(EditAnywhere, Category = "Plant Geneneration")
 	float Scale = 100.f;
 
-	UPROPERTY(EditAnywhere, Category = "Branch Geneneration")
-	TArray<FPVImportTexture2DRootPoint> RootPoints;
+	UPROPERTY(EditAnywhere, Category = "Plant Geneneration")
+	TArray<FPVImportTexture2DPlantSettings> PlantSettings;
 
-	UPROPERTY(EditAnywhere, Category = "Branch Geneneration|Tip detection", meta = (ClampMin=0.001, ClampMax=180, UIMin=0.01, UIMax=180))
+	UPROPERTY(EditAnywhere, Category = "Plant Geneneration|Tip detection", meta = (ClampMin=0.001, ClampMax=180, UIMin=0.01, UIMax=180))
 	float TipAngleThresholdInDegrees = 45.0;
 
-	UPROPERTY(EditAnywhere, Category = "Branch Geneneration|Tip detection", meta = (ClampMin=0, ClampMax=0.01, UIMin=0, UIMax=0.01))
+	UPROPERTY(EditAnywhere, Category = "Plant Geneneration|Tip detection", meta = (ClampMin=0, ClampMax=0.01, UIMin=0, UIMax=0.01))
 	float MaxTipAngleSearchDist = 0.005;
 
 	UPROPERTY(EditAnywhere, Category = "DebugSettings")
@@ -86,6 +79,9 @@ struct FPVImportTexture2DParams
 
 struct FPVImportTexture2DOutput
 {
+	TOptional<FManagedArrayCollection> DebugCollection;
 	FManagedArrayCollection VisualizationCollection;
-	FManagedArrayCollection OutputCollection;
+
+	FManagedArrayCollection CombinedPlantCollection;
+	TArray<FManagedArrayCollection> PlantCollections;
 };
\ No newline at end of file
diff --git a/Engine/Plugins/Experimental/ProceduralVegetationEditor/Source/ProceduralVegetationEditor/Private/Importer/PVImportHelpers.cpp b/Engine/Plugins/Experimental/ProceduralVegetationEditor/Source/ProceduralVegetationEditor/Private/Importer/PVImportHelpers.cpp
--- a/Engine/Plugins/Experimental/ProceduralVegetationEditor/Source/ProceduralVegetationEditor/Private/Importer/PVImportHelpers.cpp
+++ b/Engine/Plugins/Experimental/ProceduralVegetationEditor/Source/ProceduralVegetationEditor/Private/Importer/PVImportHelpers.cpp
@@ -176,17 +176,6 @@ namespace PV::DynamicMeshHelper
 		}
 	}
 
-	struct FMeshBranchHierarchy
-	{
-		struct FBranch
-		{
-			TArray<int32> BranchPoints;
-			int32 ParentBranch = INDEX_NONE;
-		};
-
-		TArray<FBranch> Branches;
-	};
-
 	static void ExtractBranchHierarchies(
 		const UE::Geometry::FDynamicMesh3& DynamicMesh,
 		const FBooleanVertexAttribute& RootPointAttribute,
@@ -323,86 +312,7 @@ namespace PV::DynamicMeshHelper
 		}
 	}
 
-	static void ConvertMeshBranchHierarchiesToBranchDescription(
-		const UE::Geometry::FDynamicMesh3& DynamicMesh,
-		const FFloatVertexAttribute& VertexScaleAttribute,
-		const FBooleanVertexAttribute& EndPointAttribute,
-		const TArray<FMeshBranchHierarchy>& InMeshBranchHierarchies,
-		TArray<FPVBranchHierarchyDescription>& OutBranchHierarchyDescriptions
-	)
-	{
-		TRACE_CPUPROFILER_EVENT_SCOPE(ConvertMeshBranchHierarchiesToBranchDescription);
-
-		for (const FMeshBranchHierarchy& BranchHierarchy : InMeshBranchHierarchies)
-		{
-			int32 MaxNrOfPoints = 0;
-			for (const auto& Branch : BranchHierarchy.Branches)
-			{
-				MaxNrOfPoints += Branch.BranchPoints.Num();
-			}
-
-			FPVBranchHierarchyDescription& OutBranchHierarchy = OutBranchHierarchyDescriptions.AddDefaulted_GetRef();
-			OutBranchHierarchy.Branches.Reserve(BranchHierarchy.Branches.Num());
-			OutBranchHierarchy.Points.Reserve(MaxNrOfPoints);
-			OutBranchHierarchy.PointsRadii.Reserve(MaxNrOfPoints);
-
-			TMap<int32, int32> VertexIndexToPointIndex;
-			VertexIndexToPointIndex.Reserve(MaxNrOfPoints);
-
-			for (const FMeshBranchHierarchy::FBranch& Branch : BranchHierarchy.Branches)
-			{
-				FPVBranchDescription& OutBranchDescription = OutBranchHierarchy.Branches.AddDefaulted_GetRef();
-				OutBranchDescription.ParentBranchIndex = Branch.ParentBranch;
-				OutBranchDescription.BranchIndex = OutBranchHierarchy.Branches.Num() - 1;
-				OutBranchDescription.PointIndices.Reserve(Branch.BranchPoints.Num());
-
-				for (int32 i = 0; i < Branch.BranchPoints.Num(); ++i)
-				{
-					const int32 VertexIndex = Branch.BranchPoints[i];
-					int32 PointIndex = INDEX_NONE;
-
-					int32* ExistingPointIndex = VertexIndexToPointIndex.Find(VertexIndex);
-					if (!ExistingPointIndex)
-					{
-						PointIndex = OutBranchHierarchy.Points.Add(FVector3f(DynamicMesh.GetVertex(VertexIndex)));
-
-						const bool bIsEndPoint = EndPointAttribute.GetValue(VertexIndex);
-						if (bIsEndPoint)
-						{
-							OutBranchHierarchy.PointsRadii.Add(0.01f);
-						}
-						else
-						{
-							const float VertexRadius = VertexScaleAttribute.GetValue(VertexIndex);
-							OutBranchHierarchy.PointsRadii.Add(VertexRadius);
-						}
-
-						VertexIndexToPointIndex.Add(VertexIndex, PointIndex);
-					}
-					else
-					{
-						PointIndex = *ExistingPointIndex;
-					}
-
-					OutBranchDescription.PointIndices.Add(PointIndex);
-				}
-			}
-
-			for (int32 i = 0; i < OutBranchHierarchy.Branches.Num(); ++i)
-			{
-				const FPVBranchDescription& BranchDescription = OutBranchHierarchy.Branches[i];
-				if (BranchDescription.ParentBranchIndex != INDEX_NONE)
-				{
-					OutBranchHierarchy.Branches[BranchDescription.ParentBranchIndex].ChildBranchIndices.Add(i);
-				}
-				else
-				{
-					check(OutBranchHierarchy.RootBranchIndex == INDEX_NONE);
-					OutBranchHierarchy.RootBranchIndex = i;
-				}
-			}
-		}
-	}
+	
 
 	const static TSet<FName> GeometryGroups = {
 		FGeometryCollection::VerticesGroup,
@@ -583,7 +493,7 @@ TArray<int32> PV::DynamicMeshHelper::FindVerticesInRadius(const UE::Geometry::FD
 	return OutVertexIndices;
 }
 
-void PV::DynamicMeshHelper::ComputeBranchHierarchies(
+void PV::DynamicMeshHelper::ComputeDynamicMeshBranchHierarchies(
 	const UE::Geometry::FDynamicMesh3& DynamicMesh,
 	const FBooleanVertexAttribute& RootPointsAttribute,
 	const FBooleanVertexAttribute& EndPointsAttribute,
@@ -591,7 +501,7 @@ void PV::DynamicMeshHelper::ComputeBranchHierarchies(
 	FFloatVertexAttribute& LengthFromRootAttribute,
 	FIntVertexAttribute& NextBranchPointAttribute,
 	const TFunction<bool(int32)>& VertexFilterPredicate,
-	TArray<FPVBranchHierarchyDescription>& OutBranchHierarchyDescriptions
+	TArray<FMeshBranchHierarchy>& OutDynamicMeshBranchHierarchies
 )
 {
 	ComputeLengthFromRoot(
@@ -611,25 +521,184 @@ void PV::DynamicMeshHelper::ComputeBranchHierarchies(
 		VertexFilterPredicate
 	);
 
-	TArray<FMeshBranchHierarchy> BranchHierarchies;
 	ExtractBranchHierarchies(
 		DynamicMesh,
 		RootPointsAttribute,
 		EndPointsAttribute,
 		NextBranchPointAttribute,
-		BranchHierarchies
+		OutDynamicMeshBranchHierarchies
 	);
+}
 
-	ConvertMeshBranchHierarchiesToBranchDescription(
-		DynamicMesh,
-		VertexScaleAttribute,
-		EndPointsAttribute,
-		BranchHierarchies,
-		OutBranchHierarchyDescriptions
-	);
+void PV::DynamicMeshHelper::ConvertDynamicMeshBranchHierarchiesToBranchDescription(
+	const UE::Geometry::FDynamicMesh3& DynamicMesh,
+	const FFloatVertexAttribute& VertexScaleAttribute,
+	const FBooleanVertexAttribute& EndPointAttribute,
+	const TArray<FMeshBranchHierarchy>& InMeshBranchHierarchies,
+	TArray<FPVBranchHierarchyDescription>& OutBranchHierarchyDescriptions
+)
+{
+	TRACE_CPUPROFILER_EVENT_SCOPE(ConvertMeshBranchHierarchiesToBranchDescription);
+
+	for (const FMeshBranchHierarchy& BranchHierarchy : InMeshBranchHierarchies)
+	{
+		int32 MaxNrOfPoints = 0;
+		for (const auto& Branch : BranchHierarchy.Branches)
+		{
+			MaxNrOfPoints += Branch.BranchPoints.Num();
+		}
+
+		FPVBranchHierarchyDescription& OutBranchHierarchy = OutBranchHierarchyDescriptions.AddDefaulted_GetRef();
+		OutBranchHierarchy.Branches.Reserve(BranchHierarchy.Branches.Num());
+		OutBranchHierarchy.Points.Reserve(MaxNrOfPoints);
+		OutBranchHierarchy.PointsRadii.Reserve(MaxNrOfPoints);
+
+		TMap<int32, int32> VertexIndexToPointIndex;
+		VertexIndexToPointIndex.Reserve(MaxNrOfPoints);
+
+		for (const FMeshBranchHierarchy::FBranch& Branch : BranchHierarchy.Branches)
+		{
+			FPVBranchDescription& OutBranchDescription = OutBranchHierarchy.Branches.AddDefaulted_GetRef();
+			OutBranchDescription.ParentBranchIndex = Branch.ParentBranch;
+			OutBranchDescription.BranchIndex = OutBranchHierarchy.Branches.Num() - 1;
+			OutBranchDescription.PointIndices.Reserve(Branch.BranchPoints.Num());
+
+			for (int32 i = 0; i < Branch.BranchPoints.Num(); ++i)
+			{
+				const int32 VertexIndex = Branch.BranchPoints[i];
+				int32 PointIndex = INDEX_NONE;
+
+				int32* ExistingPointIndex = VertexIndexToPointIndex.Find(VertexIndex);
+				if (!ExistingPointIndex)
+				{
+					PointIndex = OutBranchHierarchy.Points.Add(FVector3f(DynamicMesh.GetVertex(VertexIndex)));
+
+					const bool bIsEndPoint = EndPointAttribute.GetValue(VertexIndex);
+					if (bIsEndPoint)
+					{
+						OutBranchHierarchy.PointsRadii.Add(0.01f);
+					}
+					else
+					{
+						const float VertexRadius = VertexScaleAttribute.GetValue(VertexIndex);
+						OutBranchHierarchy.PointsRadii.Add(VertexRadius);
+					}
+
+					VertexIndexToPointIndex.Add(VertexIndex, PointIndex);
+				}
+				else
+				{
+					PointIndex = *ExistingPointIndex;
+				}
+
+				OutBranchDescription.PointIndices.Add(PointIndex);
+			}
+		}
+
+		for (int32 i = 0; i < OutBranchHierarchy.Branches.Num(); ++i)
+		{
+			const FPVBranchDescription& BranchDescription = OutBranchHierarchy.Branches[i];
+			if (BranchDescription.ParentBranchIndex != INDEX_NONE)
+			{
+				OutBranchHierarchy.Branches[BranchDescription.ParentBranchIndex].ChildBranchIndices.Add(i);
+			}
+			else
+			{
+				check(OutBranchHierarchy.RootBranchIndex == INDEX_NONE);
+				OutBranchHierarchy.RootBranchIndex = i;
+			}
+		}
+	}
+}
+
+
+FPVBranchHierarchyDescription::EValidateHierarchyResult FPVBranchHierarchyDescription::ValidateHierarchy() const
+{
+	if (!Branches.IsValidIndex(RootBranchIndex))
+	{
+		return EValidateHierarchyResult::InvalidRootBranchIndex;
+	}
+	
+	if (Branches[RootBranchIndex].ParentBranchIndex != INDEX_NONE)
+	{
+		return EValidateHierarchyResult::InvalidRootBranchParentIndex;
+	}
+	
+	for (int32 BranchIndex = 0; BranchIndex < Branches.Num(); BranchIndex++)
+	{
+		const FPVBranchDescription& Branch = Branches[BranchIndex];
+		if (Branch.BranchIndex != BranchIndex)
+		{
+			return EValidateHierarchyResult::BranchIndexMismatch;
+		}
+		
+		if (Branch.ParentBranchIndex != INDEX_NONE
+			&& !Branches.IsValidIndex(Branch.ParentBranchIndex))
+		{
+			return EValidateHierarchyResult::InvalidParentBranchIndex;
+		}
+		
+		if (Branch.ParentBranchIndex == INDEX_NONE
+			&& RootBranchIndex != Branch.BranchIndex)
+		{
+			return EValidateHierarchyResult::InvalidRootBranchIndex;
+		}
+		
+		for (int32 ChildBranchIndex : Branch.ChildBranchIndices)
+		{
+			if (!Branches.IsValidIndex(ChildBranchIndex))
+			{
+				return EValidateHierarchyResult::InvalidChildBranchIndex;
+			}
+		}
+		
+		for (int32 PointIndex : Branch.PointIndices)
+		{
+			if (!Points.IsValidIndex(PointIndex))
+			{
+				return EValidateHierarchyResult::InvalidPointIndex; 
+			}
+		}
+	}
+
+	TSet<int32> FoundBranches;
+	FoundBranches.Reserve(Branches.Num());
+
+	const TFunction<bool(int32)> WalkHierarchyRecursive = [&](int32 BranchIndex)->bool
+	{
+		bool bAlreadyAdded = false;
+		FoundBranches.Add(BranchIndex, &bAlreadyAdded);
+
+		if (bAlreadyAdded)
+		{
+			return false;
+		}
+
+		for (int32 ChildIndex : Branches[BranchIndex].ChildBranchIndices)
+		{
+			if (!WalkHierarchyRecursive(ChildIndex))
+			{
+				return false;
+			}
+		}
+
+		return true;
+	};
+
+	if (!WalkHierarchyRecursive(RootBranchIndex))
+	{
+		return EValidateHierarchyResult::CircularHierarchy;
+	}
+
+	if (FoundBranches.Num() != Branches.Num())
+	{
+		return EValidateHierarchyResult::DisjointedHierarchy;
+	}
+
+	return EValidateHierarchyResult::Success;
 }
 
-void PV::Import::Helper::CreateEmptyGrowthData(FManagedArrayCollection& OutCollection, int32 NumBranchPoints, int32 NumFoliagePoints, int32 NumBranches, int32 NumFoliageMeshes)
+void PV::ImportHelper::CreateEmptyGrowthData(FManagedArrayCollection& OutCollection, int32 NumBranchPoints, int32 NumFoliagePoints, int32 NumBranches, int32 NumFoliageMeshes)
 {
 	// PointGroup
 	OutCollection.AddGroup(PV::GroupNames::PointGroup);
@@ -746,9 +815,9 @@ void PV::Import::Helper::CreateEmptyGrowthData(FManagedArrayCollection& OutColle
 	LeafPhyllotaxyAttribute[0] = { 0, 200, 50, 1, 1, 0, 0, 0 };
 }
 
-void PV::Import::Helper::FillGrowthDataFromBranchHierarchyDescription(
+void PV::ImportHelper::FillGrowthDataFromBranchHierarchyDescription(
 	FManagedArrayCollection& InOutCollection, 
-	const TArray<FPVBranchHierarchyDescription>& InBranchHierarchies
+	TArrayView<const FPVBranchHierarchyDescription> InBranchHierarchies
 )
 {
 	PV::Facades::FPointFacade PointFacade(InOutCollection);
@@ -871,89 +940,7 @@ void PV::Import::Helper::FillGrowthDataFromBranchHierarchyDescription(
 	}
 }
 
-void PV::Import::Helper::FillGrowthDataFromFoliageDescription(
-	FManagedArrayCollection& Collection, 
-	const FPVFoliageDescription& InFoliageDescription, 
-	const TArray<FPVBranchHierarchyDescription>& InBranchHierarchies
-)
-{
-	PV::Facades::FFoliageFacade FoliageFacade(Collection);
-	PV::Facades::FPointFacade PointFacade(Collection);
-
-	const auto FindClosestBranchSegment = [&](const FVector3f& Point, int32& OutBranchIndex, FVector3f& OutClosestPointOnBranch, int32& OutParentPointIndex)
-	{
-		float ClosestBranchDistSquared = MAX_FLT;
-		
-		OutBranchIndex = INDEX_NONE;
-		OutParentPointIndex = INDEX_NONE;
-		OutClosestPointOnBranch = FVector3f::ZeroVector;
-
-		int32 BranchIndexOffset = 0;
-		for (const FPVBranchHierarchyDescription& BranchHierarchy : InBranchHierarchies)
-		{
-			for (int32 i = 0; i < BranchHierarchy.Branches.Num(); ++i)
-			{
-				FVector3f ClosestPointOnBranch;
-				float BranchDistSquared;
-				int32 PointIndex;
-				DistToBranch(Point, BranchHierarchy.Points, BranchHierarchy.Branches[i].PointIndices, ClosestPointOnBranch, BranchDistSquared, PointIndex);
-				if (BranchDistSquared < ClosestBranchDistSquared)
-				{
-					ClosestBranchDistSquared = BranchDistSquared;
-
-					OutBranchIndex = i + BranchIndexOffset;
-					OutParentPointIndex = PointIndex;
-					OutClosestPointOnBranch = ClosestPointOnBranch;
-				}
-			}
-
-			BranchIndexOffset += BranchHierarchy.Branches.Num();
-		}
-	};
-
-	for (int32 i = 0; i < InFoliageDescription.FoliageMeshes.Num(); ++i)
-	{
-		TObjectPtr<UStaticMesh> FoliageMesh = InFoliageDescription.FoliageMeshes[i];
-		FoliageFacade.SetFoliageName(i, FoliageMesh ? FoliageMesh->GetPathName() : TEXT("None"));
-	}
-
-	int32 FoliageEntryIndex = 0;
-	for (int32 i = 0; i < InFoliageDescription.FoliagePoints.Num(); ++i)
-	{
-		const FPVFoliagePointsDescription& FoliagePointsDescription = InFoliageDescription.FoliagePoints[i];
-		for (int32 j = 0; j < FoliagePointsDescription.Positions.Num(); ++j)
-		{
-			int32 ClosestBranchIndex;
-			FVector3f ClosestPointOnBranch;
-			int32 ParentPointIndex;
-			FindClosestBranchSegment(FoliagePointsDescription.Positions[j], ClosestBranchIndex, ClosestPointOnBranch, ParentPointIndex);
-
-			check(ClosestBranchIndex != INDEX_NONE);
-			check(ParentPointIndex != INDEX_NONE);
-
-			const float FoliageLengthFromRoot = PointFacade.GetLengthFromRoot(ParentPointIndex) + (PointFacade.GetPosition(ParentPointIndex) - ClosestPointOnBranch).Size();
-
-			PV::Facades::FFoliageEntryData FoliageEntry;
-			FoliageEntry.NameId = FoliagePointsDescription.FoliageIndices[j];
-			FoliageEntry.BranchId = ClosestBranchIndex;
-			FoliageEntry.PivotPoint = FoliagePointsDescription.Positions[j];
-			FoliageEntry.UpVector = FoliagePointsDescription.Orientations[j].GetUpVector();
-			FoliageEntry.NormalVector = FoliagePointsDescription.Orientations[j].GetRightVector();
-			FoliageEntry.Scale = FoliagePointsDescription.Scales[j].GetAbsMax();
-			FoliageEntry.LengthFromRoot = FoliageLengthFromRoot;
-			FoliageEntry.ParentBoneID = ClosestBranchIndex;
-			FoliageFacade.SetFoliageEntry(FoliageEntryIndex, FoliageEntry);
-
-			TArray<int32> FoliageEntryIdsForBranch = FoliageFacade.GetFoliageEntryIdsForBranch(ClosestBranchIndex);
-			FoliageEntryIdsForBranch.Add(FoliageEntryIndex);
-			FoliageFacade.SetFoliageIdsArray(ClosestBranchIndex, MoveTemp(FoliageEntryIdsForBranch));
-
-			FoliageEntryIndex++;
-		}
-	}
-}
-
-void PV::Import::Helper::ComputeLengthFromRoot(PV::Facades::FPointFacade& PointFacade, const PV::Facades::FBranchFacade& BranchFacade, const PV::Facades::FPlantFacade& PlantFacade)
+void PV::ImportHelper::ComputeLengthFromRoot(PV::Facades::FPointFacade& PointFacade, const PV::Facades::FBranchFacade& BranchFacade, const PV::Facades::FPlantFacade& PlantFacade)
 {
 	PV::FacadeHelper::WalkBranches(PlantFacade, BranchFacade, [&](int32 BranchIndex)->PV::FacadeHelper::EForEachResult
 	{
@@ -1037,7 +1024,7 @@ void PV::Import::Helper::ComputeLengthFromRoot(PV::Facades::FPointFacade& PointF
 	});
 }
 
-void PV::Import::Helper::ComputeBudDevelopment(PV::Facades::FPointFacade& PointFacade, const PV::Facades::FBranchFacade& BranchFacade, const PV::Facades::FPlantFacade& PlantFacade)
+void PV::ImportHelper::ComputeBudDevelopment(PV::Facades::FPointFacade& PointFacade, const PV::Facades::FBranchFacade& BranchFacade, const PV::Facades::FPlantFacade& PlantFacade)
 {
 	/*
 		0_Generatation: Generation is assumed to be the number of parents in the parents array and does not account for codominant structures. 
@@ -1080,7 +1067,7 @@ void PV::Import::Helper::ComputeBudDevelopment(PV::Facades::FPointFacade& PointF
 	);
 }
 
-void PV::Import::Helper::ComputeBudDirections(PV::Facades::FPointFacade& PointFacade, const PV::Facades::FBranchFacade& BranchFacade, const PV::Facades::FPlantFacade& PlantFacade)
+void PV::ImportHelper::ComputeBudDirections(PV::Facades::FPointFacade& PointFacade, const PV::Facades::FBranchFacade& BranchFacade, const PV::Facades::FPlantFacade& PlantFacade)
 {
 	/*
 		0_Apical: Apical is generated by getting the direction from the previous point to the current point, unless its the last.
@@ -1144,7 +1131,7 @@ void PV::Import::Helper::ComputeBudDirections(PV::Facades::FPointFacade& PointFa
 	);
 }
 
-void PV::Import::Helper::InterpolatePointsScales(PV::Facades::FPointFacade& PointFacade, const PV::Facades::FBranchFacade& BranchFacade, const PV::Facades::FPlantFacade& PlantFacade)
+void PV::ImportHelper::InterpolatePointsScales(PV::Facades::FPointFacade& PointFacade, const PV::Facades::FBranchFacade& BranchFacade, const PV::Facades::FPlantFacade& PlantFacade)
 {
 	PV::FacadeHelper::WalkBranches(PlantFacade, BranchFacade, [&](int32 BranchIndex)->PV::FacadeHelper::EForEachResult
 	{
@@ -1205,7 +1192,7 @@ void PV::Import::Helper::InterpolatePointsScales(PV::Facades::FPointFacade& Poin
 	});
 }
 
-bool PV::Import::Helper::SplitFusedPoints(FManagedArrayCollection& Collection)
+bool PV::ImportHelper::SplitFusedPoints(FManagedArrayCollection& Collection)
 {
 	PV::Facades::FBranchFacade BranchFacade(Collection);
 
@@ -1272,7 +1259,7 @@ bool PV::Import::Helper::SplitFusedPoints(FManagedArrayCollection& Collection)
 	return true;
 }
 
-void PV::Import::Helper::ConnectDisconnectedBranches(FPVBranchHierarchyDescription& InOutBranchHierarchy)
+void PV::ImportHelper::ConnectDisconnectedBranches(FPVBranchHierarchyDescription& InOutBranchHierarchy)
 {
 	struct FNewBranchRootIndex
 	{
@@ -1367,7 +1354,8 @@ void PV::Import::Helper::ConnectDisconnectedBranches(FPVBranchHierarchyDescripti
 		int32 ClosestApicalIndex = ClosestPointIndex;
 		for (int32 i = 0; i < ParentBranch.PointIndices.Num(); ++i)
 		{
-			const float PercentageDiff = FMath::Abs(PointDistances[ClosestPointIndex] - PointDistances[i]) / PointDistances[i];
+			const float PercentageDiff = FMath::Abs(PointDistances[ClosestPointIndex] - PointDistances[i]) / (PointDistances[i] > 0 ? PointDistances[i] : 1.f);
+
 			if (PercentageDiff < 0.5)
 			{
 				if (ApicalDots[i] > ApicalDots[ClosestApicalIndex])
@@ -1427,6 +1415,10 @@ void PV::Import::Helper::ConnectDisconnectedBranches(FPVBranchHierarchyDescripti
 			check(PointsRefs[PointIndex].Num() == 0);
 			const int32 SwappedPointIndex = InOutBranchHierarchy.Points.Num() - 1;
 			InOutBranchHierarchy.Points.RemoveAtSwap(PointIndex);
+			if (InOutBranchHierarchy.PointsRadii.IsValidIndex(PointIndex))
+			{
+				InOutBranchHierarchy.PointsRadii.RemoveAtSwap(PointIndex);
+			}
 
 			for (const FPointRef& PointRef : PointsRefs[SwappedPointIndex])
 			{
@@ -1438,7 +1430,7 @@ void PV::Import::Helper::ConnectDisconnectedBranches(FPVBranchHierarchyDescripti
 	}
 }
 
-void PV::Import::Helper::FuseCollinearBranches(FPVBranchHierarchyDescription& InOutBranchHierarchy)
+void PV::ImportHelper::FuseCollinearBranches(FPVBranchHierarchyDescription& InOutBranchHierarchy)
 {
 	TArray<int32> BranchesToFuseWithParent;
 	BranchesToFuseWithParent.Reserve(InOutBranchHierarchy.Branches.Num());
@@ -1567,32 +1559,30 @@ void PV::Import::Helper::FuseCollinearBranches(FPVBranchHierarchyDescription& In
 	}
 }
 
-bool PV::Import::Helper::GenerateGrowthDataFromPlantDescription(FManagedArrayCollection& OutCollection, const FPVPlantDescription& InPlantDescription)
+bool PV::ImportHelper::GenerateGrowthDataFromBranchHierarchies(FManagedArrayCollection& OutCollection, TArrayView<const FPVBranchHierarchyDescription> InBranchHierarchies)
 {
-	const TArray<FPVBranchHierarchyDescription>& InBranchHierarchies = InPlantDescription.BranchHierarchies;
-	const FPVFoliageDescription& InFoliageDescription = InPlantDescription.Foliage;
-
 	if (InBranchHierarchies.Num() < 1)
 	{
 		return false;
 	}
 
+	int32 NumBranchPoints = 0;
+	int32 NumBranches = 0;
+
 	for (const FPVBranchHierarchyDescription& BranchHierarchy : InBranchHierarchies)
 	{
+		NumBranchPoints += BranchHierarchy.GetNumBranchPoints();
+		NumBranches += BranchHierarchy.GetNumBranches();
+
 		if (BranchHierarchy.ValidateHierarchy() != FPVBranchHierarchyDescription::EValidateHierarchyResult::Success)
 		{
 			return false;
 		}
 	}
 
-	const int32 NumBranchPoints = InPlantDescription.GetNumBranchPoints();
-	const int32 NumFoliagePoints = InPlantDescription.GetNumFoliagePoints();
-	const int32 NumBranches = InPlantDescription.GetNumBranches();
-	const int32 NumFoliageMeshes = InFoliageDescription.FoliageMeshes.Num();
 
-	CreateEmptyGrowthData(OutCollection, NumBranchPoints, NumFoliagePoints, NumBranches, NumFoliageMeshes);
+	CreateEmptyGrowthData(OutCollection, NumBranchPoints, 0, NumBranches, 0);
 	FillGrowthDataFromBranchHierarchyDescription(OutCollection, InBranchHierarchies);
-	FillGrowthDataFromFoliageDescription(OutCollection, InFoliageDescription, InBranchHierarchies);
 
 	PV::Facades::FPointFacade PointFacade(OutCollection);
 	const PV::Facades::FBranchFacade BranchFacade(OutCollection);
@@ -1608,7 +1598,13 @@ bool PV::Import::Helper::GenerateGrowthDataFromPlantDescription(FManagedArrayCol
 	return true;
 }
 
-void PV::Import::Helper::DistToBranch(const FVector3f& Point, const TArray<FVector3f>& Points, const TArray<int32>& BranchPointIndices, FVector3f& OutClosestPoint, float& OutDistSquared, int32& OutPointIndex)
+bool PV::ImportHelper::GenerateGrowthDataFromBranchHierarchy(FManagedArrayCollection& OutCollection, const FPVBranchHierarchyDescription& InBranchHierarchy)
+{
+	TArrayView<const FPVBranchHierarchyDescription> BranchHierarchyArrayView(&InBranchHierarchy, 1);
+	return GenerateGrowthDataFromBranchHierarchies(OutCollection, BranchHierarchyArrayView);
+}
+
+void PV::ImportHelper::DistToBranch(const FVector3f& Point, const TArray<FVector3f>& Points, const TArray<int32>& BranchPointIndices, FVector3f& OutClosestPoint, float& OutDistSquared, int32& OutPointIndex)
 {
 	OutClosestPoint = FVector3f::ZeroVector;
 	OutDistSquared = MAX_FLT;
@@ -1637,7 +1633,7 @@ void PV::Import::Helper::DistToBranch(const FVector3f& Point, const TArray<FVect
 	}
 }
 
-float PV::Import::Helper::DistToBranch(const FVector3f& Point, const TArray<FVector3f>& Points, const TArray<int32>& BranchPointIndices)
+float PV::ImportHelper::DistToBranch(const FVector3f& Point, const TArray<FVector3f>& Points, const TArray<int32>& BranchPointIndices)
 {
 	FVector3f ClosestPoint;
 	float ClosestDistSquared = MAX_FLT;
@@ -1647,9 +1643,9 @@ float PV::Import::Helper::DistToBranch(const FVector3f& Point, const TArray<FVec
 	return ClosestDistSquared > 0 ? FMath::Sqrt(ClosestDistSquared) : 0;
 }
 
-TArray<int32> PV::Import::Helper::EstimateBranchHierarchy(const TArray<FVector3f>& Points, const TArray<TArray<int32>>& Branches, int32 RootBranchIndex)
+TArray<int32> PV::ImportHelper::EstimateBranchHierarchy(const TArray<FVector3f>& Points, const TArray<TArray<int32>>& Branches, int32 RootBranchIndex)
 {
-	TRACE_CPUPROFILER_EVENT_SCOPE(PV::Import::Helper::EstimateBranchHierarchy);
+	TRACE_CPUPROFILER_EVENT_SCOPE(PV::ImportHelper::EstimateBranchHierarchy);
 
 	TArray<int32> OutBranchParents;
 	OutBranchParents.SetNum(Branches.Num());
@@ -1673,6 +1669,11 @@ TArray<int32> PV::Import::Helper::EstimateBranchHierarchy(const TArray<FVector3f
 			continue;
 		}
 
+		if (Branches[i].Num() == 0)
+		{
+			continue;
+		}
+
 		const int32 BranchRootPointIndex = Branches[i][0];
 		const FVector3f& BranchRootPoint = Points[BranchRootPointIndex];
 		float ClosestBranchDistSquared = MAX_FLT;
@@ -1689,7 +1690,7 @@ TArray<int32> PV::Import::Helper::EstimateBranchHierarchy(const TArray<FVector3f
 			FVector3f ClosestPointOnBranch;
 			float BranchDistSquared = MAX_FLT;
 			int32 ClosestPointIndex;
-			Helper::DistToBranch(BranchRootPoint, Points, Branches[j], ClosestPointOnBranch, BranchDistSquared, ClosestPointIndex);
+			ImportHelper::DistToBranch(BranchRootPoint, Points, Branches[j], ClosestPointOnBranch, BranchDistSquared, ClosestPointIndex);
 
 			const bool bBranchPointIsRoot = ClosestPointIndex == BranchRootPointIndex;
 
@@ -1710,14 +1711,14 @@ TArray<int32> PV::Import::Helper::EstimateBranchHierarchy(const TArray<FVector3f
 	return OutBranchParents;
 }
 
-int32 PV::Import::Helper::FindRootBranch(const TArray<FVector3f>& InPoints, const TArray<TArray<int32>>& BranchIndices)
+int32 PV::ImportHelper::FindRootBranch(const TArray<FVector3f>& InPoints, const TArray<TArray<int32>>& BranchIndices)
 {
 	float LowestRootPoint = MAX_FLT;
 	int32 LowestRootPointIdx = INDEX_NONE;
 	for (int32 i = 0; i < BranchIndices.Num(); ++i)
 	{
 		const TArray<int32>& Branch = BranchIndices[i];
-		if (InPoints[Branch[0]].Z < LowestRootPoint)
+		if (Branch.Num() > 0 && InPoints[Branch[0]].Z < LowestRootPoint)
 		{
 			LowestRootPoint = InPoints[Branch[0]].Z;
 			LowestRootPointIdx = i;
@@ -1726,9 +1727,9 @@ int32 PV::Import::Helper::FindRootBranch(const TArray<FVector3f>& InPoints, cons
 	return LowestRootPointIdx;
 }
 
-FPVBranchHierarchyDescription PV::Import::Helper::CreateBranchHierarchyFromPoints(const TArray<FVector3f>& InPoints, const TArray<float>& InPointsRadii, const TArray<TArray<int32>>& InBranchPointIndices)
+FPVBranchHierarchyDescription PV::ImportHelper::CreateBranchHierarchyFromPoints(const TArray<FVector3f>& InPoints, const TArray<float>& InPointsRadii, const TArray<TArray<int32>>& InBranchPointIndices)
 {
-	TRACE_CPUPROFILER_EVENT_SCOPE(PV::Import::Helper::CreateBranchHierarchyFromPoints);
+	TRACE_CPUPROFILER_EVENT_SCOPE(PV::ImportHelper::CreateBranchHierarchyFromPoints);
 
 	if (InBranchPointIndices.Num() == 0 || InPoints.Num() == 0)
 	{
diff --git a/Engine/Plugins/Experimental/ProceduralVegetationEditor/Source/ProceduralVegetationEditor/Private/Importer/PVImportHelpers.h b/Engine/Plugins/Experimental/ProceduralVegetationEditor/Source/ProceduralVegetationEditor/Private/Importer/PVImportHelpers.h
--- a/Engine/Plugins/Experimental/ProceduralVegetationEditor/Source/ProceduralVegetationEditor/Private/Importer/PVImportHelpers.h
+++ b/Engine/Plugins/Experimental/ProceduralVegetationEditor/Source/ProceduralVegetationEditor/Private/Importer/PVImportHelpers.h
@@ -1,7 +1,7 @@
 // Copyright Epic Games, Inc. All Rights Reserved.
 #pragma once
 
-#include "PVPlantDescription.h"
+#include "CoreMinimal.h"
 
 namespace PV::Facades
 {
@@ -28,6 +28,40 @@ namespace UE::Geometry
 	using FDynamicMeshAABBTree3 = TMeshAABBTree3<FDynamicMesh3>;
 }
 
+struct FPVBranchDescription
+{
+	int32 ParentBranchIndex = INDEX_NONE;
+	int32 BranchIndex = INDEX_NONE;
+	TArray<int32> ChildBranchIndices;
+	TArray<int32> PointIndices;
+};
+
+struct FPVBranchHierarchyDescription
+{
+	int32 RootBranchIndex = INDEX_NONE;
+	TArray<FPVBranchDescription> Branches;
+	TArray<FVector3f> Points;
+	TArray<float> PointsRadii;
+
+	int32 GetNumBranches() const { return Branches.Num(); }
+	int32 GetNumBranchPoints() const { return Points.Num(); }
+
+	enum class EValidateHierarchyResult
+	{
+		InvalidRootBranchIndex,
+		InvalidRootBranchParentIndex,
+		InvalidParentBranchIndex,
+		InvalidBranchIndex,
+		InvalidChildBranchIndex,
+		InvalidPointIndex,
+		BranchIndexMismatch,
+		DisjointedHierarchy,
+		CircularHierarchy,
+		Success
+	};
+	EValidateHierarchyResult ValidateHierarchy() const;
+};
+
 namespace PV::DynamicMeshHelper
 {
 	using FBooleanVertexAttribute = TDynamicMeshVertexAttributeExt<bool>;
@@ -50,7 +84,18 @@ namespace PV::DynamicMeshHelper
 
 	TArray<int32> FindVerticesInRadius(const UE::Geometry::FDynamicMeshAABBTree3& DynamicMeshAABBTree, const FVector& Pos, float Radius);
 	
-	void ComputeBranchHierarchies(
+	struct FMeshBranchHierarchy
+	{
+		struct FBranch
+		{
+			TArray<int32> BranchPoints;
+			int32 ParentBranch = INDEX_NONE;
+		};
+
+		TArray<FBranch> Branches;
+	};
+
+	void ComputeDynamicMeshBranchHierarchies(
 		const UE::Geometry::FDynamicMesh3& DynamicMesh,
 		const FBooleanVertexAttribute& RootPointsAttribute,
 		const FBooleanVertexAttribute& EndPointsAttribute,
@@ -58,23 +103,25 @@ namespace PV::DynamicMeshHelper
 		FFloatVertexAttribute& LengthFromRootAttribute,
 		FIntVertexAttribute& NextBranchPointAttribute,
 		const TFunction<bool(int32)>& VertexFilterPredicate,
+		TArray<FMeshBranchHierarchy>& OutDynamicMeshBranchHierarchies
+	);
+
+	void ConvertDynamicMeshBranchHierarchiesToBranchDescription(
+		const UE::Geometry::FDynamicMesh3& DynamicMesh,
+		const FFloatVertexAttribute& VertexScaleAttribute,
+		const FBooleanVertexAttribute& EndPointAttribute,
+		const TArray<FMeshBranchHierarchy>& InMeshBranchHierarchies,
 		TArray<FPVBranchHierarchyDescription>& OutBranchHierarchyDescriptions
 	);
 }
 
-namespace PV::Import::Helper
+namespace PV::ImportHelper
 {
 	void CreateEmptyGrowthData(FManagedArrayCollection& OutCollection, int32 NumBranchPoints, int32 NumFoliagePoints, int32 NumBranches, int32 NumFoliageMeshes);
 
 	void FillGrowthDataFromBranchHierarchyDescription(
 		FManagedArrayCollection& Collection, 
-		const TArray<FPVBranchHierarchyDescription>& InBranchHierarchies
-	);
-
-	void FillGrowthDataFromFoliageDescription(
-		FManagedArrayCollection& Collection, 
-		const FPVFoliageDescription& InFoliageDescription, 
-		const TArray<FPVBranchHierarchyDescription>& InBranchHierarchies
+		TArrayView<const FPVBranchHierarchyDescription> InBranchHierarchies
 	);
 
 	void ComputeLengthFromRoot(PV::Facades::FPointFacade& PointFacade, const PV::Facades::FBranchFacade& BranchFacade, const PV::Facades::FPlantFacade& PlantFacade);
@@ -91,7 +138,9 @@ namespace PV::Import::Helper
 
 	void FuseCollinearBranches(FPVBranchHierarchyDescription& InOutBranchHierarchy);
 
-	bool GenerateGrowthDataFromPlantDescription(FManagedArrayCollection& InOutCollection, const FPVPlantDescription& InPlantDescription);
+	bool GenerateGrowthDataFromBranchHierarchies(FManagedArrayCollection& OutCollection, TArrayView<const FPVBranchHierarchyDescription> InBranchHierarchies);
+
+	bool GenerateGrowthDataFromBranchHierarchy(FManagedArrayCollection& OutCollection, const FPVBranchHierarchyDescription& InBranchHierarchy);
 
 	void DistToBranch(const FVector3f& Point, const TArray<FVector3f>& Points, const TArray<int32>& BranchPointIndices, FVector3f& OutClosestPoint, float& OutDistSquared, int32& OutPointIndex);
 
diff --git a/Engine/Plugins/Experimental/ProceduralVegetationEditor/Source/ProceduralVegetationEditor/Private/Importer/PVImporter_StaticMesh.cpp b/Engine/Plugins/Experimental/ProceduralVegetationEditor/Source/ProceduralVegetationEditor/Private/Importer/PVImporter_StaticMesh.cpp
--- a/Engine/Plugins/Experimental/ProceduralVegetationEditor/Source/ProceduralVegetationEditor/Private/Importer/PVImporter_StaticMesh.cpp
+++ b/Engine/Plugins/Experimental/ProceduralVegetationEditor/Source/ProceduralVegetationEditor/Private/Importer/PVImporter_StaticMesh.cpp
@@ -2,7 +2,6 @@
 
 #include "PVImporter_StaticMesh.h"
 
-#include "PVPlantDescription.h"
 #include "PVImportHelpers.h"
 #include "PVDynamicMeshVertexAttribute.h"
 
@@ -935,9 +934,9 @@ PV::StaticMeshImport::EImportResult PV::StaticMeshImport::ImportGrowthDataFromSt
 		return EImportResult::Success;
 	}
 
-	TArray<FPVBranchHierarchyDescription> BranchHierarchyDescriptions;
 	auto& FillerVertexAttribute = FillerVertexAttributeDefinition.GetOrAttachAttribute(DynamicMesh);
-	ComputeBranchHierarchies(
+	TArray<FMeshBranchHierarchy> DynamicMeshBranchHierarchies;
+	ComputeDynamicMeshBranchHierarchies(
 		DynamicMesh,
 		RootPointAttributeDefinition->GetOrAttachAttribute(DynamicMesh),
 		EndPointAttributeDefinition->GetOrAttachAttribute(DynamicMesh),
@@ -945,6 +944,15 @@ PV::StaticMeshImport::EImportResult PV::StaticMeshImport::ImportGrowthDataFromSt
 		LengthFromRootAttributeDefinition->GetOrAttachAttribute(DynamicMesh),
 		NextBranchPointAttributeDefinition->GetOrAttachAttribute(DynamicMesh),
 		[&](int32 VertexID) { return !FillerVertexAttribute.GetValue(VertexID); },
+		DynamicMeshBranchHierarchies
+	);
+
+	TArray<FPVBranchHierarchyDescription> BranchHierarchyDescriptions;
+	ConvertDynamicMeshBranchHierarchiesToBranchDescription(
+		DynamicMesh,
+		VertexScaleAttributeDefinition->GetOrAttachAttribute(DynamicMesh),
+		EndPointAttributeDefinition->GetOrAttachAttribute(DynamicMesh),
+		DynamicMeshBranchHierarchies,
 		BranchHierarchyDescriptions
 	);
 
@@ -961,10 +969,7 @@ PV::StaticMeshImport::EImportResult PV::StaticMeshImport::ImportGrowthDataFromSt
 		return EImportResult::Success;
 	}
 
-	FPVPlantDescription PlantDescription;
-	PlantDescription.BranchHierarchies = MoveTemp(BranchHierarchyDescriptions);
-
-	if (!PV::Import::Helper::GenerateGrowthDataFromPlantDescription(OutCollection, PlantDescription))
+	if (!PV::ImportHelper::GenerateGrowthDataFromBranchHierarchies(OutCollection, BranchHierarchyDescriptions))
 	{
 		return EImportResult::FailedToGenerateGrowthData;
 	}
diff --git a/Engine/Plugins/Experimental/ProceduralVegetationEditor/Source/ProceduralVegetationEditor/Private/Importer/PVImporter_Texture2D.cpp b/Engine/Plugins/Experimental/ProceduralVegetationEditor/Source/ProceduralVegetationEditor/Private/Importer/PVImporter_Texture2D.cpp
--- a/Engine/Plugins/Experimental/ProceduralVegetationEditor/Source/ProceduralVegetationEditor/Private/Importer/PVImporter_Texture2D.cpp
+++ b/Engine/Plugins/Experimental/ProceduralVegetationEditor/Source/ProceduralVegetationEditor/Private/Importer/PVImporter_Texture2D.cpp
@@ -2,7 +2,6 @@
 
 #include "PVImporter_Texture2D.h"
 #include "PVImportHelpers.h"
-#include "PVPlantDescription.h"
 #include "PVDynamicMeshVertexAttribute.h"
 
 #include "IntVectorTypes.h"
@@ -44,8 +43,11 @@ namespace Internal
 	const static FName PerimeterCurveGroupName = TEXT("PerimeterCurves");
 	const static FName PerimeterCurveAttributeName = TEXT("PerimeterCurve");
 	const static FName InputRootPointsGroupName = TEXT("InputRootPoints");
-	const static FName RootPointsGroupName = TEXT("RootPoints");
-	const static FName RootPointsAttributeName = TEXT("Point");
+	const static FName InputRootPointsAttributeName = TEXT("Point");
+	const static FName LabelsGroupName = TEXT("Labels");
+	const static FName LabelPositionAttributeName = TEXT("Position");
+	const static FName LabelScaleAttributeName = TEXT("Scale");
+	const static FName LabelTextAttributeName = TEXT("Text");
 
 	const static PV::TDynamicMeshVertexAttributeDefinition<bool> PerimeterAttributeDefinition(TEXT("Perimeter"));
 	
@@ -67,16 +69,76 @@ namespace Internal
 		return UE::Math::TVector<Vector3dType>(0, Vec2D.X * Scale, (1.0 - Vec2D.Y) * Scale);
 	}
 
-	static void AddRootPointsToCollection(FManagedArrayCollection& InOutCollection, const TArray<FPVImportTexture2DRootPoint>& InRootPoints, float Scale, const FVector& Offset = FVector::ZeroVector)
+	static void AddLabelsToCollection(FManagedArrayCollection& InOutCollection, const TArray<TTuple<FVector3f, FString, float>>& Labels)
+	{
+		if (!InOutCollection.HasGroup(LabelsGroupName))
+		{
+			InOutCollection.AddGroup(LabelsGroupName);
+		}
+
+		if (!InOutCollection.HasAttribute(LabelTextAttributeName, LabelsGroupName))
+		{
+			InOutCollection.AddAttribute<FString>(LabelTextAttributeName, LabelsGroupName);
+		}
+
+		if (!InOutCollection.HasAttribute(LabelPositionAttributeName, LabelsGroupName))
+		{
+			InOutCollection.AddAttribute<FVector3f>(LabelPositionAttributeName, LabelsGroupName);
+		}
+
+		if (!InOutCollection.HasAttribute(LabelScaleAttributeName, LabelsGroupName))
+		{
+			InOutCollection.AddAttribute<float>(LabelScaleAttributeName, LabelsGroupName);
+		}
+
+		if (!ensure(InOutCollection.GetAttributeType(LabelTextAttributeName, LabelsGroupName) == FManagedArrayCollection::EArrayType::FStringType))
+		{
+			return;
+		}
+
+		if (!ensure(InOutCollection.GetAttributeType(LabelPositionAttributeName, LabelsGroupName) == FManagedArrayCollection::EArrayType::FVectorType))
+		{
+			return;
+		}
+
+		if (!ensure(InOutCollection.GetAttributeType(LabelScaleAttributeName, LabelsGroupName) == FManagedArrayCollection::EArrayType::FFloatType))
+		{
+			return;
+		}
+
+		const int32 StartIndex = InOutCollection.NumElements(LabelsGroupName);
+
+		InOutCollection.AddElements(Labels.Num(), LabelsGroupName);
+
+		TManagedArray<FString>& LabelTextAttribute = *InOutCollection.FindAttribute<FString>(LabelTextAttributeName, LabelsGroupName);
+		TManagedArray<FVector3f>& LabelPositionAttribute = *InOutCollection.FindAttribute<FVector3f>(LabelPositionAttributeName, LabelsGroupName);
+		TManagedArray<float>& LabelScaleAttribute = *InOutCollection.FindAttribute<float>(LabelScaleAttributeName, LabelsGroupName);
+
+		for (int32 i = 0; i < Labels.Num(); ++i)
+		{
+			LabelPositionAttribute[StartIndex + i] = Labels[i].Get<FVector3f>();
+			LabelTextAttribute[StartIndex + i] = Labels[i].Get<FString>();
+			LabelScaleAttribute[StartIndex + i] = Labels[i].Get<float>();
+		}
+	}
+
+	static void AddPlantSettingsToCollection(FManagedArrayCollection& InOutCollection, const TArray<FPVImportTexture2DPlantSettings>& InPlantSettings, float Scale, const FVector& Offset = FVector::ZeroVector)
 	{
 		InOutCollection.AddGroup(InputRootPointsGroupName);
-		InOutCollection.AddElements(InRootPoints.Num(), InputRootPointsGroupName);
-	 
-		auto& InputRootPointAttribute = InOutCollection.AddAttribute<FVector3f>(RootPointsAttributeName, InputRootPointsGroupName);
-		for (int32 i = 0; i < InRootPoints.Num(); i++)
+		InOutCollection.AddElements(InPlantSettings.Num(), InputRootPointsGroupName);
+		
+		TArray<TTuple<FVector3f, FString, float>> Labels;
+		Labels.Reserve(InPlantSettings.Num());
+
+		auto& InputRootPointAttribute = InOutCollection.AddAttribute<FVector3f>(InputRootPointsAttributeName, InputRootPointsGroupName);
+		for (int32 i = 0; i < InPlantSettings.Num(); ++i)
 		{
-			InputRootPointAttribute[i] = ToVector3(InRootPoints[i].Position, Scale) + FVector3f(Offset);
+			const FVector3f PointPosition = ToVector3(InPlantSettings[i].RootPosition, Scale) + FVector3f(Offset);
+			InputRootPointAttribute[i] = PointPosition;
+			Labels.Emplace(PointPosition, FString::Printf(TEXT("%d"), i), 1.f);
 		}
+
+		AddLabelsToCollection(InOutCollection, Labels);
 	}
 	
 	static void AddDebugPixelValuesToCollection(FManagedArrayCollection& InOutCollection, const TArray<bool>& ImageData, const FVector2i& ImageSize, float Scale)
@@ -191,12 +253,19 @@ namespace Internal
 		MeshGeomCollection->CopyTo(&InOutCollection);
 	}
 
-	static void ArrangePlantsInGrid(FManagedArrayCollection& InCollection, const FVector& InOffset = FVector::ZeroVector)
+	static void ArrangePlantsInGrid(
+		FManagedArrayCollection& InCollection, 
+		const TArray<int32>& PlantNumberToPlantIndices,
+		const FVector& InOffset = FVector::ZeroVector
+	)
 	{
 		const PV::Facades::FPlantFacade PlantFacade(InCollection);
 		const PV::Facades::FBranchFacade BranchFacade(InCollection);
 		PV::Facades::FPointFacade PointFacade(InCollection);
 
+		TArray<TTuple<FVector3f, FString, float>> Labels;
+		Labels.Reserve(PlantNumberToPlantIndices.Num());
+
 		struct FPlantOffsetInfo
 		{
 			int32 PlantNumber;
@@ -229,8 +298,16 @@ namespace Internal
 				return PV::FacadeHelper::EForEachResult::Continue;
 			});
 
+			const int32 PlantIndex = PlantNumberToPlantIndices.IsValidIndex(PlantsOffsetInfo[i].PlantNumber)
+				? PlantNumberToPlantIndices[PlantsOffsetInfo[i].PlantNumber]
+				: INDEX_NONE;
+
+			Labels.Emplace(Offset + FVector3f(0, 0, -10.f), FString::Printf(TEXT("%d"), PlantIndex), 1.0);
+
 			Offset.Y += PlantsOffsetInfo[i].Bounds.Min.Y;
 		}
+
+		AddLabelsToCollection(InCollection, Labels);
 	}
 
 	static void SampleImage(FTextureSource& InTextureSource, FVector2i& InOutSampleSize, TArray<bool>& OutBlackAndWhiteImageData, bool bInvertImage, float WhiteLevel)
@@ -734,24 +811,33 @@ namespace Internal
 		}
 	}
 
-	static void ApplyRootPoints2D(
-		UE::Geometry::FDynamicMesh3& DynamicMesh, 
+	struct FPlantMesh
+	{
+		FMeshConnectedComponents::FComponent MeshComponent;
+		int32 RootVertexIndex = INDEX_NONE;
+	};
+
+	static void MapPlantSettingsToMeshComponents(
+		const UE::Geometry::FDynamicMesh3& DynamicMesh,
 		float Scale,
-		const TArray<FPVImportTexture2DRootPoint>& InRootPoints,
-		FBooleanVertexAttribute& RootPointAttribute,
-		FBooleanVertexAttribute& EndPointAttribute
+		const TArray<FPVImportTexture2DPlantSettings>& InputPlantSettings,
+		TArray<FPlantMesh>& OutPlantMapping
 	)
 	{
-		// Finds the closest mesh component, and vertex in said mesh component, for each point in InRootPoints
-
 		using namespace PV::DynamicMeshHelper;
 
+		/*
+		* 1. Compute the distance to the closest mesh component for each InputPlantSettings
+		* 2. For each mesh component, select the closest InputPlantSettings
+		* 3. Determine which vertex to use as root point
+		*/
+
 		FMeshConnectedComponents Components(&DynamicMesh);
 		Components.FindConnectedTriangles();
 
 		TArray<FDynamicMeshAABBTree3> AABBTrees;
 		AABBTrees.Reserve(Components.Components.Num());
-		
+
 		for (const auto& Component : Components)
 		{
 			FDynamicMeshAABBTree3 DynamicMeshAABBTree(&DynamicMesh, false);
@@ -759,95 +845,153 @@ namespace Internal
 
 			AABBTrees.Add(MoveTemp(DynamicMeshAABBTree));
 		}
-	
-		struct FMeshIndexDistPair
+
+		struct FClosestTriangle
 		{
 			int32 ComponentIndex = INDEX_NONE;
 			double DistSqr = TNumericLimits<double>::Max();
 			int32 ClosestTriangle = INDEX_NONE;
 		};
-		TArray<FMeshIndexDistPair> RootPositionsClosestMesh;
-		RootPositionsClosestMesh.SetNum(InRootPoints.Num());
-		
-		for (int32 i = 0; i < InRootPoints.Num(); i++)
+		TArray<FClosestTriangle> PlantSettingsClosestTriangle;
+		PlantSettingsClosestTriangle.SetNum(InputPlantSettings.Num());
+
+		for (int32 PlantIndex = 0; PlantIndex < InputPlantSettings.Num(); ++PlantIndex)
 		{
-			FMeshIndexDistPair& ClosestMesh = RootPositionsClosestMesh[i];
-		
+			const FPVImportTexture2DPlantSettings& PlantSettings = InputPlantSettings[PlantIndex];
+			FClosestTriangle& ClosestTriangle = PlantSettingsClosestTriangle[PlantIndex];
+
 			for (int32 ComponentIndex = 0; ComponentIndex < Components.Num(); ComponentIndex++)
 			{
 				const FDynamicMeshAABBTree3& AABBTree = AABBTrees[ComponentIndex];
-				
+
 				double DistSqr = 0;
-				const FVector3d RootPoint = ToVector3<double, float>(InRootPoints[i].Position, Scale);
+				const FVector3d RootPoint = ToVector3<double, float>(PlantSettings.RootPosition, Scale);
 				const int32 NearTriID = AABBTree.FindNearestTriangle(RootPoint, DistSqr);
 				if (NearTriID != INDEX_NONE)
 				{
 					UE::Geometry::FDistPoint3Triangle3d Query = FDynamicMeshQueries3::TriangleDistance(DynamicMesh, NearTriID, RootPoint);
 					DistSqr = Query.GetSquared();
-					if (DistSqr < ClosestMesh.DistSqr)
+					if (DistSqr < ClosestTriangle.DistSqr)
 					{
-						ClosestMesh.ComponentIndex = ComponentIndex;
-						ClosestMesh.DistSqr = DistSqr;
-						ClosestMesh.ClosestTriangle = NearTriID;
+						ClosestTriangle.ComponentIndex = ComponentIndex;
+						ClosestTriangle.DistSqr = DistSqr;
+						ClosestTriangle.ClosestTriangle = NearTriID;
 					}
 				}
 			}
 		}
-		
-		struct FRootPositionIndexDistPair
+
+		struct FPlantIndexDistPair
 		{
-			int32 RootPositionIndex = INDEX_NONE;
+			int32 PlantIndex = INDEX_NONE;
 			double DistSqr = TNumericLimits<double>::Max();
 		};
-		TArray<FRootPositionIndexDistPair> MeshClosestRootPosition;
-		MeshClosestRootPosition.SetNum(Components.Num());
-		
-		for (int32 i = 0; i < RootPositionsClosestMesh.Num(); i++)
+		TArray<FPlantIndexDistPair> ComponentToPlantDist;
+		ComponentToPlantDist.SetNum(Components.Num());
+
+		for (int32 PlantIndex = 0; PlantIndex < InputPlantSettings.Num(); ++PlantIndex)
 		{
-			const int32 ComponentIndex = RootPositionsClosestMesh[i].ComponentIndex;
+			const int32 ComponentIndex = PlantSettingsClosestTriangle[PlantIndex].ComponentIndex;
 			if (ComponentIndex == INDEX_NONE)
 			{
-				// Unable to find a closest mesh (mesh generation has most likely failed).
 				continue;
 			}
-			
-			const double DistSqr = RootPositionsClosestMesh[i].DistSqr;
-			
-			if (DistSqr < MeshClosestRootPosition[ComponentIndex].DistSqr)
+
+			if (PlantSettingsClosestTriangle[PlantIndex].DistSqr < ComponentToPlantDist[ComponentIndex].DistSqr)
 			{
-				MeshClosestRootPosition[ComponentIndex].RootPositionIndex = i;
-				MeshClosestRootPosition[ComponentIndex].DistSqr = DistSqr;
+				ComponentToPlantDist[ComponentIndex].PlantIndex = PlantIndex;
+				ComponentToPlantDist[ComponentIndex].DistSqr = PlantSettingsClosestTriangle[PlantIndex].DistSqr;
 			}
 		}
-		
+
+		OutPlantMapping.SetNum(InputPlantSettings.Num());
 		for (int32 ComponentIndex = 0; ComponentIndex < Components.Num(); ComponentIndex++)
 		{
-			if (MeshClosestRootPosition[ComponentIndex].RootPositionIndex == INDEX_NONE)
+			if (ComponentToPlantDist[ComponentIndex].PlantIndex == INDEX_NONE)
 			{
-				// Unable to find a root position for this mesh
+				// Unable to map this plant setting to a mesh component
 				continue;
 			}
-			
-			const int32 RootPositionIndex = MeshClosestRootPosition[ComponentIndex].RootPositionIndex;
-			
-			const FVector3d& RootPosition = ToVector3<double, float>(InRootPoints[RootPositionIndex].Position, Scale);
-			const int32 TriangleID = RootPositionsClosestMesh[RootPositionIndex].ClosestTriangle;
+
+			const int32 PlantIndex = ComponentToPlantDist[ComponentIndex].PlantIndex;
+			const FPVImportTexture2DPlantSettings& PlantSettings = InputPlantSettings[PlantIndex];
+
+			const FVector3d& RootPosition = ToVector3<double, float>(PlantSettings.RootPosition, Scale);
+			const int32 TriangleID = PlantSettingsClosestTriangle[PlantIndex].ClosestTriangle;
 			const FIndex3i& Triangle = DynamicMesh.GetTriangle(TriangleID);
 			const int32 ClosestTriangleVertex = FMath::Min3Index(
 				(RootPosition - DynamicMesh.GetVertex(Triangle[0])).SizeSquared(),
 				(RootPosition - DynamicMesh.GetVertex(Triangle[1])).SizeSquared(),
 				(RootPosition - DynamicMesh.GetVertex(Triangle[2])).SizeSquared()
 			);
-			
+
 			const int32 RootVertexID = Triangle[ClosestTriangleVertex];
 
+			OutPlantMapping[PlantIndex] = { MoveTemp(Components[ComponentIndex]), RootVertexID };
+		}
+	}
+
+	static void MapPlantIndicesToBranchHierarchies(
+		const TArray<FPlantMesh>& PlantMapping,
+		const TArray<PV::DynamicMeshHelper::FMeshBranchHierarchy>& DynamicMeshBranchHierarchies,
+		TArray<int32>& OutBranchHierarchiesToPlantIndices
+	)
+	{
+		TMap<int32, int32> RootVertexPlantIndexMap;
+		RootVertexPlantIndexMap.Reserve(PlantMapping.Num());
+		for (int32 PlantIndex = 0; PlantIndex < PlantMapping.Num(); ++PlantIndex)
+		{
+			const FPlantMesh& PlantMesh = PlantMapping[PlantIndex];
+			RootVertexPlantIndexMap.Emplace(PlantMesh.RootVertexIndex, PlantIndex);
+		}
+
+		OutBranchHierarchiesToPlantIndices.SetNum(DynamicMeshBranchHierarchies.Num());
+		for (int32& PlantIndex : OutBranchHierarchiesToPlantIndices)
+		{
+			PlantIndex = INDEX_NONE;
+		}
+
+		for (int32 i = 0; i < DynamicMeshBranchHierarchies.Num(); ++i)
+		{
+			check(DynamicMeshBranchHierarchies[i].Branches.Num() > 0);
+			check(DynamicMeshBranchHierarchies[i].Branches[0].BranchPoints.Num() > 0);
+
+			const int32 RootVertexID = DynamicMeshBranchHierarchies[i].Branches[0].BranchPoints[0];
+			if (int32* PlantIndex = RootVertexPlantIndexMap.Find(RootVertexID))
+			{
+				OutBranchHierarchiesToPlantIndices[i] = *PlantIndex;
+			}
+		}
+	}
+
+	static void ApplyPlantMappingToDynamicMesh(
+		UE::Geometry::FDynamicMesh3& DynamicMesh,
+		const TArray<FPVImportTexture2DPlantSettings>& InputPlantSettings,
+		const TArray<FPlantMesh>& InPlantMapping,
+		FBooleanVertexAttribute& RootPointAttribute,
+		FBooleanVertexAttribute& EndPointAttribute
+	)
+	{
+		check(InputPlantSettings.Num() == InPlantMapping.Num());
+
+		for (int32 PlantIndex = 0; PlantIndex < InputPlantSettings.Num(); ++PlantIndex)
+		{
+			const FPlantMesh& PlantMesh = InPlantMapping[PlantIndex];
+			if (PlantMesh.RootVertexIndex == INDEX_NONE)
+			{
+				continue;
+			}
+
+			const int32 RootVertexID = PlantMesh.RootVertexIndex;
+			const FPVImportTexture2DPlantSettings& PlantSettings = InputPlantSettings[PlantIndex];
+
 			RootPointAttribute.SetValue(RootVertexID, true);
 			EndPointAttribute.SetValue(RootVertexID, false);
 
 			const FVector VertexTranslation = DynamicMesh.GetVertex(RootVertexID);
-			const FQuat VertexRotation = FQuat(FVector::ForwardVector, FMath::DegreesToRadians(InRootPoints[RootPositionIndex].Rotation));
+			const FQuat VertexRotation = FQuat(FVector::ForwardVector, FMath::DegreesToRadians(PlantSettings.Rotation));
 
-			for (int32 VertexIndex : PV::DynamicMeshHelper::GetUniqueVertexIndices(DynamicMesh, Components[ComponentIndex].Indices))
+			for (int32 VertexIndex : PV::DynamicMeshHelper::GetUniqueVertexIndices(DynamicMesh, PlantMesh.MeshComponent.Indices))
 			{
 				FVector Vertex = DynamicMesh.GetVertex(VertexIndex);
 				Vertex -= VertexTranslation;
@@ -862,8 +1006,6 @@ PV::Texture2DImport::EImportResult PV::Texture2DImport::ImportGrowthDataFromText
 {
 	using namespace PV::DynamicMeshHelper;
 	using namespace Internal;
-
-	FManagedArrayCollection& OutCollection = Output.OutputCollection;
 	
 	if (!InParams.Texture2DAsset || !InParams.Texture2DAsset->Source.IsValid())
 	{
@@ -875,12 +1017,6 @@ PV::Texture2DImport::EImportResult PV::Texture2DImport::ImportGrowthDataFromText
 		return EImportResult::InvalidScale;
 	}
 
-	if (InParams.DebugState == EPVImportTexture2DDebugState::Texture)
-	{
-		AddTextureQuadToCollection(OutCollection, *InParams.Texture2DAsset, TextureVisualizationMaterial, InParams.Scale);
-		return EImportResult::Success;
-	}
-
 	FVector2i SampleSize = FVector2i(InParams.SampleWidth, InParams.SampleHeight);
 	if (SampleSize.X <= 1 || SampleSize.Y <= 1)
 	{
@@ -897,51 +1033,25 @@ PV::Texture2DImport::EImportResult PV::Texture2DImport::ImportGrowthDataFromText
 	
 	if (InParams.DebugState == EPVImportTexture2DDebugState::SampleImage)
 	{
-		AddDebugPixelValuesToCollection(OutCollection, BlackAndWhiteImageData, SampleSize, InParams.Scale);
-		AddTextureQuadToCollection(OutCollection, *InParams.Texture2DAsset, TextureVisualizationMaterial, InParams.Scale);
-		return EImportResult::Success;
+		FManagedArrayCollection& DebugCollection = Output.DebugCollection.Emplace();
+		AddDebugPixelValuesToCollection(DebugCollection, BlackAndWhiteImageData, SampleSize, InParams.Scale);
+		AddTextureQuadToCollection(DebugCollection, *InParams.Texture2DAsset, TextureVisualizationMaterial, InParams.Scale);
 	}
 
 	TArray<FPolygon2f> PerimeterCurves;
 	TraceImage(BlackAndWhiteImageData, SampleSize, PerimeterCurves);
-	if (InParams.DebugState == EPVImportTexture2DDebugState::TraceImage)
-	{
-		AddPerimeterCurvesToCollection(OutCollection, PerimeterCurves, InParams.Scale);
-		AddTextureQuadToCollection(OutCollection, *InParams.Texture2DAsset, TextureVisualizationMaterial, InParams.Scale);
-		return EImportResult::Success;
-	}
-	
 	RemoveSmallCurves(PerimeterCurves, InParams.MinBoundsArea);
-	if (InParams.DebugState == EPVImportTexture2DDebugState::RemoveSmallCurves)
-	{
-		AddPerimeterCurvesToCollection(OutCollection, PerimeterCurves, InParams.Scale);
-		AddTextureQuadToCollection(OutCollection, *InParams.Texture2DAsset, TextureVisualizationMaterial, InParams.Scale);
-		return EImportResult::Success;
-	}
-	
 	SmoothCurves(PerimeterCurves, InParams.SmoothingIterations);
-	if (InParams.DebugState == EPVImportTexture2DDebugState::SmoothCurves)
-	{
-		AddPerimeterCurvesToCollection(OutCollection, PerimeterCurves, InParams.Scale);
-		AddTextureQuadToCollection(OutCollection, *InParams.Texture2DAsset, TextureVisualizationMaterial, InParams.Scale);
-		return EImportResult::Success;
-	}
-	
 	SimplifyCurves(PerimeterCurves, InParams.SimplificationAmount);
-	if (InParams.DebugState == EPVImportTexture2DDebugState::SimplifyCurves)
-	{
-		AddPerimeterCurvesToCollection(OutCollection, PerimeterCurves, InParams.Scale);
-		AddTextureQuadToCollection(OutCollection, *InParams.Texture2DAsset, TextureVisualizationMaterial, InParams.Scale);
-		return EImportResult::Success;
-	}
-	
+
 	UE::Geometry::FDynamicMesh3 DynamicMesh;
 	TriangulateCurvesAndSplitEdges(PerimeterCurves, DynamicMesh, InParams.Scale);
+
 	if (InParams.DebugState == EPVImportTexture2DDebugState::TriangulateCurves)
 	{
-		AddDynamicMeshToCollection(DynamicMesh, OutCollection);
-		PerimeterAttributeDefinition.CopyAttributeToCollection(DynamicMesh, OutCollection);
-		return EImportResult::Success;
+		FManagedArrayCollection& DebugCollection = Output.DebugCollection.Emplace();
+		AddDynamicMeshToCollection(DynamicMesh, DebugCollection);
+		PerimeterAttributeDefinition.CopyAttributeToCollection(DynamicMesh, DebugCollection);
 	}
 
 	ComputeVertexScale2D(
@@ -949,11 +1059,12 @@ PV::Texture2DImport::EImportResult PV::Texture2DImport::ImportGrowthDataFromText
 		PerimeterAttributeDefinition.GetOrAttachAttribute(DynamicMesh),
 		VertexScaleAttributeDefinition->GetOrAttachAttribute(DynamicMesh)
 	);
+
 	if (InParams.DebugState == EPVImportTexture2DDebugState::ComputeVertexScale)
 	{
-		AddDynamicMeshToCollection(DynamicMesh, OutCollection);
-		VertexScaleAttributeDefinition->CopyAttributeToCollection(DynamicMesh, OutCollection);
-		return EImportResult::Success;
+		FManagedArrayCollection& DebugCollection = Output.DebugCollection.Emplace();
+		AddDynamicMeshToCollection(DynamicMesh, DebugCollection);
+		VertexScaleAttributeDefinition->CopyAttributeToCollection(DynamicMesh, DebugCollection);
 	}
 
 	FindEndPoints2D(
@@ -963,18 +1074,13 @@ PV::Texture2DImport::EImportResult PV::Texture2DImport::ImportGrowthDataFromText
 		EndPointAttributeDefinition->GetOrAttachAttribute(DynamicMesh),
 		PerimeterAttributeDefinition.GetOrAttachAttribute(DynamicMesh)
 	);
-	if (InParams.DebugState == EPVImportTexture2DDebugState::FindTips)
-	{
-		AddDynamicMeshToCollection(DynamicMesh, OutCollection);
-		EndPointAttributeDefinition->CopyAttributeToCollection(DynamicMesh, OutCollection);
- 		return EImportResult::Success;
-	}
 
 	const float HypotenuseLength = FMath::Sqrt(InParams.Scale * InParams.Scale * 2);
 	const FVector VisualizationOffset = FVector(0, -InParams.Scale / 2, HypotenuseLength);
+
 	// Add end points visualization before applying root points as that transforms the DynamicMesh
 	AddDynamicMeshVerticesToCollection(
-		Output.VisualizationCollection, 
+		Output.VisualizationCollection,
 		TEXT("EndPoints"), 
 		TEXT("Vertex"), 
 		DynamicMesh, 
@@ -982,26 +1088,26 @@ PV::Texture2DImport::EImportResult PV::Texture2DImport::ImportGrowthDataFromText
 		VisualizationOffset
 	);
 
-	ApplyRootPoints2D(
-		DynamicMesh, 
+	TArray<FPlantMesh> PlantMapping;
+	MapPlantSettingsToMeshComponents(
+		DynamicMesh,
 		InParams.Scale,
-		InParams.RootPoints,
+		InParams.PlantSettings,
+		PlantMapping
+	);
+
+	ApplyPlantMappingToDynamicMesh(
+		DynamicMesh,
+		InParams.PlantSettings,
+		PlantMapping,
 		RootPointAttributeDefinition->GetOrAttachAttribute(DynamicMesh),
 		EndPointAttributeDefinition->GetOrAttachAttribute(DynamicMesh)
 	);
-	if (InParams.DebugState == EPVImportTexture2DDebugState::ApplyRootPoints)
-	{
-		AddTextureQuadToCollection(OutCollection, *InParams.Texture2DAsset, TextureVisualizationMaterial, InParams.Scale);
-		AddDynamicMeshToCollection(DynamicMesh, OutCollection);
-		AddRootPointsToCollection(OutCollection, InParams.RootPoints, InParams.Scale);
-		EndPointAttributeDefinition->CopyAttributeToCollection(DynamicMesh, OutCollection);
-		RootPointAttributeDefinition->CopyAttributeToCollection(DynamicMesh, OutCollection);
- 		return EImportResult::Success;
-	}
 
 	auto& PerimeterAttribute = PerimeterAttributeDefinition.GetOrAttachAttribute(DynamicMesh);
-	TArray<FPVBranchHierarchyDescription> BranchHierarchyDescriptions;
-	ComputeBranchHierarchies(
+	
+	TArray<FMeshBranchHierarchy> DynamicMeshBranchHierarchies;
+	ComputeDynamicMeshBranchHierarchies(
 		DynamicMesh,
 		RootPointAttributeDefinition->GetOrAttachAttribute(DynamicMesh),
 		EndPointAttributeDefinition->GetOrAttachAttribute(DynamicMesh),
@@ -1009,34 +1115,64 @@ PV::Texture2DImport::EImportResult PV::Texture2DImport::ImportGrowthDataFromText
 		LengthFromRootAttributeDefinition->GetOrAttachAttribute(DynamicMesh),
 		NextBranchPointAttributeDefinition->GetOrAttachAttribute(DynamicMesh),
 		[&](int32 VertexID) { return !PerimeterAttribute.GetValue(VertexID); },
+		DynamicMeshBranchHierarchies
+	);
+
+	TArray<FPVBranchHierarchyDescription> BranchHierarchyDescriptions;
+	ConvertDynamicMeshBranchHierarchiesToBranchDescription(
+		DynamicMesh,
+		VertexScaleAttributeDefinition->GetOrAttachAttribute(DynamicMesh),
+		EndPointAttributeDefinition->GetOrAttachAttribute(DynamicMesh),
+		DynamicMeshBranchHierarchies,
 		BranchHierarchyDescriptions
 	);
+
+	check(DynamicMeshBranchHierarchies.Num() == BranchHierarchyDescriptions.Num());
+
 	if (InParams.DebugState == EPVImportTexture2DDebugState::ComputeBranchHierarchies)
 	{
-		AddDynamicMeshToCollection(DynamicMesh, OutCollection);
-		AddBranchHierarchiesToCollection(BranchHierarchyDescriptions, OutCollection);
-		RootPointAttributeDefinition->CopyAttributeToCollection(DynamicMesh, OutCollection);
-		EndPointAttributeDefinition->CopyAttributeToCollection(DynamicMesh, OutCollection);
-		VertexScaleAttributeDefinition->CopyAttributeToCollection(DynamicMesh, OutCollection);
-		LengthFromRootAttributeDefinition->CopyAttributeToCollection(DynamicMesh, OutCollection);
-		NextBranchPointAttributeDefinition->CopyAttributeToCollection(DynamicMesh, OutCollection);
-
-		return EImportResult::Success;
+		FManagedArrayCollection& DebugCollection = Output.DebugCollection.Emplace();
+		AddDynamicMeshToCollection(DynamicMesh, DebugCollection);
+		AddBranchHierarchiesToCollection(BranchHierarchyDescriptions, DebugCollection);
+		RootPointAttributeDefinition->CopyAttributeToCollection(DynamicMesh, DebugCollection);
+		EndPointAttributeDefinition->CopyAttributeToCollection(DynamicMesh, DebugCollection);
+		VertexScaleAttributeDefinition->CopyAttributeToCollection(DynamicMesh, DebugCollection);
+		LengthFromRootAttributeDefinition->CopyAttributeToCollection(DynamicMesh, DebugCollection);
+		NextBranchPointAttributeDefinition->CopyAttributeToCollection(DynamicMesh, DebugCollection);
 	}
 
-	FPVPlantDescription PlantDescription;
-	PlantDescription.BranchHierarchies = MoveTemp(BranchHierarchyDescriptions);
+	TArray<int32> BranchHierarchiesToPlantIndices;
+	MapPlantIndicesToBranchHierarchies(PlantMapping, DynamicMeshBranchHierarchies, BranchHierarchiesToPlantIndices);
 
-	if (!PV::Import::Helper::GenerateGrowthDataFromPlantDescription(OutCollection, PlantDescription))
+	// TODO: It would probably be faster to combine all the outputs of the PlantCollections
+	if (!PV::ImportHelper::GenerateGrowthDataFromBranchHierarchies(Output.CombinedPlantCollection, BranchHierarchyDescriptions))
 	{
 		return EImportResult::FailedToGenerateGrowthData;
 	}
 
-	OutCollection.CopyTo(&Output.VisualizationCollection);
-	ArrangePlantsInGrid(Output.VisualizationCollection, VisualizationOffset);
-	AddTextureQuadToCollection(Output.VisualizationCollection, *InParams.Texture2DAsset, TextureVisualizationMaterial, InParams.Scale, VisualizationOffset);
-	AddPerimeterCurvesToCollection(Output.VisualizationCollection, PerimeterCurves, InParams.Scale, VisualizationOffset);
-	AddRootPointsToCollection(Output.VisualizationCollection, InParams.RootPoints, InParams.Scale, VisualizationOffset);
+	Output.PlantCollections.SetNum(InParams.PlantSettings.Num());
+	for (int32 i = 0; i < BranchHierarchyDescriptions.Num(); ++i)
+	{
+		const int32 PlantIndex = BranchHierarchiesToPlantIndices[i];
+		if (PlantIndex == INDEX_NONE)
+		{
+			continue;
+		}
+
+		if (!PV::ImportHelper::GenerateGrowthDataFromBranchHierarchy(Output.PlantCollections[PlantIndex], BranchHierarchyDescriptions[i]))
+		{
+			return EImportResult::FailedToGenerateGrowthData;
+		}
+	}
+
+	// Build visualization output
+	{
+		Output.CombinedPlantCollection.CopyTo(&Output.VisualizationCollection);
+		ArrangePlantsInGrid(Output.VisualizationCollection, BranchHierarchiesToPlantIndices, VisualizationOffset);
+		AddTextureQuadToCollection(Output.VisualizationCollection, *InParams.Texture2DAsset, TextureVisualizationMaterial, InParams.Scale, VisualizationOffset);
+		AddPerimeterCurvesToCollection(Output.VisualizationCollection, PerimeterCurves, InParams.Scale, VisualizationOffset);
+		AddPlantSettingsToCollection(Output.VisualizationCollection, InParams.PlantSettings, InParams.Scale, VisualizationOffset);
+	}
 
 	return EImportResult::Success;
 }
diff --git a/Engine/Plugins/Experimental/ProceduralVegetationEditor/Source/ProceduralVegetationEditor/Private/Importer/PVPlantDescription.cpp b/Engine/Plugins/Experimental/ProceduralVegetationEditor/Source/ProceduralVegetationEditor/Private/Importer/PVPlantDescription.cpp
deleted file mode 100644
diff --git a/Engine/Plugins/Experimental/ProceduralVegetationEditor/Source/ProceduralVegetationEditor/Private/Importer/PVPlantDescription.h b/Engine/Plugins/Experimental/ProceduralVegetationEditor/Source/ProceduralVegetationEditor/Private/Importer/PVPlantDescription.h
deleted file mode 100644
diff --git a/Engine/Plugins/Experimental/ProceduralVegetationEditor/Source/ProceduralVegetationEditor/Private/PVEditorInterfaceImpl.cpp b/Engine/Plugins/Experimental/ProceduralVegetationEditor/Source/ProceduralVegetationEditor/Private/PVEditorInterfaceImpl.cpp
--- a/Engine/Plugins/Experimental/ProceduralVegetationEditor/Source/ProceduralVegetationEditor/Private/PVEditorInterfaceImpl.cpp
+++ b/Engine/Plugins/Experimental/ProceduralVegetationEditor/Source/ProceduralVegetationEditor/Private/PVEditorInterfaceImpl.cpp
@@ -3,7 +3,6 @@
 #include "PVEditorInterfaceImpl.h"
 
 #include "Importer/PVImportHelpers.h"
-#include "Importer/PVPlantDescription.h"
 #include "Importer/PVImporter_StaticMesh.h"
 #include "Importer/PVImporter_Texture2D.h"
 #include "PVImportStaticMeshParams.h"
@@ -28,11 +27,10 @@ bool PV::FPVEditorInterfaceImpl::ImportGrowthDataFromTexture2D(const FPVImportTe
 
 TSharedRef<FManagedArrayCollection> PV::FPVEditorInterfaceImpl::CreateMockTreeGrowthDataCollection() const
 {
-	FPVPlantDescription PlantDescription;
-	PlantDescription.BranchHierarchies.Add(PV::Import::Helper::CreateBranchHierarchyFromPoints(PVMockTree::Points, PVMockTree::PointsRadii, PVMockTree::BranchPointIndices));
+	const FPVBranchHierarchyDescription BranchHierarchy = PV::ImportHelper::CreateBranchHierarchyFromPoints(PVMockTree::Points, PVMockTree::PointsRadii, PVMockTree::BranchPointIndices);
 
 	TSharedRef<FManagedArrayCollection> OutCollection = MakeShared<FManagedArrayCollection>();
-	const bool bResult = PV::Import::Helper::GenerateGrowthDataFromPlantDescription(*OutCollection, PlantDescription);
+	const bool bResult = PV::ImportHelper::GenerateGrowthDataFromBranchHierarchy(*OutCollection, BranchHierarchy);
 	check(bResult != false); // Mock tree is staticly defined to always be valid, so it should always return valid growth data
 
 	return OutCollection;
diff --git a/Engine/Plugins/Experimental/ProceduralVegetationEditor/Source/ProceduralVegetationEditor/Private/Tests/Importer/PVImportTests.cpp b/Engine/Plugins/Experimental/ProceduralVegetationEditor/Source/ProceduralVegetationEditor/Private/Tests/Importer/PVImportTests.cpp
--- a/Engine/Plugins/Experimental/ProceduralVegetationEditor/Source/ProceduralVegetationEditor/Private/Tests/Importer/PVImportTests.cpp
+++ b/Engine/Plugins/Experimental/ProceduralVegetationEditor/Source/ProceduralVegetationEditor/Private/Tests/Importer/PVImportTests.cpp
@@ -22,7 +22,7 @@ namespace PVMockTree
 {
 	FPVBranchHierarchyDescription CreateBranchHierarchyDescription()
 	{
-		return PV::Import::Helper::CreateBranchHierarchyFromPoints(PVMockTree::Points, PVMockTree::PointsRadii, PVMockTree::BranchPointIndices);
+		return PV::ImportHelper::CreateBranchHierarchyFromPoints(PVMockTree::Points, PVMockTree::PointsRadii, PVMockTree::BranchPointIndices);
 	}
 }
 
@@ -36,7 +36,7 @@ PV_IMPORTER_TEST(FillGrowthDataFromBranchHierarchyDescription)
 	const int32 NumFoliageMeshes = 0;
 
 	FManagedArrayCollection ExpectedCollection;
-	PV::Import::Helper::CreateEmptyGrowthData(ExpectedCollection, NumBranchPoints, NumFoliagePoints, NumBranches, NumFoliageMeshes);
+	PV::ImportHelper::CreateEmptyGrowthData(ExpectedCollection, NumBranchPoints, NumFoliagePoints, NumBranches, NumFoliageMeshes);
 
 	PV::Facades::FPointFacade ExpectedPointFacade(ExpectedCollection);
 	PV::Facades::FBranchFacade ExpectedBranchFacade(ExpectedCollection);
@@ -83,12 +83,12 @@ PV_IMPORTER_TEST(FillGrowthDataFromBranchHierarchyDescription)
 	}
 
 	FManagedArrayCollection ActualCollection;
-	PV::Import::Helper::CreateEmptyGrowthData(ActualCollection, NumBranchPoints, NumFoliagePoints, NumBranches, NumFoliageMeshes);
+	PV::ImportHelper::CreateEmptyGrowthData(ActualCollection, NumBranchPoints, NumFoliagePoints, NumBranches, NumFoliageMeshes);
 
 	UTEST_EQUAL("CorrectNrOfPointsInGroup", ActualCollection.NumElements(PV::GroupNames::PointGroup), NumBranchPoints);
 	UTEST_EQUAL("CorrectNrOfBranchesInGroup", ActualCollection.NumElements(PV::GroupNames::BranchGroup), NumBranches);
 
-	PV::Import::Helper::FillGrowthDataFromBranchHierarchyDescription(ActualCollection, { BranchHierarchy });
+	PV::ImportHelper::FillGrowthDataFromBranchHierarchyDescription(ActualCollection, { BranchHierarchy });
 
 	const PV::Facades::FPointFacade ActualPointFacade(ActualCollection);
 	const PV::Facades::FBranchFacade ActualBranchFacade(ActualCollection);
@@ -136,13 +136,13 @@ PV_IMPORTER_TEST(ComputeLengthFromRoot)
 	const int32 NumFoliageMeshes = 0;
 
 	FManagedArrayCollection Collection;
-	PV::Import::Helper::CreateEmptyGrowthData(Collection, NumBranchPoints, NumFoliagePoints, NumBranches, NumFoliageMeshes);
-	PV::Import::Helper::FillGrowthDataFromBranchHierarchyDescription(Collection, { BranchHierarchy });
+	PV::ImportHelper::CreateEmptyGrowthData(Collection, NumBranchPoints, NumFoliagePoints, NumBranches, NumFoliageMeshes);
+	PV::ImportHelper::FillGrowthDataFromBranchHierarchyDescription(Collection, { BranchHierarchy });
 
 	PV::Facades::FPointFacade PointFacade(Collection);
 	const PV::Facades::FBranchFacade BranchFacade(Collection);
 	const PV::Facades::FPlantFacade PlantFacade(Collection);
-	PV::Import::Helper::ComputeLengthFromRoot(PointFacade, BranchFacade, PlantFacade);
+	PV::ImportHelper::ComputeLengthFromRoot(PointFacade, BranchFacade, PlantFacade);
 
 	UTEST_NEARLY_EQUAL("Point0_LengthFromRoot", PointFacade.GetLengthFromRoot(0), PVMockTree::Point0_LengthFromRoot, 0.01f);
 	UTEST_NEARLY_EQUAL("Point1_LengthFromRoot", PointFacade.GetLengthFromRoot(1), PVMockTree::Point1_LengthFromRoot, 0.01f);
@@ -278,13 +278,13 @@ PV_IMPORTER_TEST(ComputeBudDevelopment)
 	const int32 NumFoliageMeshes = 0;
 
 	FManagedArrayCollection Collection;
-	PV::Import::Helper::CreateEmptyGrowthData(Collection, NumBranchPoints, NumFoliagePoints, NumBranches, NumFoliageMeshes);
-	PV::Import::Helper::FillGrowthDataFromBranchHierarchyDescription(Collection, { BranchHierarchy });
+	PV::ImportHelper::CreateEmptyGrowthData(Collection, NumBranchPoints, NumFoliagePoints, NumBranches, NumFoliageMeshes);
+	PV::ImportHelper::FillGrowthDataFromBranchHierarchyDescription(Collection, { BranchHierarchy });
 
 	PV::Facades::FPointFacade PointFacade(Collection);
 	const PV::Facades::FBranchFacade BranchFacade(Collection);
 	const PV::Facades::FPlantFacade PlantFacade(Collection);
-	PV::Import::Helper::ComputeBudDevelopment(PointFacade, BranchFacade, PlantFacade);
+	PV::ImportHelper::ComputeBudDevelopment(PointFacade, BranchFacade, PlantFacade);
 
 	UTEST_TRUE("Point0_BudDevelopment_Num", PointFacade.GetBudDevelopment(0).Num() == ExpectedNumBudDevelopmentValues);
 	UTEST_TRUE("Point1_BudDevelopment_Num", PointFacade.GetBudDevelopment(1).Num() == ExpectedNumBudDevelopmentValues);
@@ -376,13 +376,13 @@ PV_IMPORTER_TEST(ComputeBudDirections)
 	const int32 NumFoliageMeshes = 0;
 
 	FManagedArrayCollection Collection;
-	PV::Import::Helper::CreateEmptyGrowthData(Collection, NumBranchPoints, NumFoliagePoints, NumBranches, NumFoliageMeshes);
-	PV::Import::Helper::FillGrowthDataFromBranchHierarchyDescription(Collection, { BranchHierarchy });
+	PV::ImportHelper::CreateEmptyGrowthData(Collection, NumBranchPoints, NumFoliagePoints, NumBranches, NumFoliageMeshes);
+	PV::ImportHelper::FillGrowthDataFromBranchHierarchyDescription(Collection, { BranchHierarchy });
 
 	PV::Facades::FPointFacade PointFacade(Collection);
 	const PV::Facades::FBranchFacade BranchFacade(Collection);
 	const PV::Facades::FPlantFacade PlantFacade(Collection);
-	PV::Import::Helper::ComputeBudDirections(PointFacade, BranchFacade, PlantFacade);
+	PV::ImportHelper::ComputeBudDirections(PointFacade, BranchFacade, PlantFacade);
 
 	UTEST_TRUE("Point0_BudDirections_Num", PointFacade.GetBudDirection(0).Num() == ExpectedNumBudDirectionValues);
 	UTEST_TRUE("Point1_BudDirections_Num", PointFacade.GetBudDirection(1).Num() == ExpectedNumBudDirectionValues);
@@ -424,18 +424,18 @@ PV_IMPORTER_TEST(SplitFusedPoints)
 	const int32 NumFoliageMeshes = 0;
 
 	FManagedArrayCollection Collection;
-	PV::Import::Helper::CreateEmptyGrowthData(Collection, NumBranchPoints, NumFoliagePoints, NumBranches, NumFoliageMeshes);
-	PV::Import::Helper::FillGrowthDataFromBranchHierarchyDescription(Collection, { BranchHierarchy });
+	PV::ImportHelper::CreateEmptyGrowthData(Collection, NumBranchPoints, NumFoliagePoints, NumBranches, NumFoliageMeshes);
+	PV::ImportHelper::FillGrowthDataFromBranchHierarchyDescription(Collection, { BranchHierarchy });
 
 	PV::Facades::FPointFacade PointFacade(Collection);
 	const PV::Facades::FBranchFacade BranchFacade(Collection);
 	const PV::Facades::FPlantFacade PlantFacade(Collection);
-	PV::Import::Helper::ComputeBudDirections(PointFacade, BranchFacade, PlantFacade);
-	PV::Import::Helper::ComputeLengthFromRoot(PointFacade, BranchFacade, PlantFacade);
-	PV::Import::Helper::ComputeBudDevelopment(PointFacade, BranchFacade, PlantFacade);
-	PV::Import::Helper::ComputeBudDirections(PointFacade, BranchFacade, PlantFacade);
+	PV::ImportHelper::ComputeBudDirections(PointFacade, BranchFacade, PlantFacade);
+	PV::ImportHelper::ComputeLengthFromRoot(PointFacade, BranchFacade, PlantFacade);
+	PV::ImportHelper::ComputeBudDevelopment(PointFacade, BranchFacade, PlantFacade);
+	PV::ImportHelper::ComputeBudDirections(PointFacade, BranchFacade, PlantFacade);
 
-	const bool bSplitFusedPointsResult = PV::Import::Helper::SplitFusedPoints(Collection);
+	const bool bSplitFusedPointsResult = PV::ImportHelper::SplitFusedPoints(Collection);
 	UTEST_TRUE("SplitFusedPointsSucceeded", bSplitFusedPointsResult);
 
 	const int32 Branch0_OldRootPointIndex = 0;
@@ -474,7 +474,7 @@ PV_IMPORTER_TEST(SplitFusedPoints)
 PV_IMPORTER_TEST(FindRootBranch)
 {
 	const TArray<TArray<int32>>& Branches = PVMockTree::BranchPointIndices;
-	const int32 RootBranchIndex = PV::Import::Helper::FindRootBranch(PVMockTree::Points, Branches);
+	const int32 RootBranchIndex = PV::ImportHelper::FindRootBranch(PVMockTree::Points, Branches);
 
 	UTEST_NOT_EQUAL("ValidRootBranch", RootBranchIndex, int32(INDEX_NONE));
 
@@ -501,8 +501,8 @@ PV_IMPORTER_TEST(EstimateBranchHierarchy)
 
 		// Naive testing of hierarchy
 		{
-			const int32 RootBranchIndex = PV::Import::Helper::FindRootBranch(PVMockTree::Points, Branches);
-			const TArray<int32> BranchHierarchy = PV::Import::Helper::EstimateBranchHierarchy(PVMockTree::Points, Branches, RootBranchIndex);
+			const int32 RootBranchIndex = PV::ImportHelper::FindRootBranch(PVMockTree::Points, Branches);
+			const TArray<int32> BranchHierarchy = PV::ImportHelper::EstimateBranchHierarchy(PVMockTree::Points, Branches, RootBranchIndex);
 
 			UTEST_EQUAL("CorrectNrOfBranches", BranchHierarchy.Num(), ExpectedBranchHierarchy.Num());
 			for (int32 i = 0; i < Branches.Num(); ++i)
@@ -523,8 +523,8 @@ PV_IMPORTER_TEST(EstimateBranchHierarchy)
 				}
 			}
 
-			const int32 RootBranchIndex = PV::Import::Helper::FindRootBranch(PVMockTree::Points, Branches);
-			const TArray<int32> BranchHierarchy = PV::Import::Helper::EstimateBranchHierarchy(PVMockTree::Points, Branches, RootBranchIndex);
+			const int32 RootBranchIndex = PV::ImportHelper::FindRootBranch(PVMockTree::Points, Branches);
+			const TArray<int32> BranchHierarchy = PV::ImportHelper::EstimateBranchHierarchy(PVMockTree::Points, Branches, RootBranchIndex);
 			
 			UTEST_EQUAL("CorrectNrOfBranches_Reverse", BranchHierarchy.Num(), ExpectedBranchHierarchy.Num());
 			for (int32 i = 0; i < Branches.Num(); ++i)
@@ -588,8 +588,8 @@ PV_IMPORTER_TEST(EstimateBranchHierarchy)
 
 		// Naive testing of hierarchy
 		{
-			const int32 RootBranchIndex = PV::Import::Helper::FindRootBranch(PointsCpy, Branches);
-			const TArray<int32> BranchHierarchy = PV::Import::Helper::EstimateBranchHierarchy(PointsCpy, Branches, RootBranchIndex);
+			const int32 RootBranchIndex = PV::ImportHelper::FindRootBranch(PointsCpy, Branches);
+			const TArray<int32> BranchHierarchy = PV::ImportHelper::EstimateBranchHierarchy(PointsCpy, Branches, RootBranchIndex);
 
 			UTEST_EQUAL("CorrectNrOfBranches", BranchHierarchy.Num(), ExpectedBranchHierarchy.Num());
 			for (int32 i = 0; i < Branches.Num(); ++i)
@@ -610,8 +610,8 @@ PV_IMPORTER_TEST(EstimateBranchHierarchy)
 				}
 			}
 
-			const int32 RootBranchIndex = PV::Import::Helper::FindRootBranch(PointsCpy, Branches);
-			const TArray<int32> BranchHierarchy = PV::Import::Helper::EstimateBranchHierarchy(PointsCpy, Branches, RootBranchIndex);
+			const int32 RootBranchIndex = PV::ImportHelper::FindRootBranch(PointsCpy, Branches);
+			const TArray<int32> BranchHierarchy = PV::ImportHelper::EstimateBranchHierarchy(PointsCpy, Branches, RootBranchIndex);
 
 			UTEST_EQUAL("CorrectNrOfBranches_Reverse", BranchHierarchy.Num(), ExpectedBranchHierarchy.Num());
 			for (int32 i = 0; i < Branches.Num(); ++i)
diff --git a/Engine/Plugins/Experimental/ProceduralVegetationEditor/Source/ProceduralVegetationEditor/Private/Visualizations/DebugVisualization/PVDebugVisualizer.cpp b/Engine/Plugins/Experimental/ProceduralVegetationEditor/Source/ProceduralVegetationEditor/Private/Visualizations/DebugVisualization/PVDebugVisualizer.cpp
--- a/Engine/Plugins/Experimental/ProceduralVegetationEditor/Source/ProceduralVegetationEditor/Private/Visualizations/DebugVisualization/PVDebugVisualizer.cpp
+++ b/Engine/Plugins/Experimental/ProceduralVegetationEditor/Source/ProceduralVegetationEditor/Private/Visualizations/DebugVisualization/PVDebugVisualizer.cpp
@@ -41,21 +41,6 @@ void FPVDebugVisualizer::DrawDebugVisualizations(
 	}
 }
 
-//void FPVDebugVisualizer::DrawVisualizationCollections(const TArray<FPVVisualizationCollection>& VisualizationCollections, FPCGSceneSetupParams& InOutParams)
-//{
-//	for (const FPVVisualizationCollection& VisualizationCollection : VisualizationCollections)
-//	{
-//		for (const FPVVisualizationSettings& VisualizationSetting : VisualizationCollection.VisualizationSettings)
-//		{
-//			FVisualizerDrawContext Context(VisualizationCollection.Collection, VisualizationSetting, InOutParams);
-//			if (const FPVDebugVisualizationPtr VisualizationPtr = CreateVisualizer(VisualizationSetting))
-//			{
-//				VisualizationPtr->Draw(Context);
-//			}
-//		}
-//	}
-//}
-
 void FPVDebugVisualizer::DrawLoopDebug(
 	const TArray<FPVLoopDebugSettings>& LoopDebugVisualizationSettings, 
 	bool bAutoFocusLoopDebug, 
