[PVE] Enable Subgraphs
- Added a subgraph settings that derives from PCGBaseSubgraphSettings. We are deriving from the base:
- Core functionality of the subgraphs resides in the base class.
- We want to only enable it for PVE graphs. This avoids exposing the PCGGraph interface in the details panel.
- For now, parameters on the subgraph are only visible via PCGSubgraphInstance. To avoid further changes, I have kept it as a transient variable.
- Enabled PCG input/output in the subgraph
- Refactored ProceduralVegetation asset code to have UProceduralVegetationInterface, which contains common functions between PV and PVInstance.
- Added new PVInstance class which is required for subgraphing.
- FYI - Currently, visualizations are not working and will require some refactoring around PVSettings.

#Synced-CL 50393756

diff --git a/Engine/Plugins/Experimental/ProceduralVegetationEditor/Source/ProceduralVegetation/Private/Nodes/PVSubgraphSettings.cpp b/Engine/Plugins/Experimental/ProceduralVegetationEditor/Source/ProceduralVegetation/Private/Nodes/PVSubgraphSettings.cpp
new file mode 100644
--- /dev/null
+++ b/Engine/Plugins/Experimental/ProceduralVegetationEditor/Source/ProceduralVegetation/Private/Nodes/PVSubgraphSettings.cpp
@@ -0,0 +1,72 @@
+﻿// Copyright Epic Games, Inc. All Rights Reserved.
+ 
+#include "Nodes/PVSubgraphSettings.h"
+
+#include "ProceduralVegetation.h"
+#include "ProceduralVegetationModule.h"
+#include "Helpers/PCGSettingsHelpers.h"
+#include "Helpers/PVFoliageJSONHelper.h"
+#include "Helpers/PVJSONHelper.h"
+
+#define LOCTEXT_NAMESPACE "PVSubgraph"
+
+#if WITH_EDITOR
+void UPVSubgraphSettings::PostEditChangeProperty(struct FPropertyChangedEvent& PropertyChangedEvent)
+{
+	if (PropertyChangedEvent.Property && PropertyChangedEvent.Property->GetFName() == GET_MEMBER_NAME_CHECKED(UPVSubgraphSettings, SubgraphInstance))
+	{
+		// Also rebuild the overrides
+		InitializeCachedOverridableParams(/*bReset=*/true);
+	}
+
+	Super::PostEditChangeProperty(PropertyChangedEvent);
+}
+
+UPCGGraphInterface* UPVSubgraphSettings::GetSubgraphInterface() const
+{
+	return SubgraphInstance->GetGraph();
+}
+
+UObject* UPVSubgraphSettings::GetJumpTargetForDoubleClick() const
+{
+	return SubgraphInstance->ProceduralVegetation;
+}
+
+#endif
+
+UPVSubgraphSettings::UPVSubgraphSettings(const FObjectInitializer& InObjectInitializer)
+	: Super(InObjectInitializer)
+{
+	SubgraphInstance = InObjectInitializer.CreateDefaultSubobject<UProceduralVegetationInstance>(this, TEXT("ProceduralVegetationSubgraphInstance"));
+}
+
+
+UPCGNode* UPVSubgraphSettings::CreateNode() const
+{
+	return NewObject<UPCGSubgraphNode>();
+}
+
+FString UPVSubgraphSettings::GetAdditionalTitleInformation() const
+{
+	if (auto TargetSubgraph = SubgraphInstance->ProceduralVegetation)
+	{
+		// Use the same transformation than in the palette view to add spaces between uppercase characters
+		return FName::NameToDisplayString(TargetSubgraph->GetName(), /*bIsBool=*/false);
+	}
+	else
+	{
+		return LOCTEXT("NodeTitleExtendedInvalidSubgraph", "Missing Subgraph").ToString();
+	}
+}
+
+FPCGElementPtr UPVSubgraphSettings::CreateElement() const
+{
+	return MakeShared<FPCGSubgraphElement>();
+}
+
+void UPVSubgraphSettings::SetSubgraphInternal(UPCGGraphInterface* InGraph)
+{
+	SubgraphInstance->GraphInstance->SetGraph(InGraph);
+}
+
+#undef LOCTEXT_NAMESPACE
diff --git a/Engine/Plugins/Experimental/ProceduralVegetationEditor/Source/ProceduralVegetation/Private/ProceduralVegetation.cpp b/Engine/Plugins/Experimental/ProceduralVegetationEditor/Source/ProceduralVegetation/Private/ProceduralVegetation.cpp
--- a/Engine/Plugins/Experimental/ProceduralVegetationEditor/Source/ProceduralVegetation/Private/ProceduralVegetation.cpp
+++ b/Engine/Plugins/Experimental/ProceduralVegetationEditor/Source/ProceduralVegetation/Private/ProceduralVegetation.cpp
@@ -98,4 +98,33 @@ void UProceduralVegetation::CreateGraphFromPreset(const TObjectPtr<UProceduralVe
 		Position.Y += 150;
 	}
 }
+
+UProceduralVegetationInstance::UProceduralVegetationInstance(const FObjectInitializer& InObjectInitializer)
+: Super(InObjectInitializer)
+{
+	GraphInstance = InObjectInitializer.CreateDefaultSubobject<UProceduralVegetationGraphInstance>(this, TEXT("ProceduralVegetationGraphInstance"));
+}
+
+void UProceduralVegetationInstance::PostEditChangeProperty(struct FPropertyChangedEvent& PropertyChangedEvent)
+{
+	if (PropertyChangedEvent.Property && PropertyChangedEvent.Property->GetFName() == GET_MEMBER_NAME_CHECKED(UProceduralVegetationInstance, ProceduralVegetation))
+	{
+		if (GraphInstance)
+		{
+			GraphInstance->SetGraph(ProceduralVegetation ? ProceduralVegetation->GetGraph() : nullptr);
+		}
+	}
+	
+	Super::PostEditChangeProperty(PropertyChangedEvent);
+}
+
+void UProceduralVegetationInstance::PostLoad()
+{
+	Super::PostLoad();
+
+	if (GraphInstance && ProceduralVegetation)
+	{
+		GraphInstance->SetGraph(ProceduralVegetation->GetGraph());
+	}
+}
 #endif
diff --git a/Engine/Plugins/Experimental/ProceduralVegetationEditor/Source/ProceduralVegetation/Public/Nodes/PVSubgraphSettings.h b/Engine/Plugins/Experimental/ProceduralVegetationEditor/Source/ProceduralVegetation/Public/Nodes/PVSubgraphSettings.h
new file mode 100644
--- /dev/null
+++ b/Engine/Plugins/Experimental/ProceduralVegetationEditor/Source/ProceduralVegetation/Public/Nodes/PVSubgraphSettings.h
@@ -0,0 +1,54 @@
+﻿// Copyright Epic Games, Inc. All Rights Reserved.
+#pragma once
+
+#include "CoreMinimal.h"
+#include "PCGSubgraph.h"
+#include "ProceduralVegetation.h"
+#include "Nodes/PVBaseSettings.h"
+#include "PVSubgraphSettings.generated.h"
+
+class UProceduralVegetationInstance;
+class UProceduralVegetation;
+
+UCLASS(BlueprintType, ClassGroup = (Procedural))
+class UPVSubgraphSettings : public UPCGBaseSubgraphSettings
+{
+	GENERATED_BODY()
+
+public:
+	UPVSubgraphSettings(const FObjectInitializer& InObjectInitializer);
+	
+#if WITH_EDITOR
+	virtual FName GetDefaultNodeName() const override { return FName(TEXT("ProceduralVegetationSubgraph")); }
+	virtual FText GetDefaultNodeTitle() const override { return NSLOCTEXT("PVSubgraphSettings", "NodeTitle", "Subgraph"); }
+	virtual FText GetNodeTooltipText() const override { return NSLOCTEXT("PVSubgraphSettings", "NodeTooltip", "Create Procedural Vegetation subgraph"); }
+	virtual FLinearColor GetNodeTitleColor() const override { return FLinearColor(FColor::Black); }
+	virtual EPCGSettingsType GetType() const override { return EPCGSettingsType::Subgraph; }
+	virtual UObject* GetJumpTargetForDoubleClick() const override;
+	virtual bool GetPinExtraIcon(const UPCGPin* InPin, FName& OutExtraIcon, FText& OutTooltip) const override { return false; }
+	// TODO: rendering stuff
+	// virtual TArray<EPVRenderType> GetDefaultRenderType() const override { return TArray{ EPVRenderType::PointData }; }
+#endif
+	
+	//~UPCGSettings interface implementation
+	virtual UPCGNode* CreateNode() const override;
+	virtual FString GetAdditionalTitleInformation() const override;
+
+protected:
+	virtual FPCGElementPtr CreateElement() const override;
+	virtual void SetSubgraphInternal(UPCGGraphInterface* InGraph) override;
+	
+#if WITH_EDITOR
+	virtual void PostEditChangeProperty(struct FPropertyChangedEvent& PropertyChangedEvent) override;
+#endif
+
+private:
+	UPROPERTY(VisibleAnywhere, Category = Properties, Instanced, meta = (NoResetToDefault))
+	TObjectPtr<UProceduralVegetationInstance> SubgraphInstance;
+	
+public:
+	
+	virtual UPCGGraphInterface* GetSubgraphInterface() const override;
+
+	virtual UPCGGraphInterface* GetOriginalSubgraphInterface() const override { return SubgraphInstance->GraphInstance; }
+};
\ No newline at end of file
diff --git a/Engine/Plugins/Experimental/ProceduralVegetationEditor/Source/ProceduralVegetation/Public/ProceduralVegetation.h b/Engine/Plugins/Experimental/ProceduralVegetationEditor/Source/ProceduralVegetation/Public/ProceduralVegetation.h
--- a/Engine/Plugins/Experimental/ProceduralVegetationEditor/Source/ProceduralVegetation/Public/ProceduralVegetation.h
+++ b/Engine/Plugins/Experimental/ProceduralVegetationEditor/Source/ProceduralVegetation/Public/ProceduralVegetation.h
@@ -22,9 +22,8 @@ public:
 #endif
 		
 #if WITH_EDITOR
-		// Hides input and output nodes
-		SetHiddenFlagInputNode(/*bHidden=*/true);
-		SetHiddenFlagOutputNode(/*bHidden=*/true);
+		SetHiddenFlagInputNode(/*bHidden=*/false);
+		SetHiddenFlagOutputNode(/*bHidden=*/false);
 #endif
 	}
 
@@ -41,11 +40,21 @@ public:
 	virtual bool SupportsEmbeddedSubgraphs() const override { return false; }
 };
 
+UCLASS(Abstract)
+class PROCEDURALVEGETATION_API UProceduralVegetationInterface : public UObject
+{
+	GENERATED_BODY()
+
+public:
+	virtual UPCGGraphInterface* GetGraph() PURE_VIRTUAL(UProceduralVegetationInterface::GetGraph, return nullptr;)
+	virtual const UPCGGraphInterface* GetGraph() const PURE_VIRTUAL(UProceduralVegetationInterface::GetGraph, return nullptr;)
+};
+
 /**
  * Asset type for procedural plant generation
  */
 UCLASS()
-class PROCEDURALVEGETATION_API UProceduralVegetation : public UObject
+class PROCEDURALVEGETATION_API UProceduralVegetation : public UProceduralVegetationInterface
 {
 	GENERATED_BODY()
 
@@ -53,8 +62,8 @@ class PROCEDURALVEGETATION_API UProceduralVegetation : public UObject
 	TObjectPtr<UProceduralVegetationGraph> Graph;
 
 public:
-	TObjectPtr<UProceduralVegetationGraph>& GetGraph() { return Graph; }
-	const TObjectPtr<UProceduralVegetationGraph>& GetGraph() const { return Graph; }
+	virtual UPCGGraphInterface* GetGraph() override { return Graph; }
+	virtual const UPCGGraphInterface* GetGraph() const override { return Graph; }
 
 	void CreateGraph(const UProceduralVegetationGraph* InGraph = nullptr);
 
@@ -62,3 +71,37 @@ public:
 	void CreateGraphFromPreset(const TObjectPtr<UProceduralVegetationPreset> InPreset);
 #endif
 };
+
+UCLASS(MinimalAPI)
+class UProceduralVegetationGraphInstance : public UPCGGraphInstance
+{
+	GENERATED_BODY()
+};
+
+UCLASS()
+class PROCEDURALVEGETATION_API UProceduralVegetationInstance : public UProceduralVegetationInterface
+{
+	GENERATED_BODY()
+
+	UProceduralVegetationInstance(const FObjectInitializer& InObjectInitializer);
+	
+public:
+	UPROPERTY(EditAnywhere, BlueprintReadOnly, Category = Instance)
+	TObjectPtr<UProceduralVegetationInterface> ProceduralVegetation;
+
+	virtual UPCGGraphInterface* GetGraph() override { return GraphInstance ? GraphInstance->GetGraph() : nullptr; }
+	virtual const UPCGGraphInterface* GetGraph() const override { return GraphInstance ? GraphInstance->Graph : nullptr; }
+	
+	UPROPERTY(BlueprintReadOnly, VisibleAnywhere, Category = Properties, Instanced, Transient, meta = (NoResetToDefault))
+	TObjectPtr<UProceduralVegetationGraphInstance> GraphInstance;
+
+protected:
+#if WITH_EDITOR
+	virtual void PostEditChangeProperty(struct FPropertyChangedEvent& PropertyChangedEvent) override;
+
+public:
+	virtual void PostLoad() override;
+
+protected:
+#endif
+};
diff --git a/Engine/Plugins/Experimental/ProceduralVegetationEditor/Source/ProceduralVegetationEditor/Private/PVEditor.cpp b/Engine/Plugins/Experimental/ProceduralVegetationEditor/Source/ProceduralVegetationEditor/Private/PVEditor.cpp
--- a/Engine/Plugins/Experimental/ProceduralVegetationEditor/Source/ProceduralVegetationEditor/Private/PVEditor.cpp
+++ b/Engine/Plugins/Experimental/ProceduralVegetationEditor/Source/ProceduralVegetationEditor/Private/PVEditor.cpp
@@ -63,7 +63,7 @@ void FPVEditor::Initialize(const EToolkitMode::Type InMode, const TSharedPtr<cla
 	UE_LOG(LogProceduralVegetationEditor, Log, TEXT("PVEditor initialized for %s"), *InProceduralVegetation->GetFName().ToString());
 	ProceduralVegetationBeingEdited = InProceduralVegetation;
 	
-	UProceduralVegetationGraph* ProceduralVegetationGraph = InProceduralVegetation->GetGraph();
+	UProceduralVegetationGraph* ProceduralVegetationGraph = Cast<UProceduralVegetationGraph>(InProceduralVegetation->GetGraph());
 
 	check(ProceduralVegetationGraph);
 
diff --git a/Engine/Plugins/Experimental/ProceduralVegetationEditor/Source/ProceduralVegetationEditor/Private/PVEditorSchema.cpp b/Engine/Plugins/Experimental/ProceduralVegetationEditor/Source/ProceduralVegetationEditor/Private/PVEditorSchema.cpp
--- a/Engine/Plugins/Experimental/ProceduralVegetationEditor/Source/ProceduralVegetationEditor/Private/PVEditorSchema.cpp
+++ b/Engine/Plugins/Experimental/ProceduralVegetationEditor/Source/ProceduralVegetationEditor/Private/PVEditorSchema.cpp
@@ -1,10 +1,13 @@
 ﻿// Copyright Epic Games, Inc. All Rights Reserved.
 
 #include "PVEditorSchema.h"
+
+#include "PCGParamData.h"
 #include "PCGSubgraph.h"
 #include "Elements/PCGCreateAttribute.h"
 #include "Elements/Metadata/PCGMetadataMathsOpElement.h"
 #include "Nodes/PVBaseSettings.h"
+#include "Nodes/PVSubgraphSettings.h"
 
 EPCGElementType UPVEditorSchema::GetElementTypeFiltering() const
 {
@@ -20,7 +23,8 @@ const FPCGGraphEditorCustomization& UPVEditorSchema::GetGraphEditorCustomization
 	{
 			UPVBaseSettings::StaticClass(),
 			UPCGCreateAttributeSetSettings::StaticClass(),
-			UPCGMetadataMathsSettings::StaticClass()
+			UPCGMetadataMathsSettings::StaticClass(),
+			UPVSubgraphSettings::StaticClass(),
 		}
 	};
 
diff --git a/Engine/Plugins/Experimental/ProceduralVegetationEditor/Source/ProceduralVegetationEditor/Private/ProceduralVegetationFactory.cpp b/Engine/Plugins/Experimental/ProceduralVegetationEditor/Source/ProceduralVegetationEditor/Private/ProceduralVegetationFactory.cpp
--- a/Engine/Plugins/Experimental/ProceduralVegetationEditor/Source/ProceduralVegetationEditor/Private/ProceduralVegetationFactory.cpp
+++ b/Engine/Plugins/Experimental/ProceduralVegetationEditor/Source/ProceduralVegetationEditor/Private/ProceduralVegetationFactory.cpp
@@ -25,7 +25,7 @@ UObject* UProceduralVegetationFactory::FactoryCreateNew(UClass* InClass, UObject
 	{
 		NewProceduralVegetation = NewObject<UProceduralVegetation>(InParent, InClass, InName, InFlags, SampleProceduralVegetation);
 
-		TObjectPtr<UProceduralVegetationGraph> SampleGraph = SampleProceduralVegetation->GetGraph();
+		UProceduralVegetationGraph* SampleGraph = Cast<UProceduralVegetationGraph>(SampleProceduralVegetation->GetGraph());
 		NewProceduralVegetation->CreateGraph(SampleGraph);
 		UE_LOG(LogProceduralVegetationEditor, Log, TEXT("New Procedural Vegetation Asset created %s from Preset %s"), *NewProceduralVegetation->GetFName().ToString(), *SampleProceduralVegetation.GetFName().ToString());
 	}
diff --git a/Engine/Plugins/Experimental/ProceduralVegetationEditor/Source/ProceduralVegetationEditor/Private/Tests/ScenarioTests/PVScenarioTests.cpp b/Engine/Plugins/Experimental/ProceduralVegetationEditor/Source/ProceduralVegetationEditor/Private/Tests/ScenarioTests/PVScenarioTests.cpp
--- a/Engine/Plugins/Experimental/ProceduralVegetationEditor/Source/ProceduralVegetationEditor/Private/Tests/ScenarioTests/PVScenarioTests.cpp
+++ b/Engine/Plugins/Experimental/ProceduralVegetationEditor/Source/ProceduralVegetationEditor/Private/Tests/ScenarioTests/PVScenarioTests.cpp
@@ -335,7 +335,7 @@ namespace PVScenarioTests
 				continue;
 			}
 
-			UProceduralVegetationGraph* Graph = ProceduralVegetation->GetGraph();
+			UPCGGraphInterface* Graph = ProceduralVegetation->GetGraph();
 			if (!Graph)
 			{
 				continue;
