[PVE] Expose settings as node pins for modifiers

#Synced-CL 50393756
#Synced-CL 50393756

diff --git a/Engine/Plugins/Experimental/ProceduralVegetationEditor/Source/ProceduralVegetation/Private/Nodes/PVCarveSettings.h b/Engine/Plugins/Experimental/ProceduralVegetationEditor/Source/ProceduralVegetation/Private/Nodes/PVCarveSettings.h
--- a/Engine/Plugins/Experimental/ProceduralVegetationEditor/Source/ProceduralVegetation/Private/Nodes/PVCarveSettings.h
+++ b/Engine/Plugins/Experimental/ProceduralVegetationEditor/Source/ProceduralVegetation/Private/Nodes/PVCarveSettings.h
@@ -31,7 +31,7 @@ protected:
 	virtual FPCGElementPtr CreateElement() const override;
 
 public:
-	UPROPERTY(EditAnywhere, Category = "Carve", meta = (ShowOnlyInnerProperties))
+	UPROPERTY(EditAnywhere, Category = "Carve", meta = (ShowOnlyInnerProperties, PCG_Overridable))
 	FPVCarveParams CarveSettings;
 };
 
diff --git a/Engine/Plugins/Experimental/ProceduralVegetationEditor/Source/ProceduralVegetation/Private/Nodes/PVGravitySettings.h b/Engine/Plugins/Experimental/ProceduralVegetationEditor/Source/ProceduralVegetation/Private/Nodes/PVGravitySettings.h
--- a/Engine/Plugins/Experimental/ProceduralVegetationEditor/Source/ProceduralVegetation/Private/Nodes/PVGravitySettings.h
+++ b/Engine/Plugins/Experimental/ProceduralVegetationEditor/Source/ProceduralVegetation/Private/Nodes/PVGravitySettings.h
@@ -28,7 +28,7 @@ protected:
 	virtual FPCGElementPtr CreateElement() const override;
 
 public:
-	UPROPERTY(EditAnywhere, Category="Gravity", meta=(ShowOnlyInnerProperties))
+	UPROPERTY(EditAnywhere, Category="Gravity", meta=(ShowOnlyInnerProperties, PCG_Overridable))
 	FPVGravityParams GravitySettings;
 
 #if WITH_EDITORONLY_DATA
diff --git a/Engine/Plugins/Experimental/ProceduralVegetationEditor/Source/ProceduralVegetation/Private/Nodes/PVRemoveBranchesSettings.h b/Engine/Plugins/Experimental/ProceduralVegetationEditor/Source/ProceduralVegetation/Private/Nodes/PVRemoveBranchesSettings.h
--- a/Engine/Plugins/Experimental/ProceduralVegetationEditor/Source/ProceduralVegetation/Private/Nodes/PVRemoveBranchesSettings.h
+++ b/Engine/Plugins/Experimental/ProceduralVegetationEditor/Source/ProceduralVegetation/Private/Nodes/PVRemoveBranchesSettings.h
@@ -28,10 +28,10 @@ protected:
 	virtual FPCGElementPtr CreateElement() const override;
 
 public:
-	UPROPERTY(EditAnywhere, Category = Settings, meta=(Tooltip="Sets the rule used to remove branches.\n\nChooses the basis for pruning (e.g., generation, length from root, age, or radius). Determines which branches are candidates for removal."))
+	UPROPERTY(EditAnywhere, Category = Settings, meta=(PCG_Overridable, Tooltip="Sets the rule used to remove branches.\n\nChooses the basis for pruning (e.g., generation, length from root, age, or radius). Determines which branches are candidates for removal."))
 	ERemoveBranchesBasis BranchRemoveBasis = ERemoveBranchesBasis::Length;
 
-	UPROPERTY(EditAnywhere, Category = Settings, meta=(ClampMin = 0f, ClampMax = 1.0f, UIMin = 0f, UIMax = 1.0f, Tooltip="Removal threshold for the chosen basis.\n\nBranches exceeding or falling below this threshold (per the selected basis) are removed. Use lower/higher values to prune sparsely or aggressively."))
+	UPROPERTY(EditAnywhere, Category = Settings, meta=(PCG_Overridable, ClampMin = 0f, ClampMax = 1.0f, UIMin = 0f, UIMax = 1.0f, Tooltip="Removal threshold for the chosen basis.\n\nBranches exceeding or falling below this threshold (per the selected basis) are removed. Use lower/higher values to prune sparsely or aggressively."))
 	float Threshold = 0;
 };
 
diff --git a/Engine/Plugins/Experimental/ProceduralVegetationEditor/Source/ProceduralVegetation/Private/Nodes/PVScaleSettings.h b/Engine/Plugins/Experimental/ProceduralVegetationEditor/Source/ProceduralVegetation/Private/Nodes/PVScaleSettings.h
--- a/Engine/Plugins/Experimental/ProceduralVegetationEditor/Source/ProceduralVegetation/Private/Nodes/PVScaleSettings.h
+++ b/Engine/Plugins/Experimental/ProceduralVegetationEditor/Source/ProceduralVegetation/Private/Nodes/PVScaleSettings.h
@@ -27,7 +27,7 @@ protected:
 	virtual FPCGElementPtr CreateElement() const override;
 
 public:
-	UPROPERTY(EditAnywhere, Category = Settings, meta=(ClampMin = 0.01f, ClampMax = 100.0f, UIMin = 0.01f, UIMax = 10.0f, Tooltip="Controls overall size of generated elements.\n\nUniform size multiplier for branches, leaves, or sub-structures. Combine with randomness or ramps for natural variation."))
+	UPROPERTY(EditAnywhere, Category = Settings, meta=(PCG_Overridable, ClampMin = 0.01f, ClampMax = 100.0f, UIMin = 0.01f, UIMax = 10.0f, Tooltip="Controls overall size of generated elements.\n\nUniform size multiplier for branches, leaves, or sub-structures. Combine with randomness or ramps for natural variation."))
 	float Scale = 1;
 
 };
diff --git a/Engine/Plugins/Experimental/ProceduralVegetationEditor/Source/ProceduralVegetation/Private/Nodes/PVSlopeSettings.h b/Engine/Plugins/Experimental/ProceduralVegetationEditor/Source/ProceduralVegetation/Private/Nodes/PVSlopeSettings.h
--- a/Engine/Plugins/Experimental/ProceduralVegetationEditor/Source/ProceduralVegetation/Private/Nodes/PVSlopeSettings.h
+++ b/Engine/Plugins/Experimental/ProceduralVegetationEditor/Source/ProceduralVegetation/Private/Nodes/PVSlopeSettings.h
@@ -28,7 +28,7 @@ protected:
 	virtual FPCGElementPtr CreateElement() const override;
 
 public:
-	UPROPERTY(EditAnywhere, Category = "Slope", meta = (ShowOnlyInnerProperties))
+	UPROPERTY(EditAnywhere, Category = "Slope", meta = (ShowOnlyInnerProperties, PCG_Overridable))
 	FPVSlopeParams SlopeParams;
 };
 
