[PVE] QOL Improvements for Object Interaction Node + Scale fix after trim in Object Avoidance Node

#Synced-CL 50389151

diff --git a/Engine/Plugins/Experimental/ProceduralVegetationEditor/Source/ProceduralVegetation/Private/Implementations/PVObjectInteraction.cpp b/Engine/Plugins/Experimental/ProceduralVegetationEditor/Source/ProceduralVegetation/Private/Implementations/PVObjectInteraction.cpp
--- a/Engine/Plugins/Experimental/ProceduralVegetationEditor/Source/ProceduralVegetation/Private/Implementations/PVObjectInteraction.cpp
+++ b/Engine/Plugins/Experimental/ProceduralVegetationEditor/Source/ProceduralVegetation/Private/Implementations/PVObjectInteraction.cpp
@@ -223,8 +223,8 @@ void FPVObjectInteraction::ObjectInteraction(const FManagedArrayCollection& InCo
 
 					PVE_LOOP_DEBUG_STEP(break);
 
-					constexpr int32 Smoothness = 10;
-					constexpr float SmoothnessFactor = 1.5f;
+					constexpr int32 Smoothness = 15;
+					const float SmoothnessFactor = 1.1f + (1.0f - Collider.SmoothnessFactor);
 
 					const int32 StartIndex = FMath::Max(1, BranchPointIndex - Smoothness);
 					FVector3f PreviousSiblingPosition;
@@ -346,19 +346,23 @@ void FPVObjectInteraction::ObjectInteraction(const FManagedArrayCollection& InCo
 							}
 						}
 
-						const float OldMax = PointFacade.GetPointScale(BranchPoints[0]);
-						const float OldMin = PointFacade.GetPointScale(BranchPoints[BranchPointIndex - 1]);
-						const float NewMax = OldMax * (static_cast<float>(BranchPointIndex) / BranchPoints.Num());
-						const float NewMin = OldMin * (NewMax / OldMax);
-
+						const float Ratio = static_cast<float>(BranchPointIndex) / static_cast<float>(BranchPoints.Num());
+						
 						// Rescale the branches based on the cut point
 						for (int32 Index = 0; Index <= BranchPointIndex; ++Index)
 						{
-							const float UpdatedScale = FMath::GetMappedRangeValueClamped(
-								FVector2f(OldMin, OldMax),
-								FVector2f(NewMin, NewMax),
-								PointFacade.GetPointScale(BranchPoints[Index]));
-							PointFacade.SetPointScale(BranchPoints[Index], UpdatedScale);
+							PointFacade.SetPointScale(BranchPoints[Index], PointFacade.GetPointScale(BranchPoints[Index]) * Ratio);
+						}
+						
+						for (int32 ChildBranch : BranchFacade.GetChildren(BranchIndex))
+						{
+							if (const int32 ChildBranchIndex = BranchFacade.GetBranchNumbers().Find(ChildBranch); ChildBranchIndex != INDEX_NONE)
+							{
+								for (int32 ChildBranchPoint : BranchFacade.GetPoints(ChildBranchIndex))
+								{
+									PointFacade.SetPointScale(ChildBranchPoint, PointFacade.GetPointScale(ChildBranchPoint) * Ratio);
+								}
+							}
 						}
 
 						// Remove all branch children that are after the cut point
diff --git a/Engine/Plugins/Experimental/ProceduralVegetationEditor/Source/ProceduralVegetation/Private/Nodes/PVObjectInteractionSettings.cpp b/Engine/Plugins/Experimental/ProceduralVegetationEditor/Source/ProceduralVegetation/Private/Nodes/PVObjectInteractionSettings.cpp
--- a/Engine/Plugins/Experimental/ProceduralVegetationEditor/Source/ProceduralVegetation/Private/Nodes/PVObjectInteractionSettings.cpp
+++ b/Engine/Plugins/Experimental/ProceduralVegetationEditor/Source/ProceduralVegetation/Private/Nodes/PVObjectInteractionSettings.cpp
@@ -51,18 +51,19 @@ bool FPVObjectInteractionElement::ExecuteInternal(FPCGContext* InContext) const
 	{
 		if (const UPVGrowthData* InputData = Cast<UPVGrowthData>(Inputs[0].Data))
 		{
-			FManagedArrayCollection OutCollection;
+			FManagedArrayCollection OutCollection = InputData->GetCollection();
 			UPVData* OutManagedArrayCollectionData = FPCGContext::NewObject_AnyThread<UPVGrowthData>(InContext);
 
 #if WITH_EDITORONLY_DATA
 			PVE_LOOP_DEBUG_INIT(InputSettings, LoopDebugStepper, OutManagedArrayCollectionData);
 #endif
-
+			
 			for (const auto& Collider : InputSettings->ObjectInteractionSettings.Colliders)
 			{
 				PVE_OUTER_LOOP_DEBUG_CHECK(break);
-
-				FPVObjectInteraction::ObjectInteraction(InputData->GetCollection(), Collider, OutCollection);
+				
+				const FManagedArrayCollection InCollection = MoveTemp(OutCollection);
+				FPVObjectInteraction::ObjectInteraction(InCollection, Collider, OutCollection);
 			}
 
 			OutManagedArrayCollectionData->Initialize(MoveTemp(OutCollection));
diff --git a/Engine/Plugins/Experimental/ProceduralVegetationEditor/Source/ProceduralVegetation/Public/Nodes/PVObjectInteractionSettings.h b/Engine/Plugins/Experimental/ProceduralVegetationEditor/Source/ProceduralVegetation/Public/Nodes/PVObjectInteractionSettings.h
--- a/Engine/Plugins/Experimental/ProceduralVegetationEditor/Source/ProceduralVegetation/Public/Nodes/PVObjectInteractionSettings.h
+++ b/Engine/Plugins/Experimental/ProceduralVegetationEditor/Source/ProceduralVegetation/Public/Nodes/PVObjectInteractionSettings.h
@@ -25,7 +25,10 @@ struct FPVColliderParams
 	FTransform Transform;
 
 	UPROPERTY(EditAnywhere, Category = "Collider")
-	EPVCollisionType CollisionType;
+	EPVCollisionType CollisionType = EPVCollisionType::AVOID;
+	
+	UPROPERTY(EditAnywhere, Category = "Collider", meta=(EditCondition = "CollisionType == EPVCollisionType::AVOID", EditConditionHides, UIMin=0, UIMax=1, ClampMin=0, ClampMax=1))
+	float SmoothnessFactor = 1.0f;
 };
 
 USTRUCT()
diff --git a/Engine/Plugins/Experimental/ProceduralVegetationEditor/Source/ProceduralVegetationEditor/Private/Customizations/PVColliderParamsCustomization.cpp b/Engine/Plugins/Experimental/ProceduralVegetationEditor/Source/ProceduralVegetationEditor/Private/Customizations/PVColliderParamsCustomization.cpp
new file mode 100644
--- /dev/null
+++ b/Engine/Plugins/Experimental/ProceduralVegetationEditor/Source/ProceduralVegetationEditor/Private/Customizations/PVColliderParamsCustomization.cpp
@@ -0,0 +1,243 @@
+// Copyright Epic Games, Inc. All Rights Reserved.
+
+#include "PVColliderParamsCustomization.h"
+
+#include "ClassIconFinder.h"
+#include "DetailWidgetRow.h"
+#include "IDetailChildrenBuilder.h"
+#include "PropertyHandle.h"
+
+#include "ActorFactories/ActorFactoryBasicShape.h"
+
+#include "AssetRegistry/IAssetRegistry.h"
+
+#include "Nodes/PVObjectInteractionSettings.h"
+
+#include "Widgets/Images/SImage.h"
+#include "Widgets/Input/SButton.h"
+#include "Widgets/Layout/SBox.h"
+#include "Widgets/Layout/SSeparator.h"
+#include "Widgets/Text/STextBlock.h"
+
+class SPVDefaultCollidersButton : public SCompoundWidget
+{
+public:
+	SLATE_BEGIN_ARGS(SPVDefaultCollidersButton)
+			: _Text()
+			, _Image(FAppStyle::GetBrush("Default"))
+		{}
+
+		SLATE_ATTRIBUTE(FText, Text)
+		SLATE_ATTRIBUTE(const FSlateBrush*, Image)
+		SLATE_EVENT(FSimpleDelegate, OnClickAction)
+	SLATE_END_ARGS()
+
+	void Construct(const FArguments& InArgs)
+	{
+		OnClickAction = InArgs._OnClickAction;
+
+		ChildSlot
+		[
+			SNew(SBox)
+			.HAlign(HAlign_Center)
+			.VAlign(VAlign_Center)
+			.WidthOverride(22.0f)
+			.HeightOverride(22.0f)
+			.ToolTipText(InArgs._Text)
+			[
+				SNew(SButton)
+				.ButtonStyle(FAppStyle::Get(), "SimpleButton")
+				.OnClicked(this, &SPVDefaultCollidersButton::OnClick)
+				.ContentPadding(0.0f)
+				.IsFocusable(false)
+				[
+					SNew(SImage)
+					.Image(InArgs._Image)
+					.ColorAndOpacity(FSlateColor::UseForeground())
+				]
+			]
+		];
+	}
+
+private:
+	FReply OnClick()
+	{
+		OnClickAction.ExecuteIfBound();
+		return FReply::Handled();
+	}
+
+private:
+	FSimpleDelegate OnClickAction;
+};
+
+TSharedRef<SWidget> MakeDefaultColliderButton(const FName IconName, const TAttribute<FText>& ToolTipText, const FSimpleDelegate& OnClicked)
+{
+	return SNew(SPVDefaultCollidersButton)
+		.Text(ToolTipText)
+		.Image(FSlateIconFinder::FindCustomIconBrushForClass(nullptr, TEXT("ClassIcon"), IconName))
+		.OnClickAction(OnClicked);
+}
+
+TSharedRef<IPropertyTypeCustomization> FPVColliderParamsCustomization::MakeInstance()
+{
+	return MakeShareable(new FPVColliderParamsCustomization());
+}
+
+void FPVColliderParamsCustomization::CustomizeHeader(
+	TSharedRef<IPropertyHandle> InPropertyHandle,
+	FDetailWidgetRow& HeaderRow,
+	IPropertyTypeCustomizationUtils& InCustomizationUtils
+)
+{}
+
+void FPVColliderParamsCustomization::CustomizeChildren(
+	TSharedRef<IPropertyHandle> InPropertyHandle,
+	IDetailChildrenBuilder& ChildBuilder,
+	IPropertyTypeCustomizationUtils& InCustomizationUtils
+)
+{
+	const auto GetInternalName = [](const TSharedPtr<IPropertyHandle>& InHandle)
+		{
+			return InHandle->GetProperty()->GetName();
+		};
+
+	uint32 NumChild;
+	InPropertyHandle->GetNumChildren(NumChild);
+
+	for (uint32 ChildIndex = 0; ChildIndex < NumChild; ++ChildIndex)
+	{
+		const TSharedPtr<IPropertyHandle> ChildHandle = InPropertyHandle->GetChildHandle(ChildIndex);
+		IDetailPropertyRow& RowBuilder = ChildBuilder.AddProperty(ChildHandle.ToSharedRef());
+
+		if (GetInternalName(ChildHandle) == GET_MEMBER_NAME_STRING_VIEW_CHECKED(FPVColliderParams, Mesh))
+		{
+			CustomizeColliderSelectorWidget(ChildHandle, RowBuilder);
+		}
+	}
+}
+
+void FPVColliderParamsCustomization::CustomizeColliderSelectorWidget(
+	const TSharedPtr<IPropertyHandle>& InChildHandle,
+	IDetailPropertyRow& InRowBuilder
+)
+{
+	const TSharedRef<SWidget> NameWidget = InChildHandle->CreatePropertyNameWidget();
+	const TSharedRef<SWidget> ValueWidget =
+		SNew(SHorizontalBox)
+		+ SHorizontalBox::Slot()
+		.HAlign(HAlign_Left)
+		[
+			InChildHandle->CreatePropertyValueWidget()
+		]
+		+ SHorizontalBox::Slot()
+		.HAlign(HAlign_Right)
+		[
+			SNew(SHorizontalBox)
+			+ SHorizontalBox::Slot()
+			.HAlign(HAlign_Left)
+			[
+				SNew(SSeparator)
+				.Thickness(1.0f)
+				.Orientation(Orient_Vertical)
+			]
+			+ SHorizontalBox::Slot()
+			.HAlign(HAlign_Right)
+			.AutoWidth()
+			.Padding(2.0f, 0.0f)
+			[
+				SNew(SBox)
+				.HAlign(HAlign_Center)
+				.VAlign(VAlign_Center)
+				[
+					SNew(SVerticalBox)
+					+ SVerticalBox::Slot()
+					.VAlign(VAlign_Top)
+					.HAlign(HAlign_Center)
+					[
+						SNew(STextBlock)
+						.Justification(ETextJustify::Center)
+						.Text(FText::FromString("Basic Shapes"))
+					]
+					+ SVerticalBox::Slot()
+					.VAlign(VAlign_Bottom)
+					[
+						SNew(SHorizontalBox)
+						+ SHorizontalBox::Slot()
+						.AutoWidth()
+						.Padding(1.0f)
+						[
+							MakeDefaultColliderButton(
+								"ClassIcon.Cube",
+								FText::FromString("Cube"),
+								FSimpleDelegate::CreateLambda(
+									[=]()
+										{
+											InChildHandle->SetValue(IAssetRegistry::Get()->GetAssetByObjectPath(UActorFactoryBasicShape::BasicCube));
+										}
+								)
+							)
+						]
+						+ SHorizontalBox::Slot()
+						.AutoWidth()
+						.Padding(1.0f)
+						[
+							MakeDefaultColliderButton(
+								"ClassIcon.Sphere",
+								FText::FromString("Sphere"),
+								FSimpleDelegate::CreateLambda(
+									[=]()
+										{
+											InChildHandle->
+												SetValue(IAssetRegistry::Get()->GetAssetByObjectPath(UActorFactoryBasicShape::BasicSphere));
+										}
+								)
+							)
+						]
+						+ SHorizontalBox::Slot()
+						.AutoWidth()
+						.Padding(1.0f)
+						[
+							MakeDefaultColliderButton(
+								"ClassIcon.Cylinder",
+								FText::FromString("Cylinder"),
+								FSimpleDelegate::CreateLambda(
+									[=]()
+										{
+											InChildHandle->SetValue(
+												IAssetRegistry::Get()->GetAssetByObjectPath(UActorFactoryBasicShape::BasicCylinder));
+										}
+								)
+							)
+						]
+						+ SHorizontalBox::Slot()
+						.AutoWidth()
+						.Padding(1.0f)
+						[
+							MakeDefaultColliderButton(
+								"ClassIcon.Cone",
+								FText::FromString("Cone"),
+								FSimpleDelegate::CreateLambda(
+									[=]()
+										{
+											InChildHandle->SetValue(IAssetRegistry::Get()->GetAssetByObjectPath(UActorFactoryBasicShape::BasicCone));
+										}
+								)
+							)
+						]
+					]
+				]
+			]
+		];
+
+	InRowBuilder
+		.CustomWidget(true)
+		.NameContent()
+		[
+			NameWidget
+		]
+		.ValueContent()
+		.HAlign(HAlign_Fill)
+		[
+			ValueWidget
+		];
+}
diff --git a/Engine/Plugins/Experimental/ProceduralVegetationEditor/Source/ProceduralVegetationEditor/Private/Customizations/PVColliderParamsCustomization.h b/Engine/Plugins/Experimental/ProceduralVegetationEditor/Source/ProceduralVegetationEditor/Private/Customizations/PVColliderParamsCustomization.h
new file mode 100644
--- /dev/null
+++ b/Engine/Plugins/Experimental/ProceduralVegetationEditor/Source/ProceduralVegetationEditor/Private/Customizations/PVColliderParamsCustomization.h
@@ -0,0 +1,28 @@
+// Copyright Epic Games, Inc. All Rights Reserved.
+
+#pragma once
+
+#include "IDetailPropertyRow.h"
+#include "IPropertyTypeCustomization.h"
+
+class FPVColliderParamsCustomization : public IPropertyTypeCustomization
+{
+public:
+	static TSharedRef<IPropertyTypeCustomization> MakeInstance();
+
+	virtual void CustomizeHeader(
+		TSharedRef<IPropertyHandle> InPropertyHandle,
+		FDetailWidgetRow& HeaderRow,
+		IPropertyTypeCustomizationUtils& InCustomizationUtils
+	) override;
+
+	virtual void CustomizeChildren(
+		TSharedRef<IPropertyHandle> InPropertyHandle,
+		IDetailChildrenBuilder& ChildBuilder,
+		IPropertyTypeCustomizationUtils& InCustomizationUtils
+	)
+	override;
+
+private:
+	void CustomizeColliderSelectorWidget(const TSharedPtr<IPropertyHandle>& InChildHandle, IDetailPropertyRow& InRowBuilder);
+};
diff --git a/Engine/Plugins/Experimental/ProceduralVegetationEditor/Source/ProceduralVegetationEditor/Private/Customizations/PVOutputSettingsCustomizations.cpp b/Engine/Plugins/Experimental/ProceduralVegetationEditor/Source/ProceduralVegetationEditor/Private/Customizations/PVOutputSettingsCustomizations.cpp
--- a/Engine/Plugins/Experimental/ProceduralVegetationEditor/Source/ProceduralVegetationEditor/Private/Customizations/PVOutputSettingsCustomizations.cpp
+++ b/Engine/Plugins/Experimental/ProceduralVegetationEditor/Source/ProceduralVegetationEditor/Private/Customizations/PVOutputSettingsCustomizations.cpp
@@ -22,11 +22,6 @@ FReply OnOpenProjectSettings()
 	return FReply::Handled();
 }
 
-FString GetInternalName(const TSharedPtr<IPropertyHandle>& ChildHandle)
-{
-	return ChildHandle->GetProperty()->GetName();
-}
-
 TSharedRef<IPropertyTypeCustomization> FPVOutputSettingsCustomizations::MakeInstance()
 {
 	return MakeShareable(new FPVOutputSettingsCustomizations());
@@ -45,6 +40,11 @@ void FPVOutputSettingsCustomizations::CustomizeChildren(
 	IPropertyTypeCustomizationUtils& CustomizationUtils
 )
 {
+	const auto GetInternalName = [](const TSharedPtr<IPropertyHandle>& ChildHandle)
+		{
+			return ChildHandle->GetProperty()->GetName();
+		};
+	
 	uint32 NumChild;
 	PropertyHandle->GetNumChildren(NumChild);
 
diff --git a/Engine/Plugins/Experimental/ProceduralVegetationEditor/Source/ProceduralVegetationEditor/Private/ProceduralVegetationEditorModule.cpp b/Engine/Plugins/Experimental/ProceduralVegetationEditor/Source/ProceduralVegetationEditor/Private/ProceduralVegetationEditorModule.cpp
--- a/Engine/Plugins/Experimental/ProceduralVegetationEditor/Source/ProceduralVegetationEditor/Private/ProceduralVegetationEditorModule.cpp
+++ b/Engine/Plugins/Experimental/ProceduralVegetationEditor/Source/ProceduralVegetationEditor/Private/ProceduralVegetationEditorModule.cpp
@@ -40,10 +40,13 @@
 #include "Modules/ModuleManager.h"
 #include "MessageLogModule.h"
 
+#include "Customizations/PVColliderParamsCustomization.h"
 #include "Customizations/PVLoopDebugStepperCustomization.h"
 
 #include "Helpers/PVUtilities.h"
 
+#include "Nodes/PVObjectInteractionSettings.h"
+
 #include "Templates/SharedPointer.h"
 
 #include "Visualizations/PVDataVisualization.h"
@@ -68,6 +71,7 @@ void FProceduralVegetationEditorModule::StartupModule()
 
 	FPropertyEditorModule& PropertyEditor = FModuleManager::LoadModuleChecked<FPropertyEditorModule>("PropertyEditor");
 	PropertyEditor.RegisterCustomPropertyTypeLayout(FPVExportParams::StaticStruct()->GetFName(), FOnGetPropertyTypeCustomizationInstance::CreateStatic(&FPVOutputSettingsCustomizations::MakeInstance));
+	PropertyEditor.RegisterCustomPropertyTypeLayout(FPVColliderParams::StaticStruct()->GetFName(), FOnGetPropertyTypeCustomizationInstance::CreateStatic(&FPVColliderParamsCustomization::MakeInstance));
 	PropertyEditor.RegisterCustomPropertyTypeLayout(FPVFloatRamp::StaticStruct()->GetFName(), FOnGetPropertyTypeCustomizationInstance::CreateStatic(&FPVFloatRampCustomization::MakeInstance));
 	PropertyEditor.RegisterCustomPropertyTypeLayout(FLoopDebugStepper::StaticStruct()->GetFName(), FOnGetPropertyTypeCustomizationInstance::CreateStatic(&FPVLoopDebugStepperCustomization::MakeInstance));
 
diff --git a/Engine/Plugins/Experimental/ProceduralVegetationEditor/Source/ProceduralVegetationEditor/Private/SPVEditorViewport.cpp b/Engine/Plugins/Experimental/ProceduralVegetationEditor/Source/ProceduralVegetationEditor/Private/SPVEditorViewport.cpp
--- a/Engine/Plugins/Experimental/ProceduralVegetationEditor/Source/ProceduralVegetationEditor/Private/SPVEditorViewport.cpp
+++ b/Engine/Plugins/Experimental/ProceduralVegetationEditor/Source/ProceduralVegetationEditor/Private/SPVEditorViewport.cpp
@@ -412,6 +412,11 @@ void SPVEditorViewport::ResetVisualizationScene()
 		ScaleVisualizationComponent->MarkAsGarbage();
 		ScaleVisualizationComponent = nullptr;
 	}
+	
+	if (SkeletonVisualizerComponent)
+	{
+		SkeletonVisualizerComponent = nullptr;
+	}
 }
 
 TSharedRef<SWidget> SPVEditorViewport::CreateMannequinOffsetWidget() const
