ProceduralVegetation: Fix LOCTEXT key conflict

LogLocTextHelper: Warning: Text conflict from LOCTEXT macro for namespace 'PVEditorViewportToolbarSections' and key 'DebugModeSelectionSubSectionLabel'. First entry is Plugins/Experimental/ProceduralVegetationEditor/Source/ProceduralVegetationEditor/Private/SPVEditorViewport.cpp(823):'Toggle Debug Visualizations'. Conflicting entry is Plugins/Experimental/ProceduralVegetationEditor/Source/ProceduralVegetationEditor/Private/SPVEditorViewport.cpp(822):'Debug Visualizations'.

#rnx

#Synced-CL 50390180

diff --git a/Engine/Plugins/Experimental/ProceduralVegetationEditor/Source/ProceduralVegetationEditor/Private/SPVEditorViewport.cpp b/Engine/Plugins/Experimental/ProceduralVegetationEditor/Source/ProceduralVegetationEditor/Private/SPVEditorViewport.cpp
--- a/Engine/Plugins/Experimental/ProceduralVegetationEditor/Source/ProceduralVegetationEditor/Private/SPVEditorViewport.cpp
+++ b/Engine/Plugins/Experimental/ProceduralVegetationEditor/Source/ProceduralVegetationEditor/Private/SPVEditorViewport.cpp
@@ -820,7 +820,7 @@ FToolMenuEntry SPVEditorViewport::CreateVisualizationModeToolbarMenu()
 							RenderModeSelectionSection.AddSubMenu(
 								TEXT("DebugModeSection"),
 								LOCTEXT("DebugModeSelectionSubSectionLabel", "Debug Visualizations"),
-								LOCTEXT("DebugModeSelectionSubSectionLabel", "Toggle Debug Visualizations"),
+								LOCTEXT("DebugModeSelectionSubSectionLabelToolTip", "Toggle Debug Visualizations"),
 								FNewToolMenuDelegate::CreateSPLambda(this,
 									[this](UToolMenu* SubMenu)
 										{
