Make PVE 2D/3D importer output work with resume growth on the growth node

This CL aims to finalize and solidify the importer growth data generation. The main goal is to have the growth data generated from the importers to be compatible with the growth and grafting nodes.

This CL can be split into 2 major parts:
1. Compute estimate for missing attributes required for resume growth (PVImportHelpers.cpp)
- Implement all attribute estimation functions (ComputeBudNumbers, ComputeNjordPixelIndex, ComputeBudDevelopment, ComputeBudDirections, ComputeBudHormoneLevels, ComputeGradients, ComputeBudStatus, ComputeBudLateralMeristem, ComputeBudLightDetected)
- Add unit tests for all attributes (PVImportTests.cpp)

2. Add new strongly typed attribute system to make PVE attribute manipulation more reusable and robust (PVAttributes.h). The main motivation for incluiding this new PV::TAttributeDefinition type in this CL is due to the attribute computation in part 1 having strong dependencies between the different attributes (e.g. they need to be computed in a particular order). Using the facades or the collection lead to issues where it was not clear which attributes a computation step required, and hence which order the operations needed to be executed.
- Add new PV::TAttributeDefinition type which allows for strongly typed attributes (PVAttributes.h)
- Add definitions for all current PVE growth attributes
- Restructure PVAttributesNames.h to allow using attribute names as template arguments
- Re-name PVFacadeHelper.h to PVAttributeHelper.h and replace use of facades with explicit attribute dependencies
- Update the slope node to use direct attribute access instead of the facades (required as the slope node uses the same helper library as the importer. It is also a good test-case for using direct attribute access instead of facades for attribute manipulation)

#Synced-CL 50481844

diff --git a/Engine/Plugins/Experimental/ProceduralVegetationEditor/Source/ProceduralVegetation/Private/Helpers/PVFacadeHelper.cpp b/Engine/Plugins/Experimental/ProceduralVegetationEditor/Source/ProceduralVegetation/Private/Helpers/PVFacadeHelper.cpp
deleted file mode 100644
diff --git a/Engine/Plugins/Experimental/ProceduralVegetationEditor/Source/ProceduralVegetation/Private/Implementations/PVSlope.cpp b/Engine/Plugins/Experimental/ProceduralVegetationEditor/Source/ProceduralVegetation/Private/Implementations/PVSlope.cpp
--- a/Engine/Plugins/Experimental/ProceduralVegetationEditor/Source/ProceduralVegetation/Private/Implementations/PVSlope.cpp
+++ b/Engine/Plugins/Experimental/ProceduralVegetationEditor/Source/ProceduralVegetation/Private/Implementations/PVSlope.cpp
@@ -1,12 +1,8 @@
 // Copyright Epic Games, Inc. All Rights Reserved.
 
 #include "PVSlope.h"
-#include "Helpers/PVFacadeHelper.h"
-#include "Facades/PVBranchFacade.h"
-#include "Facades/PVBudVectorsFacade.h"
-#include "Facades/PVFoliageFacade.h"
-#include "Facades/PVPlantFacade.h"
-#include "Facades/PVPointFacade.h"
+#include "Helpers/PVAttributesHelper.h"
+#include "Facades/PVAttributes.h"
 
 namespace PVSlope
 {
@@ -21,44 +17,61 @@ namespace PVSlope
 	}
 
 	float CalcBranchLength(
-		const PV::Facades::FBranchFacade& BranchFacade, 
-		const PV::Facades::FPointFacade& PointFacade, 
+		PV::FBranchChildrenAttributeConstView BranchChildrenAttribute,
+		PV::FBranchParentNumberAttributeConstView BranchParentNumberAttribute,
+		PV::FBranchNumberAttributeConstView BranchNumberAttribute,
+		PV::FBranchPointsAttributeConstView BranchPointsAttribute,
+		PV::FPointPositionAttributeConstView PointPositionAttribute,
 		int32 BranchIndex)
 	{
-		const TArray<int32>& CurrentBranchPoints = BranchFacade.GetPoints(BranchIndex);
+		const TArray<int32>& CurrentBranchPoints = BranchPointsAttribute[BranchIndex];
 		if (CurrentBranchPoints.Num() == 0)
 		{
 			return 0;
 		}
 
-		const TManagedArray<FVector3f>& PointPositions = PointFacade.GetPositions();
-		const FVector3f& RootPosition = PointPositions[CurrentBranchPoints[0]];
+		const FVector3f& RootPosition = PointPositionAttribute[CurrentBranchPoints[0]];
 
 		float BranchLengthSquared = 0;
-		PV::FacadeHelper::WalkBranchPoints(
-			BranchFacade,
-			PointFacade,
+		PV::AttributesHelper::RecursiveWalkBranchPoints_FusedPoints(
+			BranchChildrenAttribute,
+			BranchParentNumberAttribute,
+			BranchNumberAttribute,
+			BranchPointsAttribute,
 			BranchIndex,
-			[&](int32 BranchIndex, int32 BranchPointIndex, int32 PointIndex, const TArray<int32>& ChildPointIndices)->PV::FacadeHelper::EForEachResult
+			[&](int32 BranchIndex, int32 BranchPointIndex, int32 PointIndex, const TArray<int32>& ChildPointIndices)->PV::AttributesHelper::EForEachResult
 			{
-				const FVector3f& PointPosition = PointPositions[PointIndex];
+				const FVector3f& PointPosition = PointPositionAttribute[PointIndex];
 				const float PointDistSquared = (PointPosition - RootPosition).SizeSquared();
 				if (PointDistSquared > BranchLengthSquared)
 				{
 					BranchLengthSquared = PointDistSquared;
 				}
-				return PV::FacadeHelper::EForEachResult::Continue;
+				return PV::AttributesHelper::EForEachResult::Continue;
 			}
 		);
 
 		return BranchLengthSquared > 0 ? FMath::Sqrt(BranchLengthSquared) : 0.f;
 	}
 
+	struct FApplySlopeAttributes
+	{
+		PV::FPointPositionAttributeView PointPosition;
+		PV::FBudDirectionAttributeView BudDirection;
+		PV::FPointLengthFromRootAttributeView PointLengthFromRoot;
+		PV::FFoliagePivotPointAttributeView FoliagePivotPoint;
+		PV::FFoliageUpVectorAttributeView FoliageUpVector;
+		PV::FFoliageNormalVectorAttributeView FoliageNormalVector;
+		PV::FFoliageLengthFromRootAttributeConstView FoliageLengthFromRoot;
+		PV::FBranchChildrenAttributeConstView BranchChildren;
+		PV::FBranchParentNumberAttributeConstView BranchParentNumber;
+		PV::FBranchNumberAttributeConstView BranchNumber;
+		PV::FBranchPointsAttributeConstView BranchPoints;
+		PV::FBranchFoliageIDsAttributeConstView BranchFoliageIDs;
+	};
+
 	void ApplySlope_Recursive(
-		PV::Facades::FBranchFacade& BranchFacade,
-		PV::Facades::FPointFacade& PointFacade,
-		PV::Facades::FBudVectorsFacade& BudVectorFacade,
-		PV::Facades::FFoliageFacade& FoliageFacade,
+		const FApplySlopeAttributes& Attributes,
 		const int32 BranchIndex,
 		const FPVSlopeParams& SlopeParams,
 		float TreeHeight,
@@ -67,47 +80,47 @@ namespace PVSlope
 		FVector3f ParentIdentityPosition
 	)
 	{
-		TManagedArray<FVector3f>& PointPositions = PointFacade.ModifyPositions();
-		TManagedArray<TArray<FVector3f>>& PointBudDirections = BudVectorFacade.ModifyBudDirections();
-
-		const TArray<int32>& CurrentBranchPoints = BranchFacade.GetPoints(BranchIndex);
+		const TArray<int32>& CurrentBranchPoints = Attributes.BranchPoints[BranchIndex];
 
 		for (int32 Index = 0; Index < CurrentBranchPoints.Num(); ++Index)
 		{
 			const int32 PointIndex = CurrentBranchPoints[Index];
 			const int32 PrevPointIndex = Index > 0 ? CurrentBranchPoints[Index - 1] : INDEX_NONE;
-			FVector3f& Position = PointPositions[PointIndex];
+			FVector3f& Position = Attributes.PointPosition[PointIndex];
 
 			const float TreeAngleAlpha = 1.0 - (TreeHeight > 0 ? ((Position - TrunkPosition).Size() / TreeHeight) : 0);
 			const float AngleStrength = FMath::Clamp(FMath::Pow(TreeAngleAlpha, SlopeParams.BendStrength), 0.0, 1.0);
 
 			const FQuat4f SlopeRotation = CalculateBranchRotation(SlopeParams.SlopeDirection, SlopeParams.SlopeAngle * AngleStrength);
 
-			TArray<FVector3f>& CurrentBudDirections = PointBudDirections[PointIndex];
-			for (int32 i = 0; i < CurrentBudDirections.Num(); ++i)
-			{
-				if (i == PV::Facades::BudDirectionsLightOptimalIndex || i == PV::Facades::BudDirectionsLightSubOptimal)
-				{
-					continue;
-				}
+			PV::FBudDirectionView BudDirections = Attributes.BudDirection[PointIndex];
+			BudDirections.Apical = SlopeRotation.RotateVector(BudDirections.Apical);
+			BudDirections.Axillary = SlopeRotation.RotateVector(BudDirections.Axillary);
+			BudDirections.GuideCurve = SlopeRotation.RotateVector(BudDirections.GuideCurve);
+			BudDirections.UpVector = SlopeRotation.RotateVector(BudDirections.UpVector);
+
+			const TArray<int32> BranchSegmentFoliageEntryIds = PV::AttributesHelper::GetBranchPointFoliageEntryIds(
+				Attributes.BranchPoints, 
+				Attributes.PointLengthFromRoot,
+				Attributes.BranchFoliageIDs, 
+				Attributes.FoliageLengthFromRoot,
+				BranchIndex, 
+				PointIndex
+			);
 
-				CurrentBudDirections[i] = SlopeRotation.RotateVector(CurrentBudDirections[i]);
-			}
-
-			const TArray<int32> BranchSegmentFoliageEntryIds = PV::FacadeHelper::GetBranchPointFoliageEntryIds(FoliageFacade, BranchFacade, PointFacade, BranchIndex, PointIndex);
 			for (int32 FoliageEntryId : BranchSegmentFoliageEntryIds)
 			{
-				const FVector3f FoliageEntryPivotPoint = FoliageFacade.GetPivotPoint(FoliageEntryId);
+				const FVector3f& FoliageEntryPivotPoint = Attributes.FoliagePivotPoint[FoliageEntryId];
 				const FVector3f NewFoliageEntryPivotPoint = SlopeRotation.RotateVector(FoliageEntryPivotPoint - ParentIdentityPosition) + ParentPosition;
-				FoliageFacade.SetPivotPoint(FoliageEntryId, NewFoliageEntryPivotPoint);
+				Attributes.FoliagePivotPoint[FoliageEntryId] = NewFoliageEntryPivotPoint;
 
-				const FVector3f FoliageEntryUpVector = FoliageFacade.GetUpVector(FoliageEntryId);
+				const FVector3f FoliageEntryUpVector = Attributes.FoliageUpVector[FoliageEntryId];
 				const FVector3f NewFoliageEntryUpVector = SlopeRotation.RotateVector(FoliageEntryUpVector);
-				FoliageFacade.SetUpVector(FoliageEntryId, NewFoliageEntryUpVector);
+				Attributes.FoliageUpVector[FoliageEntryId] = NewFoliageEntryUpVector;
 
-				const FVector3f FoliageEntryNormalVector = FoliageFacade.GetNormalVector(FoliageEntryId);
+				const FVector3f FoliageEntryNormalVector = Attributes.FoliageNormalVector[FoliageEntryId];
 				const FVector3f NewFoliageEntryNormalVector = SlopeRotation.RotateVector(FoliageEntryNormalVector);
-				FoliageFacade.SetNormalVector(FoliageEntryId, NewFoliageEntryNormalVector);
+				Attributes.FoliageNormalVector[FoliageEntryId] = NewFoliageEntryNormalVector;
 			}
 
 			const FVector3f NewPosition = SlopeRotation.RotateVector(Position - ParentIdentityPosition) + ParentPosition;
@@ -116,14 +129,20 @@ namespace PVSlope
 			Position = NewPosition;
 			ParentPosition = Position;
 
-			const TArray<int32> BranchSegmentChildren = PV::FacadeHelper::GetBranchPointChildPoints(BranchFacade, PointFacade, BranchIndex, PointIndex);
+			const TArray<int32> BranchSegmentChildren = PV::AttributesHelper::GetBranchPointChildBranchIndices_SplitPoints(
+				Attributes.BranchChildren,
+				Attributes.BranchParentNumber,
+				Attributes.BranchNumber,
+				Attributes.BranchPoints,
+				Attributes.PointLengthFromRoot,
+				BranchIndex, 
+				PointIndex
+			);
+
 			for (const int32 ChildIndex : BranchSegmentChildren)
 			{
 				ApplySlope_Recursive(
-					BranchFacade,
-					PointFacade,
-					BudVectorFacade,
-					FoliageFacade,
+					Attributes,
 					ChildIndex,
 					SlopeParams,
 					TreeHeight,
@@ -138,34 +157,72 @@ namespace PVSlope
 
 void FPVSlope::ApplySlope(const FPVSlopeParams& InSlopeParams, FManagedArrayCollection& OutCollection)
 {
-	PV::Facades::FBranchFacade BranchFacade = PV::Facades::FBranchFacade(OutCollection);
-	PV::Facades::FPointFacade PointFacade = PV::Facades::FPointFacade(OutCollection);
-	PV::Facades::FBudVectorsFacade BudVectorFacade = PV::Facades::FBudVectorsFacade(OutCollection);
-	PV::Facades::FFoliageFacade FoliageFacade = PV::Facades::FFoliageFacade(OutCollection);
-	const PV::Facades::FPlantFacade PlantFacade(OutCollection);
-
-	if (!PointFacade.IsValid() || !BranchFacade.IsValid())
+	const PV::FBranchChildrenAttributeConstView BranchChildrenAttribute = PV::FBranchChildrenAttribute::FindAttribute(OutCollection);
+	const PV::FBranchParentNumberAttributeConstView BranchParentNumberAttribute = PV::FBranchParentNumberAttribute::FindAttribute(OutCollection);
+	const PV::FBranchNumberAttributeConstView BranchNumberAttribute = PV::FBranchNumberAttribute::FindAttribute(OutCollection);
+	const PV::FBranchPointsAttributeConstView BranchPointsAttribute = PV::FBranchPointsAttribute::FindAttribute(OutCollection);
+	const PV::FBranchFoliageIDsAttributeConstView BranchFoliageIDsAttribute = PV::FBranchFoliageIDsAttribute::FindAttribute(OutCollection);
+	const PV::FFoliageLengthFromRootAttributeConstView FoliageLengthFromRootAttribute = PV::FFoliageLengthFromRootAttribute::FindAttribute(OutCollection);
+
+	PV::FPointPositionAttributeView PointPositionAttribute = PV::FPointPositionAttribute::FindAttribute(OutCollection);
+	PV::FBudDirectionAttributeView BudDirectionAttribute = PV::FBudDirectionAttribute::FindAttribute(OutCollection);
+	PV::FPointLengthFromRootAttributeView PointLengthFromRootAttribute = PV::FPointLengthFromRootAttribute::FindAttribute(OutCollection);
+	PV::FFoliagePivotPointAttributeView FoliagePivotPointAttribute = PV::FFoliagePivotPointAttribute::FindAttribute(OutCollection);
+	PV::FFoliageUpVectorAttributeView FoliageUpVectorAttribute = PV::FFoliageUpVectorAttribute::FindAttribute(OutCollection);
+	PV::FFoliageNormalVectorAttributeView FoliageNormalVectorAttribute = PV::FFoliageNormalVectorAttribute::FindAttribute(OutCollection);
+
+	if (!BranchChildrenAttribute
+	 || !BranchParentNumberAttribute
+	 || !BranchNumberAttribute
+	 || !BranchPointsAttribute
+	 || !BranchFoliageIDsAttribute
+	 || !PointPositionAttribute
+	 || !BudDirectionAttribute
+	 || !PointLengthFromRootAttribute
+	 || !FoliagePivotPointAttribute
+	 || !FoliageUpVectorAttribute
+	 || !FoliageNormalVectorAttribute
+	 || !FoliageLengthFromRootAttribute)
 	{
 		return;
 	}
 
-	for (const int32 TrunkIndex : PlantFacade.GetTrunkIndices())
+	const PVSlope::FApplySlopeAttributes Attributes = {
+		PointPositionAttribute,
+		BudDirectionAttribute,
+		PointLengthFromRootAttribute,
+		FoliagePivotPointAttribute,
+		FoliageUpVectorAttribute,
+		FoliageNormalVectorAttribute,
+		FoliageLengthFromRootAttribute,
+		BranchChildrenAttribute,
+		BranchParentNumberAttribute,
+		BranchNumberAttribute,
+		BranchPointsAttribute,
+		BranchFoliageIDsAttribute
+	};
+
+	for (const int32 TrunkIndex : PV::AttributesHelper::GetTrunkIndices(BranchParentNumberAttribute))
 	{
-		const TArray<int32>& CurrentBranchPoints = BranchFacade.GetPoints(TrunkIndex);
+		const TArray<int32>& CurrentBranchPoints = BranchPointsAttribute[TrunkIndex];
 		if (CurrentBranchPoints.Num() == 0)
 		{
 			continue;
 		}
 
-		const float BranchLength = PVSlope::CalcBranchLength(BranchFacade, PointFacade, TrunkIndex);
-		const FVector3f TrunkPosition = PointFacade.GetPosition(CurrentBranchPoints[0]);
+		const float BranchLength = PVSlope::CalcBranchLength(
+			BranchChildrenAttribute,
+			BranchParentNumberAttribute,
+			BranchNumberAttribute,
+			BranchPointsAttribute,
+			PointPositionAttribute,
+			TrunkIndex
+		);
+		const FVector3f TrunkPosition = PointPositionAttribute[CurrentBranchPoints[0]];
 		const FVector3f PivotPoint = InSlopeParams.TrunkPivotPoint == EPVSlopeTrunkPivotPoint::Origin ? FVector3f::ZeroVector : TrunkPosition;
 
 		PVSlope::ApplySlope_Recursive(
-			BranchFacade,
-			PointFacade,
-			BudVectorFacade,
-			FoliageFacade,
+			Attributes,
 			TrunkIndex,
 			InSlopeParams,
 			BranchLength,
diff --git a/Engine/Plugins/Experimental/ProceduralVegetationEditor/Source/ProceduralVegetation/Private/Tests/PVMockTree.cpp b/Engine/Plugins/Experimental/ProceduralVegetationEditor/Source/ProceduralVegetation/Private/Tests/PVMockTree.cpp
--- a/Engine/Plugins/Experimental/ProceduralVegetationEditor/Source/ProceduralVegetation/Private/Tests/PVMockTree.cpp
+++ b/Engine/Plugins/Experimental/ProceduralVegetationEditor/Source/ProceduralVegetation/Private/Tests/PVMockTree.cpp
@@ -71,16 +71,16 @@ PVMockTreeCollection::CompareResult PVMockTreeCollection::Compare(const FManaged
 	return CompareResult::Success;
 }
 
-TSharedRef<FManagedArrayCollection> PVMockTreeCollection::CreateCollection()
+TSharedRef<FManagedArrayCollection> PVMockTreeCollection::CreateCollection(bool bSplitPoints)
 {
 #if WITH_EDITOR
 	if (PV::GPVEditorInterface)
 	{
-		return PV::GPVEditorInterface->CreateMockTreeGrowthDataCollection();
+		return PV::GPVEditorInterface->CreateMockTreeGrowthDataCollection(bSplitPoints);
 	}
 #endif // WITH_EDITOR
 
 	return MakeShared<FManagedArrayCollection>();
 }
 
-#endif // WITH_DEV_AUTOMATION_TESTS
\ No newline at end of file
+#endif // WITH_DEV_AUTOMATION_TESTS
diff --git a/Engine/Plugins/Experimental/ProceduralVegetationEditor/Source/ProceduralVegetation/Public/Facades/PVAttributes.h b/Engine/Plugins/Experimental/ProceduralVegetationEditor/Source/ProceduralVegetation/Public/Facades/PVAttributes.h
new file mode 100644
--- /dev/null
+++ b/Engine/Plugins/Experimental/ProceduralVegetationEditor/Source/ProceduralVegetation/Public/Facades/PVAttributes.h
@@ -0,0 +1,545 @@
+// Copyright Epic Games, Inc. All Rights Reserved.
+#pragma once
+
+#include "CoreMinimal.h"
+
+#include "GeometryCollection/ManagedArrayCollection.h"
+
+#include "Facades/PVAttributesNames.h"
+
+/*
+This file defines the TAttributeDefinition type which allows for compile-time type checking of attributes.
+
+- Defining a new attribute:
+	To define a new attribute, use the PV_DEFINE_ATTRIBUTE as follows:
+
+	inline constexpr TCHAR MyPropertyNameStr[] = TEXT("MyPropertyName");
+	inline constexpr TCHAR MyPropertyGroupStr[] = TEXT("MyPropertyGroup");
+    PV_DEFINE_ATTRIBUTE(MyProperty, FVector3f, MyPropertyNameStr, MyPropertyGroupStr);
+	PV_DEFINE_STRUCTURED_ATTRIBUTE can be used when custom element accessor views are needed.
+
+- Usage in code: 
+	PV_DEFINE_ATTRIBUTE defines a new attribute type, the example above defines the type FMyPropertyAttribute.
+	This type can be used with the FManagedArrayCollection to get, find, and add attributes using `FMyPropertyAttribute::FindAttribute(Collection)`, 
+	`FMyPropertyAttribute::GetAttribute(Collection)`, or `FMyPropertyAttribute::AddAttribute(Collection)`.
+
+	PV_DEFINE_ATTRIBUTE also defines a view and a const view for the attribute which can be used to access elements in the attribute.
+	Usage example:
+	void MultiplyMyProperty(FMyPropertyAttributeView MyPropertyAttribute, float Multiplier)
+	{
+		for (FVector3f& Vec : MyPropertyAttribute)
+		{
+			Vec *= Multiplier;
+		}
+	}
+	
+	// At call site:
+	MultiplyMyProperty(FMyPropertyAttribute::GetAttribute(Collection), 10.f);
+*/
+
+namespace PV
+{
+template<
+	typename InElementType,
+	typename InElementViewType,
+	typename InElementConstViewType,
+	const TCHAR* InAttributeNameStr, 
+	const TCHAR* InGroupNameStr
+>
+class TAttributeView
+{
+public:
+	using ElementType = std::remove_const_t<InElementType>;
+	
+	using ManagedArrayType =
+		std::conditional_t<std::is_const_v<InElementType>,
+		const TManagedArray<ElementType>,
+		TManagedArray<ElementType>>;
+	
+	using ElementViewType = 
+		std::conditional_t<std::is_const_v<InElementType>,
+		InElementConstViewType,
+		InElementViewType>;
+
+	using NonConstAttributeViewType = TAttributeView<ElementType, InElementViewType, InElementConstViewType, InAttributeNameStr, InGroupNameStr>;
+
+	using SizeType = TArray<ElementType>::SizeType;
+
+	inline static const FName AttributeName = FName(InAttributeNameStr);
+	inline static const FName GroupName = FName(InGroupNameStr);
+
+	ManagedArrayType* ManagedArray;
+
+	TAttributeView(ManagedArrayType* InManagedArray)
+		: ManagedArray(InManagedArray)
+	{}
+
+	TAttributeView(const NonConstAttributeViewType& Other)
+		: ManagedArray(Other.ManagedArray)
+	{}
+
+	bool IsValid() const
+	{
+		return ManagedArray != nullptr;
+	}
+
+	ElementViewType operator[](SizeType Index) const
+	{
+		check(IsValid() && ManagedArray->IsValidIndex(Index));
+		return ElementViewType((*ManagedArray)[Index]);
+	}
+
+	operator bool() const 
+	{
+		return IsValid();
+	}
+
+	bool operator!() const
+	{
+		return !IsValid();
+	}
+
+	/**
+	* Find first index of the element
+	*/
+	int32 Find(const ElementType& Item) const
+	{
+		return ManagedArray ? ManagedArray->Find(Item) : INDEX_NONE;
+	}
+
+	SizeType Num() const
+	{
+		return ManagedArray ? ManagedArray->Num() : 0;
+	}
+
+	using RangedForIteratorType      = ManagedArrayType::RangedForIteratorType;
+	using RangedForConstIteratorType = ManagedArrayType::RangedForConstIteratorType;
+
+	using AttributeRangedForIteratorType =
+		std::conditional_t<std::is_const_v<InElementType>,
+		RangedForConstIteratorType,
+		RangedForIteratorType>;
+
+	/**
+	* DO NOT USE DIRECTLY
+	* STL-like iterators to enable range-based for loop support.
+	*/
+	FORCEINLINE AttributeRangedForIteratorType begin() const { check(ManagedArray); return ManagedArray->begin(); }
+	FORCEINLINE AttributeRangedForIteratorType end() const { check(ManagedArray); return ManagedArray->end(); }
+};
+
+template<
+	typename InAttributeType, 
+	typename InElementViewType,
+	typename InElementConstViewType,
+	const TCHAR* InAttributeNameStr, 
+	const TCHAR* InGroupNameStr
+>
+class TAttributeDefinition
+{
+	static_assert(!std::is_const_v<InAttributeType>, "InAttributeType cannot be const");
+
+public:
+	using AttributeType = InAttributeType;
+	using AttributeView = TAttributeView<AttributeType, InElementViewType, InElementConstViewType, InAttributeNameStr, InGroupNameStr>;
+	using AttributeConstView = TAttributeView<const AttributeType, InElementViewType, InElementConstViewType, InAttributeNameStr, InGroupNameStr>;
+
+	inline static const FName AttributeName = FName(InAttributeNameStr);
+	inline static const FName GroupName = FName(InGroupNameStr);
+
+	/**
+	* Tries to find the attribute in the supplied collection. Returns a null view if attribute is not in collection.
+	* @return A const view of the found attribute.
+	*/
+	static AttributeConstView FindAttribute(const FManagedArrayCollection& InCollection)
+	{
+		const TManagedArray<AttributeType>* Attribute = InCollection.FindAttributeTyped<AttributeType>(AttributeName, GroupName);
+		if (ensure(Attribute != nullptr))
+		{
+			return AttributeConstView(Attribute);
+		}
+
+		return AttributeConstView(nullptr);
+	}
+	
+	/**
+	* Tries to find the attribute in the supplied collection. Returns a null view if attribute is not in collection.
+	* @return A mutable view of the found attribute.
+	*/
+	static AttributeView FindAttribute(FManagedArrayCollection& InCollection)
+	{
+		TManagedArray<AttributeType>* Attribute = InCollection.FindAttributeTyped<AttributeType>(AttributeName, GroupName);
+		if (ensure(Attribute != nullptr))
+		{
+			return AttributeView(Attribute);
+		}
+
+		return AttributeView(nullptr);
+	}
+
+	/**
+	* Checked version of FindAttribute. Will assert if attribute is missing from collection.
+	* @return A const view of the found attribute.
+	*/
+	static AttributeConstView GetAttribute(const FManagedArrayCollection& InCollection)
+	{
+		AttributeConstView AttributeView = FindAttribute(InCollection);
+		check(AttributeView.IsValid());
+		return AttributeView;
+	}
+
+	/**
+	* Checked version of FindAttribute. Will assert if attribute is missing from collection.
+	* @return A mutable view of the found attribute.
+	*/ 
+	static AttributeView GetAttribute(FManagedArrayCollection& InCollection)
+	{
+		AttributeView AttributeView = FindAttribute(InCollection);
+		check(AttributeView.IsValid());
+		return AttributeView;
+	}
+
+	/**
+	* Add the attribute to the collection and returns a view to the added collection. Does nothing if the attribute is already present in the collection.
+	* @return A mutable view of the found attribute.
+	*/
+	static AttributeView AddAttribute(FManagedArrayCollection& InCollection)
+	{
+		return AttributeView(&InCollection.AddAttribute<AttributeType>(AttributeName, GroupName));
+	}
+
+	/**
+	* Add the attribute to the collection and returns a view to the added collection. Does nothing if the attribute is already present in the collection.
+	* @param DefaultValue - Initialize all elements in the added attribute with this value (only if the attribute is not already present)
+	* @return A mutable view of the found attribute.
+	*/
+	static AttributeView AddAttribute(FManagedArrayCollection& InCollection, const AttributeType& DefaultValue)
+	{
+		AttributeView AttributeView = FindAttribute(InCollection);
+		if (AttributeView.IsValid())
+		{
+			return AttributeView;
+		}
+
+		AttributeView = AddAttribute(InCollection);
+		for (int32 i = 0; i < AttributeView.Num(); ++i)
+		{
+			AttributeView[i] = DefaultValue;
+		}
+		return AttributeView;
+	}
+
+	/**
+	* Check for the existence of the attribute in the collection.
+	*/
+	static bool HasAttribute(const FManagedArrayCollection& InCollection)
+	{
+		return InCollection.HasAttribute(AttributeName, GroupName);
+	}
+};
+
+// Helper class to allow for initialization, and selecton of, a fixed size array based of off default input array.
+template<typename InElementType>
+struct TFixedSizeArrayHelper
+{
+	using ElementType = std::remove_const_t<InElementType>;
+	using MutableArrayType = TArray<ElementType>;
+	using ConstArrayType = const TArray<ElementType>;
+	using ArrayType = std::conditional_t<std::is_const_v<InElementType>, ConstArrayType, MutableArrayType>;
+
+	TFixedSizeArrayHelper() = delete;
+
+	static MutableArrayType& GetFixedArray(MutableArrayType& InMutableArray, const ConstArrayType& InDefaultData)
+	{
+		if (InMutableArray.Num() < InDefaultData.Num())
+		{
+			const int32 FirstIndex = InMutableArray.Num();
+			InMutableArray.SetNum(InDefaultData.Num());
+			for (int32 i = FirstIndex; i < InDefaultData.Num(); ++i)
+			{
+				InMutableArray[i] = InDefaultData[i];
+			}
+		}
+
+		return InMutableArray;
+	}
+
+	static ConstArrayType& GetFixedArray(ConstArrayType& InConstArray, const ConstArrayType& InDefaultData)
+	{
+		if (InConstArray.Num() < InDefaultData.Num())
+		{
+			return InDefaultData;
+		}
+		return InConstArray;
+	}
+};
+
+#define PV_DEFINE_STRUCTURED_ATTRIBUTE(Name, Type, ElementViewType, ElementConstViewType, AttributeNameStr, AttributeGroupStr)          \
+	using F##Name##Attribute = TAttributeDefinition<Type, ElementViewType, ElementConstViewType, AttributeNameStr, AttributeGroupStr>;  \
+	using F##Name##AttributeView = F##Name##Attribute::AttributeView;                                                                   \
+	using F##Name##AttributeConstView = F##Name##Attribute::AttributeConstView;
+
+#define PV_DEFINE_ATTRIBUTE(Name, Type, AttributeNameStr, AttributeGroupStr) \
+	PV_DEFINE_STRUCTURED_ATTRIBUTE(Name, Type, Type&, const Type&, AttributeNameStr, AttributeGroupStr)
+
+// Branch Attributes
+PV_DEFINE_ATTRIBUTE(BranchParents, TArray<int32>, AttributeNames::BranchParentsStr, GroupNames::BranchGroupStr);
+PV_DEFINE_ATTRIBUTE(BranchChildren, TArray<int32>, AttributeNames::BranchChildrenStr, GroupNames::BranchGroupStr);
+PV_DEFINE_ATTRIBUTE(BranchPoints, TArray<int32>, AttributeNames::BranchPointsStr, GroupNames::BranchGroupStr);
+PV_DEFINE_ATTRIBUTE(BranchNumber, int32, AttributeNames::BranchNumberStr, GroupNames::BranchGroupStr);
+PV_DEFINE_ATTRIBUTE(BranchSourceBudNumber, int32, AttributeNames::BranchSourceBudNumberStr, GroupNames::BranchGroupStr);
+PV_DEFINE_ATTRIBUTE(BranchFoliageIDs, TArray<int32>, AttributeNames::BranchFoliageIDsStr, GroupNames::BranchGroupStr);
+PV_DEFINE_ATTRIBUTE(BranchUVMaterial, int32, AttributeNames::BranchUVMaterialStr, GroupNames::BranchGroupStr);
+PV_DEFINE_ATTRIBUTE(BranchHierarchyNumber, int32, AttributeNames::BranchHierarchyNumberStr, GroupNames::BranchGroupStr);
+PV_DEFINE_ATTRIBUTE(BranchSimulationGroupIndex, int32, AttributeNames::BranchSimulationGroupIndexStr, GroupNames::BranchGroupStr);
+PV_DEFINE_ATTRIBUTE(BranchPlantNumber, int32, AttributeNames::PlantNumberStr, GroupNames::BranchGroupStr);
+PV_DEFINE_ATTRIBUTE(BranchParentNumber, int32, AttributeNames::BranchParentNumberStr, GroupNames::BranchGroupStr);
+
+// Point Attributes
+PV_DEFINE_ATTRIBUTE(PointPosition, FVector3f, AttributeNames::PointPositionStr, GroupNames::PointGroupStr);
+PV_DEFINE_ATTRIBUTE(PointLengthFromRoot, float, AttributeNames::LengthFromRootStr, GroupNames::PointGroupStr);
+PV_DEFINE_ATTRIBUTE(PointLengthFromSeed, float, AttributeNames::LengthFromSeedStr, GroupNames::PointGroupStr);
+PV_DEFINE_ATTRIBUTE(PointScale, float, AttributeNames::PointScaleStr, GroupNames::PointGroupStr);
+PV_DEFINE_ATTRIBUTE(PointTextureCoordV, float, AttributeNames::TextureCoordVStr, GroupNames::PointGroupStr);
+PV_DEFINE_ATTRIBUTE(PointTextureCoordUOffset, float, AttributeNames::TextureCoordUOffsetStr, GroupNames::PointGroupStr);
+PV_DEFINE_ATTRIBUTE(PointURange, FVector2f, AttributeNames::URangeStr, GroupNames::PointGroupStr);
+PV_DEFINE_ATTRIBUTE(PointHullGradient, float, AttributeNames::HullGradientStr, GroupNames::PointGroupStr);
+PV_DEFINE_ATTRIBUTE(PointMainTrunkGradient, float, AttributeNames::MainTrunkGradientStr, GroupNames::PointGroupStr);
+PV_DEFINE_ATTRIBUTE(PointGroundGradient, float, AttributeNames::GroundGradientStr, GroupNames::PointGroupStr);
+PV_DEFINE_ATTRIBUTE(PointScaleGradient, float, AttributeNames::PointScaleGradientStr, GroupNames::PointGroupStr);
+PV_DEFINE_ATTRIBUTE(PointBranchGradient, float, AttributeNames::BranchGradientStr, GroupNames::PointGroupStr);
+PV_DEFINE_ATTRIBUTE(PointPlantGradient, float, AttributeNames::PlantGradientStr, GroupNames::PointGroupStr);
+PV_DEFINE_ATTRIBUTE(PointBudNumber, int32, AttributeNames::BudNumberStr, GroupNames::PointGroupStr);
+PV_DEFINE_ATTRIBUTE(PointNjordPixelIndex, float, AttributeNames::NjordPixelIndexStr, GroupNames::PointGroupStr);
+
+template<typename TFloatType>
+struct TBudLightDetectedView
+{
+	using ArrayType = TFixedSizeArrayHelper<TFloatType>::ArrayType;
+
+	inline const static TArray<float> DefaultData = { 0, 0, 0, 0 };
+
+	ArrayType& Array;
+	TFloatType& Availible;
+	TFloatType& Resource;
+	TFloatType& Branch;
+	TFloatType& Collision;
+
+	TBudLightDetectedView() = delete;
+	TBudLightDetectedView(ArrayType& InData)
+		: Array(TFixedSizeArrayHelper<TFloatType>::GetFixedArray(InData, DefaultData))
+		, Availible(Array[0])
+		, Resource(Array[1])
+		, Branch(Array[2])
+		, Collision(Array[3])
+	{
+	}
+};
+using FBudLightDetectedView = TBudLightDetectedView<float>;
+using FBudLightDetectedConstView = TBudLightDetectedView<const float>;
+PV_DEFINE_STRUCTURED_ATTRIBUTE(BudLightDetected, TArray<float>, FBudLightDetectedView, FBudLightDetectedConstView, AttributeNames::BudLightDetectedStr, GroupNames::PointGroupStr);
+
+template<typename TFloatType>
+struct TBudLateralMeristemView
+{
+	using ArrayType = TFixedSizeArrayHelper<TFloatType>::ArrayType;
+	
+	inline const static TArray<float> DefaultData = { 0, 0, 0, 0, 0, 0, 0 };
+
+	ArrayType& Array;
+	TFloatType& LateralMeristem;
+	TFloatType& Multiplier;
+	TFloatType& Inactive;
+	TFloatType& Davinci;
+	TFloatType& ParentDot;
+	TFloatType& RootDistance;
+	TFloatType& Degredation;
+
+	TBudLateralMeristemView() = delete;
+	TBudLateralMeristemView(ArrayType& InData)
+		: Array(TFixedSizeArrayHelper<TFloatType>::GetFixedArray(InData, DefaultData))
+		, LateralMeristem(Array[0])
+		, Multiplier(Array[1])
+		, Inactive(Array[2])
+		, Davinci(Array[3])
+		, ParentDot(Array[4])
+		, RootDistance(Array[5])
+		, Degredation(Array[6])
+	{
+	}
+};
+using FBudLateralMeristemView = TBudLateralMeristemView<float>;
+using FBudLateralMeristemConstView = TBudLateralMeristemView<const float>;
+PV_DEFINE_STRUCTURED_ATTRIBUTE(BudLateralMeristem, TArray<float>, FBudLateralMeristemView, FBudLateralMeristemConstView, AttributeNames::BudLateralMeristemStr, GroupNames::PointGroupStr);
+
+template<typename TFloatType>
+struct TBudHormoneLevelsView
+{
+	using ArrayType = TFixedSizeArrayHelper<TFloatType>::ArrayType;
+
+	inline const static TArray<float> DefaultData = { 0, 0, 0, 0, 0, 0 };
+
+	ArrayType& Array;
+	TFloatType& Apical;
+	TFloatType& Axillary;
+	TFloatType& AxillaryInhibition;
+	TFloatType& Radical;
+	TFloatType& Ethylene;
+	TFloatType& Cykininin;
+
+	TBudHormoneLevelsView() = delete;
+	TBudHormoneLevelsView(ArrayType& InData)
+		: Array(TFixedSizeArrayHelper<TFloatType>::GetFixedArray(InData, DefaultData))
+		, Apical(Array[0])
+		, Axillary(Array[1])
+		, AxillaryInhibition(Array[2])
+		, Radical(Array[3])
+		, Ethylene(Array[4])
+		, Cykininin(Array[5])
+	{
+	}
+};
+using FBudHormoneLevelsView = TBudHormoneLevelsView<float>;
+using FBudHormoneLevelsConstView = TBudHormoneLevelsView<const float>;
+PV_DEFINE_STRUCTURED_ATTRIBUTE(BudHormoneLevels, TArray<float>, FBudHormoneLevelsView, FBudHormoneLevelsConstView, AttributeNames::BudHormoneLevelsStr, GroupNames::PointGroupStr);
+
+template<typename TVec3Type>
+struct TBudDirectionView
+{
+	using ArrayType = TFixedSizeArrayHelper<TVec3Type>::ArrayType;
+
+	inline const static TArray<FVector3f> DefaultData = { FVector3f::ZeroVector, FVector3f::ZeroVector, FVector3f::ZeroVector, FVector3f::ZeroVector, FVector3f::ZeroVector, FVector3f::ZeroVector };
+
+	ArrayType& Array;
+	TVec3Type& Apical;
+	TVec3Type& Axillary;
+	TVec3Type& LightOptimal;
+	TVec3Type& LightSubOptimal;
+	TVec3Type& GuideCurve;
+	TVec3Type& UpVector;
+
+	TBudDirectionView() = delete;
+	TBudDirectionView(ArrayType& InData)
+		: Array(TFixedSizeArrayHelper<TVec3Type>::GetFixedArray(InData, DefaultData))
+		, Apical(Array[0])
+		, Axillary(Array[1])
+		, LightOptimal(Array[2])
+		, LightSubOptimal(Array[3])
+		, GuideCurve(Array[4])
+		, UpVector(Array[5])
+	{
+	}
+};
+using FBudDirectionView = TBudDirectionView<FVector3f>;
+using FBudDirectionConstView = TBudDirectionView<const FVector3f>;
+PV_DEFINE_STRUCTURED_ATTRIBUTE(BudDirection, TArray<FVector3f>, FBudDirectionView, FBudDirectionConstView, AttributeNames::BudDirectionStr, GroupNames::PointGroupStr);
+
+template<typename TIntType>
+struct TBudStatusView
+{
+	using ArrayType = TFixedSizeArrayHelper<TIntType>::ArrayType;
+	
+	inline const static TArray<int32> DefaultData = { 0, 0, 0, 0, 0, 0, 0, 0, 0, 0 };
+
+	ArrayType& Array;
+	TIntType& ApicalMeristem;
+	TIntType& Codominant;
+	TIntType& Axillary;
+	TIntType& Seed;
+	TIntType& Dormant;
+	TIntType& Triggered;
+	TIntType& NumTriggered;
+	TIntType& Inactive;
+	TIntType& BrokenTip;
+	TIntType& Broken;
+
+	TBudStatusView() = delete;
+	TBudStatusView(ArrayType& InData)
+		: Array(TFixedSizeArrayHelper<TIntType>::GetFixedArray(InData, DefaultData))
+		, ApicalMeristem(Array[0])
+		, Codominant(Array[1])
+		, Axillary(Array[2])
+		, Seed(Array[3])
+		, Dormant(Array[4])
+		, Triggered(Array[5])
+		, NumTriggered(Array[6])
+		, Inactive(Array[7])
+		, BrokenTip(Array[8])
+		, Broken(Array[9])
+	{
+	}
+};
+using FBudStatusView = TBudStatusView<int32>;
+using FBudStatusConstView = TBudStatusView<const int32>;
+PV_DEFINE_STRUCTURED_ATTRIBUTE(BudStatus, TArray<int32>, FBudStatusView, FBudStatusConstView, AttributeNames::BudStatusStr, GroupNames::PointGroupStr);
+
+template<typename TIntType>
+struct TBudDevelopmentView
+{
+	using ArrayType = TFixedSizeArrayHelper<TIntType>::ArrayType;
+
+	inline const static TArray<int32> DefaultData = { 0, 0, 0, 0, 0, 0 };
+
+	ArrayType& Array;
+	TIntType& Generation;
+	TIntType& BudAge;
+	TIntType& BranchAge;
+	TIntType& AgeSenescense;
+	TIntType& LightSenescense;
+	TIntType& RelativeBudAge;
+
+	TBudDevelopmentView() = delete;
+	TBudDevelopmentView(ArrayType& InData)
+		: Array(TFixedSizeArrayHelper<TIntType>::GetFixedArray(InData, DefaultData))
+		, Generation(Array[0])
+		, BudAge(Array[1])
+		, BranchAge(Array[2])
+		, AgeSenescense(Array[3])
+		, LightSenescense(Array[4])
+		, RelativeBudAge(Array[5])
+	{
+	}
+};
+using FBudDevelopmentView = TBudDevelopmentView<int32>;
+using FBudDevelopmentConstView = TBudDevelopmentView<const int32>;
+PV_DEFINE_STRUCTURED_ATTRIBUTE(BudDevelopment, TArray<int32>, FBudDevelopmentView, FBudDevelopmentConstView, AttributeNames::BudDevelopmentStr, GroupNames::PointGroupStr);
+
+// Foliage Attributes
+PV_DEFINE_ATTRIBUTE(FoliageNameID, int32, AttributeNames::FoliageNameIDStr, GroupNames::FoliageGroupStr);
+PV_DEFINE_ATTRIBUTE(FoliageBranchID, int32, AttributeNames::FoliageBranchIDStr, GroupNames::FoliageGroupStr);
+PV_DEFINE_ATTRIBUTE(FoliagePivotPoint, FVector3f, AttributeNames::FoliagePivotPointStr, GroupNames::FoliageGroupStr);
+PV_DEFINE_ATTRIBUTE(FoliageUpVector, FVector3f, AttributeNames::FoliageUPVectorStr, GroupNames::FoliageGroupStr);
+PV_DEFINE_ATTRIBUTE(FoliageNormalVector, FVector3f, AttributeNames::FoliageNormalVectorStr, GroupNames::FoliageGroupStr);
+PV_DEFINE_ATTRIBUTE(FoliageScale, float, AttributeNames::FoliageScaleStr, GroupNames::FoliageGroupStr);
+PV_DEFINE_ATTRIBUTE(FoliageLengthFromRoot, float, AttributeNames::FoliageLengthFromRootStr, GroupNames::FoliageGroupStr);
+PV_DEFINE_ATTRIBUTE(FoliageParentBoneID, int32, AttributeNames::FoliageParentBoneIDStr, GroupNames::FoliageGroupStr);
+
+// FoliageNames Attributes
+PV_DEFINE_ATTRIBUTE(FoliageName, FString, AttributeNames::FoliageNameStr, GroupNames::FoliageNamesGroupStr);
+
+// Details Attributes
+PV_DEFINE_ATTRIBUTE(DetailLeafPhyllotaxy, TArray<float>, AttributeNames::LeafPhyllotaxyStr, GroupNames::DetailsGroupStr);
+PV_DEFINE_ATTRIBUTE(DetailFoliagePath, FString, AttributeNames::FoliagePathStr, GroupNames::DetailsGroupStr);
+PV_DEFINE_ATTRIBUTE(DetailGuid, FGuid, AttributeNames::GuidStr, GroupNames::DetailsGroupStr);
+PV_DEFINE_ATTRIBUTE(DetailLeafGrowth, TArray<float>, AttributeNames::LeafGrowthStr, GroupNames::DetailsGroupStr);
+PV_DEFINE_ATTRIBUTE(DetailAbscissionSenescense, TArray<float>, AttributeNames::AbscissionSenescenseStr, GroupNames::DetailsGroupStr);
+PV_DEFINE_ATTRIBUTE(DetailLateralElongation, TArray<float>, AttributeNames::LateralElongationStr, GroupNames::DetailsGroupStr);
+PV_DEFINE_ATTRIBUTE(DetailTrunkURange, TArray<FVector2f>, AttributeNames::TrunkURangeStr, GroupNames::DetailsGroupStr);
+PV_DEFINE_ATTRIBUTE(DetailTrunkMaterialPath, FString, AttributeNames::TrunkMaterialPathStr, GroupNames::DetailsGroupStr);
+
+// PlantProfiles Attributes
+PV_DEFINE_ATTRIBUTE(ProfilePoints, TArray<float>, AttributeNames::ProfilePointsStr, GroupNames::PlantProfilesGroupStr);
+
+// Bone Attributes
+PV_DEFINE_ATTRIBUTE(BoneName, FString, AttributeNames::BoneNameStr, GroupNames::BonesGroupStr);
+PV_DEFINE_ATTRIBUTE(BoneParentIndex, int32, AttributeNames::BoneParentIndexStr, GroupNames::BonesGroupStr);
+PV_DEFINE_ATTRIBUTE(BonePose, FTransform, AttributeNames::BonePoseStr, GroupNames::BonesGroupStr);
+PV_DEFINE_ATTRIBUTE(BonePointIndex, int32, AttributeNames::BonePointIndexStr, GroupNames::BonesGroupStr);
+PV_DEFINE_ATTRIBUTE(BoneAbsolutePosition, FVector3f, AttributeNames::BoneAbsolutePositionStr, GroupNames::BonesGroupStr);
+PV_DEFINE_ATTRIBUTE(BoneBranchIndex, int32, AttributeNames::BoneBranchIndexStr, GroupNames::BonesGroupStr);
+PV_DEFINE_ATTRIBUTE(BoneId, int32, AttributeNames::BoneIdStr, GroupNames::BonesGroupStr);
+
+// VerticesGroup
+PV_DEFINE_ATTRIBUTE(VerticesPointIds, int32, AttributeNames::VertexPointIdsStr, GroupNames::VerticesGroupStr);
+
+};
\ No newline at end of file
diff --git a/Engine/Plugins/Experimental/ProceduralVegetationEditor/Source/ProceduralVegetation/Public/Facades/PVAttributesNames.h b/Engine/Plugins/Experimental/ProceduralVegetationEditor/Source/ProceduralVegetation/Public/Facades/PVAttributesNames.h
--- a/Engine/Plugins/Experimental/ProceduralVegetationEditor/Source/ProceduralVegetation/Public/Facades/PVAttributesNames.h
+++ b/Engine/Plugins/Experimental/ProceduralVegetationEditor/Source/ProceduralVegetation/Public/Facades/PVAttributesNames.h
@@ -6,91 +6,164 @@
 namespace PV::GroupNames
 {
 	//Group Names
-	inline static const FName BranchGroup = FName(TEXT("Primitives"));
-	inline static const FName PointGroup = FName(TEXT("Points"));
-	inline static const FName FoliageGroup = FName(TEXT("Foliage"));
-	inline static const FName FoliageNamesGroup = FName(TEXT("FoliageNames"));
-	inline static const FName DetailsGroup = FName(TEXT("Details"));
-	inline static const FName PlantProfilesGroup = FName(TEXT("PlantProfiles"));
-	inline static const FName BonesGroup = FName(TEXT("Bones"));
-	inline static const FName VerticesGroup = FName(TEXT("Vertices"));
+	inline constexpr TCHAR BranchGroupStr[] = TEXT("Primitives");
+	inline constexpr TCHAR PointGroupStr[] = TEXT("Points");
+	inline constexpr TCHAR FoliageGroupStr[] = TEXT("Foliage");
+	inline constexpr TCHAR FoliageNamesGroupStr[] = TEXT("FoliageNames");
+	inline constexpr TCHAR DetailsGroupStr[] = TEXT("Details");
+	inline constexpr TCHAR PlantProfilesGroupStr[] = TEXT("PlantProfiles");
+	inline constexpr TCHAR BonesGroupStr[] = TEXT("Bones");
+	inline constexpr TCHAR VerticesGroupStr[] = TEXT("Vertices");
+
+	inline static const FName BranchGroup = FName(BranchGroupStr);
+	inline static const FName PointGroup = FName(PointGroupStr);
+	inline static const FName FoliageGroup = FName(FoliageGroupStr);
+	inline static const FName FoliageNamesGroup = FName(FoliageNamesGroupStr);
+	inline static const FName DetailsGroup = FName(DetailsGroupStr);
+	inline static const FName PlantProfilesGroup = FName(PlantProfilesGroupStr);
+	inline static const FName BonesGroup = FName(BonesGroupStr);
+	inline static const FName VerticesGroup = FName(VerticesGroupStr);
 }
 
 namespace PV::AttributeNames
 {
 	//Branch Attributes
-	inline static const FName BranchParents = FName(TEXT("parents"));
-	inline static const FName BranchChildren = FName(TEXT("children"));
-	inline static const FName BranchPoints = FName(TEXT("Points"));
-	inline static const FName BranchNumber = FName(TEXT("branchNumber"));
-	inline static const FName BranchSourceBudNumber = FName(TEXT("branchSourceBudNumber"));
-	inline static const FName BranchFoliageIDs = FName(TEXT("FoliageID"));
-	inline static const FName BranchUVMaterial = FName(TEXT("branchUVMat"));
-	inline static const FName TrunkMaterialPath = FName(TEXT("TrunkMaterialPath"));
-	inline static const FName TrunkURange = FName(TEXT("TrunkURange"));
-	inline static const FName BranchHierarchyNumber = FName(TEXT("branchHierarchyNumber"));
-	inline static const FName BranchSimulationGroupIndex = FName(TEXT("branchSimulationGroupIndex"));
-	inline static const FName PlantNumber = FName(TEXT("plantNumber"));
-	inline static const FName BranchParentNumber = FName(TEXT("branchParentNumber"));
+	inline constexpr TCHAR BranchParentsStr[] = TEXT("parents");
+	inline constexpr TCHAR BranchChildrenStr[] = TEXT("children");
+	inline constexpr TCHAR BranchPointsStr[] = TEXT("Points");
+	inline constexpr TCHAR BranchNumberStr[] = TEXT("branchNumber");
+	inline constexpr TCHAR BranchSourceBudNumberStr[] = TEXT("branchSourceBudNumber");
+	inline constexpr TCHAR BranchFoliageIDsStr[] = TEXT("FoliageID");
+	inline constexpr TCHAR BranchUVMaterialStr[] = TEXT("branchUVMat");
+	inline constexpr TCHAR TrunkMaterialPathStr[] = TEXT("TrunkMaterialPath");
+	inline constexpr TCHAR TrunkURangeStr[] = TEXT("TrunkURange");
+	inline constexpr TCHAR BranchHierarchyNumberStr[] = TEXT("branchHierarchyNumber");
+	inline constexpr TCHAR BranchSimulationGroupIndexStr[] = TEXT("branchSimulationGroupIndex");
+	inline constexpr TCHAR PlantNumberStr[] = TEXT("plantNumber");
+	inline constexpr TCHAR BranchParentNumberStr[] = TEXT("branchParentNumber");
+
+	inline static const FName BranchParents = FName(BranchParentsStr);
+	inline static const FName BranchChildren = FName(BranchChildrenStr);
+	inline static const FName BranchPoints = FName(BranchPointsStr);
+	inline static const FName BranchNumber = FName(BranchNumberStr);
+	inline static const FName BranchSourceBudNumber = FName(BranchSourceBudNumberStr);
+	inline static const FName BranchFoliageIDs = FName(BranchFoliageIDsStr);
+	inline static const FName BranchUVMaterial = FName(BranchUVMaterialStr);
+	inline static const FName TrunkMaterialPath = FName(TrunkMaterialPathStr);
+	inline static const FName TrunkURange = FName(TrunkURangeStr);
+	inline static const FName BranchHierarchyNumber = FName(BranchHierarchyNumberStr);
+	inline static const FName BranchSimulationGroupIndex = FName(BranchSimulationGroupIndexStr);
+	inline static const FName PlantNumber = FName(PlantNumberStr);
+	inline static const FName BranchParentNumber = FName(BranchParentNumberStr);
 
 	//Point Attributes
-	inline static const FName PointPosition = FName(TEXT("Position"));
-	inline static const FName LengthFromRoot = FName(TEXT("lengthFromRoot"));
-	inline static const FName PointScaleGradient = FName(TEXT("LOD_totalPscaleGradient"));
-	inline static const FName PointScale = FName(TEXT("Scale"));
-	
-	inline static const FName LengthFromSeed = FName(TEXT("lengthFromSeed"));
-	inline static const FName BudLightDetected = FName(TEXT("budLightDetected"));
-	inline static const FName BudDevelopment = FName(TEXT("budDevelopment"));
-	inline static const FName BudLateralMeristem = FName(TEXT("budLateralMeristem"));
-	inline static const FName TextureCoordV = FName(TEXT("uv_v"));
-	inline static const FName TextureCoordUOffset = FName(TEXT("uv_uOffset"));
-	inline static const FName URange = FName(TEXT("uv_uRange"));
-	
-	inline static const FName HullGradient = FName(TEXT("LOD_hullGradient"));
-	inline static const FName MainTrunkGradient = FName(TEXT("LOD_mainTrunkGradient"));
-	inline static const FName GroundGradient = FName(TEXT("LOD_groundGradient"));
-	inline static const FName BudNumber = FName(TEXT("budNumber"));
-	inline static const FName BudHormoneLevels = FName(TEXT("budHormoneLevels"));
-	inline static const FName BudDirection = FName(TEXT("budDirection"));
-	inline static const FName BudStatus = FName(TEXT("budStatus"));
-	inline static const FName BranchGradient = FName(TEXT("branchGradient"));
-	inline static const FName PlantGradient = FName(TEXT("plantGradient"));
-	inline static const FName NjordPixelIndex = FName(TEXT("njord_pixelIdx"));
+	inline constexpr TCHAR PointPositionStr[] = TEXT("Position");
+	inline constexpr TCHAR LengthFromRootStr[] = TEXT("lengthFromRoot");
+	inline constexpr TCHAR PointScaleGradientStr[] = TEXT("LOD_totalPscaleGradient");
+	inline constexpr TCHAR PointScaleStr[] = TEXT("Scale");
+	inline constexpr TCHAR LengthFromSeedStr[] = TEXT("lengthFromSeed");
+	inline constexpr TCHAR BudLightDetectedStr[] = TEXT("budLightDetected");
+	inline constexpr TCHAR BudDevelopmentStr[] = TEXT("budDevelopment");
+	inline constexpr TCHAR BudLateralMeristemStr[] = TEXT("budLateralMeristem");
+	inline constexpr TCHAR TextureCoordVStr[] = TEXT("uv_v");
+	inline constexpr TCHAR TextureCoordUOffsetStr[] = TEXT("uv_uOffset");
+	inline constexpr TCHAR URangeStr[] = TEXT("uv_uRange");
+	inline constexpr TCHAR HullGradientStr[] = TEXT("LOD_hullGradient");
+	inline constexpr TCHAR MainTrunkGradientStr[] = TEXT("LOD_mainTrunkGradient");
+	inline constexpr TCHAR GroundGradientStr[] = TEXT("LOD_groundGradient");
+	inline constexpr TCHAR BudNumberStr[] = TEXT("budNumber");
+	inline constexpr TCHAR BudHormoneLevelsStr[] = TEXT("budHormoneLevels");
+	inline constexpr TCHAR BudDirectionStr[] = TEXT("budDirection");
+	inline constexpr TCHAR BudStatusStr[] = TEXT("budStatus");
+	inline constexpr TCHAR BranchGradientStr[] = TEXT("branchGradient");
+	inline constexpr TCHAR PlantGradientStr[] = TEXT("plantGradient");
+	inline constexpr TCHAR NjordPixelIndexStr[] = TEXT("njord_pixelIdx");
+
+	inline static const FName PointPosition = FName(PointPositionStr);
+	inline static const FName LengthFromRoot = FName(LengthFromRootStr);
+	inline static const FName PointScaleGradient = FName(PointScaleGradientStr);
+	inline static const FName PointScale = FName(PointScaleStr);
+	inline static const FName LengthFromSeed = FName(LengthFromSeedStr);
+	inline static const FName BudLightDetected = FName(BudLightDetectedStr);
+	inline static const FName BudDevelopment = FName(BudDevelopmentStr);
+	inline static const FName BudLateralMeristem = FName(BudLateralMeristemStr);
+	inline static const FName TextureCoordV = FName(TextureCoordVStr);
+	inline static const FName TextureCoordUOffset = FName(TextureCoordUOffsetStr);
+	inline static const FName URange = FName(URangeStr);
+	inline static const FName HullGradient = FName(HullGradientStr);
+	inline static const FName MainTrunkGradient = FName(MainTrunkGradientStr);
+	inline static const FName GroundGradient = FName(GroundGradientStr);
+	inline static const FName BudNumber = FName(BudNumberStr);
+	inline static const FName BudHormoneLevels = FName(BudHormoneLevelsStr);
+	inline static const FName BudDirection = FName(BudDirectionStr);
+	inline static const FName BudStatus = FName(BudStatusStr);
+	inline static const FName BranchGradient = FName(BranchGradientStr);
+	inline static const FName PlantGradient = FName(PlantGradientStr);
+	inline static const FName NjordPixelIndex = FName(NjordPixelIndexStr);
 
 	//Foliage Attributes
-	inline static const FName FoliageNameID = FName(TEXT("FoliageNameID"));
-	inline static const FName FoliageBranchID = FName(TEXT("FoliageBranchID"));
-	inline static const FName FoliagePivotPoint = FName(TEXT("FoliagePivotPoint"));
-	inline static const FName FoliageUPVector = FName(TEXT("FoliageUPVector"));
-	inline static const FName FoliageNormalVector = FName(TEXT("FoliageNormalVector"));
-	inline static const FName FoliageScale = FName(TEXT("FoliageScale"));
-	inline static const FName FoliageLengthFromRoot = FName(TEXT("LengthFromRoot"));
-	inline static const FName FoliageParentBoneID = FName(TEXT("FoliageParentBoneID"));
+	inline constexpr TCHAR FoliageNameIDStr[] = TEXT("FoliageNameID");
+	inline constexpr TCHAR FoliageBranchIDStr[] = TEXT("FoliageBranchID");
+	inline constexpr TCHAR FoliagePivotPointStr[] = TEXT("FoliagePivotPoint");
+	inline constexpr TCHAR FoliageUPVectorStr[] = TEXT("FoliageUPVector");
+	inline constexpr TCHAR FoliageNormalVectorStr[] = TEXT("FoliageNormalVector");
+	inline constexpr TCHAR FoliageScaleStr[] = TEXT("FoliageScale");
+	inline constexpr TCHAR FoliageLengthFromRootStr[] = TEXT("LengthFromRoot");
+	inline constexpr TCHAR FoliageParentBoneIDStr[] = TEXT("FoliageParentBoneID");
+
+	inline static const FName FoliageNameID = FName(FoliageNameIDStr);
+	inline static const FName FoliageBranchID = FName(FoliageBranchIDStr);
+	inline static const FName FoliagePivotPoint = FName(FoliagePivotPointStr);
+	inline static const FName FoliageUPVector = FName(FoliageUPVectorStr);
+	inline static const FName FoliageNormalVector = FName(FoliageNormalVectorStr);
+	inline static const FName FoliageScale = FName(FoliageScaleStr);
+	inline static const FName FoliageLengthFromRoot = FName(FoliageLengthFromRootStr);
+	inline static const FName FoliageParentBoneID = FName(FoliageParentBoneIDStr);
 
 	//FoliageNames Attributes
-	inline static const FName FoliageName = FName(TEXT("Name"));
+	inline constexpr TCHAR FoliageNameStr[] = TEXT("Name");
+
+	inline static const FName FoliageName = FName(FoliageNameStr);
 
 	//Details Attributes
-	inline static const FName LeafPhyllotaxy = FName(TEXT("phyllotaxyLeaf"));
-	inline static const FName FoliagePath = FName(TEXT("FoliagePath"));
-	inline static const FName Guid = FName(TEXT("Guid"));
-	inline static const FName LeafGrowth = FName(TEXT("leafGrowth"));
-	inline static const FName AbscissionSenescense = FName(TEXT("abscissionSenescense"));
-	inline static const FName LateralElongation = FName(TEXT("lateralElongation"));
+	inline constexpr TCHAR LeafPhyllotaxyStr[] = TEXT("phyllotaxyLeaf");
+	inline constexpr TCHAR FoliagePathStr[] = TEXT("FoliagePath");
+	inline constexpr TCHAR GuidStr[] = TEXT("Guid");
+	inline constexpr TCHAR LeafGrowthStr[] = TEXT("leafGrowth");
+	inline constexpr TCHAR AbscissionSenescenseStr[] = TEXT("abscissionSenescense");
+	inline constexpr TCHAR LateralElongationStr[] = TEXT("lateralElongation");
+
+	inline static const FName LeafPhyllotaxy = FName(LeafPhyllotaxyStr);
+	inline static const FName FoliagePath = FName(FoliagePathStr);
+	inline static const FName Guid = FName(GuidStr);
+	inline static const FName LeafGrowth = FName(LeafGrowthStr);
+	inline static const FName AbscissionSenescense = FName(AbscissionSenescenseStr);
+	inline static const FName LateralElongation = FName(LateralElongationStr);
 
 	//PlantProfiles Attributes
-	inline static const FName ProfilePoints = FName(TEXT("ProfilePoints"));
+	inline constexpr TCHAR ProfilePointsStr[] = TEXT("ProfilePoints");
+
+	inline static const FName ProfilePoints = FName(ProfilePointsStr);
 
 	//Bone Attributes
-	inline static const FName BoneName = FName(TEXT("BoneName"));
-	inline static const FName BoneParentIndex = FName(TEXT("BoneParentIndex"));
-	inline static const FName BonePose = FName(TEXT("BonePose"));
-	inline static const FName BonePointIndex = FName(TEXT("BonePointIndex"));
-	inline static const FName BoneAbsolutePosition = FName(TEXT("BoneAbsolutePosition"));
-	inline static const FName BoneBranchIndex = FName(TEXT("BoneBranchIndex"));
-	inline static const FName NjordPixelId = FName(TEXT("NjordPixelId"));
-	inline static const FName VertexPointIds = FName(TEXT("VertexPointIds"));
-	inline static const FName BoneId = FName(TEXT("BoneId"));
+	inline constexpr TCHAR BoneNameStr[] = TEXT("BoneName");
+	inline constexpr TCHAR BoneParentIndexStr[] = TEXT("BoneParentIndex");
+	inline constexpr TCHAR BonePoseStr[] = TEXT("BonePose");
+	inline constexpr TCHAR BonePointIndexStr[] = TEXT("BonePointIndex");
+	inline constexpr TCHAR BoneAbsolutePositionStr[] = TEXT("BoneAbsolutePosition");
+	inline constexpr TCHAR BoneBranchIndexStr[] = TEXT("BoneBranchIndex");
+	inline constexpr TCHAR NjordPixelIdStr[] = TEXT("NjordPixelId");
+	inline constexpr TCHAR VertexPointIdsStr[] = TEXT("VertexPointIds");
+	inline constexpr TCHAR BoneIdStr[] = TEXT("BoneId");
+
+	inline static const FName BoneName = FName(BoneNameStr);
+	inline static const FName BoneParentIndex = FName(BoneParentIndexStr);
+	inline static const FName BonePose = FName(BonePoseStr);
+	inline static const FName BonePointIndex = FName(BonePointIndexStr);
+	inline static const FName BoneAbsolutePosition = FName(BoneAbsolutePositionStr);
+	inline static const FName BoneBranchIndex = FName(BoneBranchIndexStr);
+	inline static const FName NjordPixelId = FName(NjordPixelIdStr);
+	inline static const FName VertexPointIds = FName(VertexPointIdsStr);
+	inline static const FName BoneId = FName(BoneIdStr);
 }
 
diff --git a/Engine/Plugins/Experimental/ProceduralVegetationEditor/Source/ProceduralVegetation/Public/Helpers/PVAttributesHelper.h b/Engine/Plugins/Experimental/ProceduralVegetationEditor/Source/ProceduralVegetation/Public/Helpers/PVAttributesHelper.h
new file mode 100644
--- /dev/null
+++ b/Engine/Plugins/Experimental/ProceduralVegetationEditor/Source/ProceduralVegetation/Public/Helpers/PVAttributesHelper.h
@@ -0,0 +1,604 @@
+// Copyright Epic Games, Inc. All Rights Reserved.
+#pragma once
+
+#include "Facades/PVAttributes.h"
+
+namespace PV::AttributesHelper
+{
+	inline int32 GetBranchParentIndex(
+		FBranchParentNumberAttributeConstView BranchParentNumberAttribute,
+		FBranchNumberAttributeConstView BranchNumberAttribute,
+		int32 BranchIndex
+	)
+	{
+		return BranchNumberAttribute.Find(BranchParentNumberAttribute[BranchIndex]);
+	}
+
+	inline TArray<int32> GetPlantBranchIndices(
+		PV::FBranchPlantNumberAttributeConstView BranchPlantNumberAttribute,
+		int32 PlantNumber
+	)
+	{
+		TArray<int32> OutBranchIndices;
+
+		for (int32 i = 0; i < BranchPlantNumberAttribute.Num(); i++)
+		{
+			if (BranchPlantNumberAttribute[i] == PlantNumber)
+			{
+				OutBranchIndices.Add(i);
+			}
+		}
+
+		return OutBranchIndices;
+	}
+
+	inline TArray<int32> GetTrunkIndices(FBranchParentNumberAttributeConstView BranchParentNumberAttribute)
+	{
+		TArray<int32> TrunkIndicesResult;
+		for (int32 i = 0; i < BranchParentNumberAttribute.Num(); i++)
+		{
+			if (BranchParentNumberAttribute[i] == 0)
+			{
+				TrunkIndicesResult.Add(i);
+			}
+		}
+		return TrunkIndicesResult;
+	}
+
+	inline TArray<int32> GetUniquePlantNumbers(PV::FBranchPlantNumberAttributeConstView BranchPlantNumberAttribute)
+	{
+		TSet<int32> PlantNumberElements;
+		for (int32 i = 0; i < BranchPlantNumberAttribute.Num(); i++)
+		{
+			PlantNumberElements.Add(BranchPlantNumberAttribute[i]);
+		}
+		return PlantNumberElements.Array();
+	}
+
+	enum class EForEachResult
+	{
+		Break    = 0,
+		Continue = 1
+	};
+
+	// Walks the branches in hierarchical order (e.g. start with the trunk and recursivly visit the child branches)
+	template <typename StateType>
+	void RecursiveWalkBranches(
+		FBranchParentNumberAttributeConstView BranchParentNumberAttribute,
+		FBranchChildrenAttributeConstView BranchChildrenAttribute,
+		FBranchNumberAttributeConstView BranchNumberAttribute,
+		const StateType& RootBranchState,
+		const TFunction<EForEachResult(int32, StateType&)>& Body
+	)
+	{
+		const TFunction<EForEachResult(int32, const StateType&)> WalkBranches_Recursive = [&](int32 BranchIndex, const StateType& State) -> EForEachResult
+		{
+			StateType BranchState = State;
+			if (Body(BranchIndex, BranchState) == EForEachResult::Break)
+			{
+				return EForEachResult::Break;
+			}
+
+			const TArray<int32>& BranchChildren = BranchChildrenAttribute[BranchIndex];
+			for (int32 BranchChild : BranchChildren)
+			{
+				const int32 ChildBranchIndex = BranchNumberAttribute.Find(BranchChild);
+				check(ChildBranchIndex != INDEX_NONE);
+
+				if (WalkBranches_Recursive(ChildBranchIndex, BranchState) == EForEachResult::Break)
+				{
+					return EForEachResult::Break;
+				}
+			}
+
+			return EForEachResult::Continue;
+		};
+
+		for (const int32 TrunkIndex : GetTrunkIndices(BranchParentNumberAttribute))
+		{
+			if (WalkBranches_Recursive(TrunkIndex, RootBranchState) == EForEachResult::Break)
+			{
+				return;
+			}
+		}
+	}
+
+	inline void RecursiveWalkBranches(
+		FBranchParentNumberAttributeConstView BranchParentNumberAttribute,
+		FBranchChildrenAttributeConstView BranchChildrenAttribute,
+		FBranchNumberAttributeConstView BranchNumberAttribute,
+		const TFunction<EForEachResult(int32 BranchIndex)>& Body
+	)
+	{
+		RecursiveWalkBranches<int32>(
+			BranchParentNumberAttribute,
+			BranchChildrenAttribute,
+			BranchNumberAttribute,
+			INDEX_NONE,
+			[&](int32 BranchIndex, int32&) -> EForEachResult
+			{
+				if (Body(BranchIndex) == EForEachResult::Break)
+				{
+					return EForEachResult::Break;
+				}
+				return EForEachResult::Continue;
+			}
+		);
+	}
+
+	// Walks the branches in reversed hierarchical order (e.g. first visit all leaf nodes, then their parents, and so on, visiting the tunk last)
+	inline void RecursiveWalkBranches_Reversed(
+		FBranchParentNumberAttributeConstView BranchParentNumberAttribute,
+		FBranchChildrenAttributeConstView BranchChildrenAttribute,
+		FBranchNumberAttributeConstView BranchNumberAttribute,
+		const TFunction<EForEachResult(int32 BranchIndex)>& Body
+	)
+	{
+		TArray<int32> BranchIndices;
+		BranchIndices.Reserve(BranchParentNumberAttribute.Num());
+
+		RecursiveWalkBranches(
+			BranchParentNumberAttribute,
+			BranchChildrenAttribute,
+			BranchNumberAttribute,
+			[&](int32 BranchIndex) -> EForEachResult
+			{
+				BranchIndices.Add(BranchIndex);
+				return EForEachResult::Continue;
+			}
+		);
+
+		for (int32 i = BranchIndices.Num() - 1; i >= 0; --i)
+		{
+			if (Body(BranchIndices[i]) == EForEachResult::Break)
+			{
+				break;
+			}
+		}
+	}
+
+	inline TArray<int32> GetBranchPointFoliageEntryIds(
+		PV::FBranchPointsAttributeConstView BranchPointsAttribute,
+		PV::FPointLengthFromRootAttributeConstView PointLengthFromRootAttribute,
+		PV::FBranchFoliageIDsAttributeConstView BranchFoliageIDsAttribute,
+		PV::FFoliageLengthFromRootAttributeConstView FoliageLengthFromRootAttribute,
+		int32 BranchIndex,
+		int32 PointIndex)
+	{
+		const TArray<int32> BranchPoints = BranchPointsAttribute[BranchIndex];
+		const int32 PointBranchIndex = BranchPoints.Find(PointIndex);
+		if (PointBranchIndex == INDEX_NONE)
+		{
+			return TArray<int32>();
+		}
+
+		const int32 PreviousPointBranchIndex = PointBranchIndex > 0 ? PointBranchIndex - 1 : INDEX_NONE;
+		if (PreviousPointBranchIndex == INDEX_NONE)
+		{
+			return TArray<int32>();
+		}
+
+		const int32 PreviousPointIndex = BranchPoints[PreviousPointBranchIndex];
+
+		const float BranchPointLengthFromRoot = PointLengthFromRootAttribute[PointIndex];
+		const float PreviousBranchPointLengthFromRoot = PointLengthFromRootAttribute[PreviousPointIndex];
+
+		const TArray<int32> FoliageEntryIds = BranchFoliageIDsAttribute[BranchIndex];
+
+		TArray<int32> OutFoliageEntryIds;
+		OutFoliageEntryIds.Reserve(FoliageEntryIds.Num());
+
+		for (const int32 FoliageEntryId : FoliageEntryIds)
+		{
+			const float FoliageEntryLengthFromRoot = FoliageLengthFromRootAttribute[FoliageEntryId];
+
+			if (FoliageEntryLengthFromRoot <= BranchPointLengthFromRoot && FoliageEntryLengthFromRoot > PreviousBranchPointLengthFromRoot)
+			{
+				OutFoliageEntryIds.Add(FoliageEntryId);
+			}
+		}
+
+		return OutFoliageEntryIds;
+	}
+
+	inline TArray<int32> GetBranchPointChildBranchIndices_SplitPoints(
+		PV::FBranchChildrenAttributeConstView BranchChildrenAttribute,
+		PV::FBranchParentNumberAttributeConstView BranchParentNumberAttribute,
+		PV::FBranchNumberAttributeConstView BranchNumberAttribute,
+		PV::FBranchPointsAttributeConstView BranchPointsAttribute,
+		PV::FPointLengthFromRootAttributeConstView PointLengthFromRootAttribute,
+		int32 BranchIndex,
+		int32 PointIndex
+	)
+	{
+		const TArray<int32> BranchPoints = BranchPointsAttribute[BranchIndex];
+		const int32 PointBranchIndex = BranchPoints.Find(PointIndex);
+		if (PointBranchIndex == INDEX_NONE)
+		{
+			return TArray<int32>();
+		}
+
+		const int32 PreviousPointBranchIndex = PointBranchIndex > 0 ? PointBranchIndex - 1 : INDEX_NONE;
+		if (PreviousPointBranchIndex == INDEX_NONE)
+		{
+			return TArray<int32>();
+		}
+
+		const int32 PreviousPointIndex = BranchPoints[PreviousPointBranchIndex];
+
+		const TArray<int32>& BranchChildren = BranchChildrenAttribute[BranchIndex];
+
+		const float BranchPointLengthFromRoot = PointLengthFromRootAttribute[PointIndex];
+		const float PreviousBranchPointLengthFromRoot = PointLengthFromRootAttribute[PreviousPointIndex];
+
+		TArray<int32> ChildrenIndices;
+		ChildrenIndices.Reserve(BranchChildren.Num());
+
+		for (const int32 BranchChild : BranchChildren)
+		{
+			const int32 ChildIndex = BranchNumberAttribute.Find(BranchChild);
+			if (ChildIndex == INDEX_NONE)
+			{
+				continue;
+			}
+
+			const int32 ParentIndex = GetBranchParentIndex(BranchParentNumberAttribute, BranchNumberAttribute, ChildIndex);
+			if (ParentIndex != BranchIndex)
+			{
+				continue;
+			}
+
+			const TArray<int32>& ChildBranchPoints = BranchPointsAttribute[ChildIndex];
+			if (ChildBranchPoints.Num() == 0)
+			{
+				continue;
+			}
+
+			const int32 FirstChildPointIndex = ChildBranchPoints[0];
+			const float ChildLengthFromRoot = PointLengthFromRootAttribute[FirstChildPointIndex];
+			if (ChildLengthFromRoot > BranchPointLengthFromRoot || ChildLengthFromRoot <= PreviousBranchPointLengthFromRoot)
+			{
+				continue;
+			}
+
+			ChildrenIndices.Add(ChildIndex);
+		}
+
+		return ChildrenIndices;
+	}
+
+	inline TArray<int32> GetBranchPointChildBranchIndices_FusedPoints(
+		PV::FBranchChildrenAttributeConstView BranchChildrenAttribute,
+		PV::FBranchParentNumberAttributeConstView BranchParentNumberAttribute,
+		PV::FBranchNumberAttributeConstView BranchNumberAttribute,
+		PV::FBranchPointsAttributeConstView BranchPointsAttribute,
+		int32 BranchIndex,
+		int32 PointIndex
+	)
+	{
+		const TArray<int32>& BranchChildren = BranchChildrenAttribute[BranchIndex];
+
+		TArray<int32> ChildrenIndices;
+		ChildrenIndices.Reserve(BranchChildren.Num());
+
+		for (const int32 BranchChild : BranchChildren)
+		{
+			const int32 ChildIndex = BranchNumberAttribute.Find(BranchChild);
+			if (ChildIndex == INDEX_NONE)
+			{
+				continue;
+			}
+
+			const int32 ParentIndex = GetBranchParentIndex(BranchParentNumberAttribute, BranchNumberAttribute, ChildIndex);
+			if (ParentIndex != BranchIndex)
+			{
+				continue;
+			}
+
+			const TArray<int32>& ChildBranchPoints = BranchPointsAttribute[ChildIndex];
+			if (ChildBranchPoints.Num() == 0)
+			{
+				continue;
+			}
+
+			const int32 FirstChildPointIndex = ChildBranchPoints[0];
+			if (FirstChildPointIndex == PointIndex)
+			{
+				ChildrenIndices.Add(ChildIndex);
+			}
+		}
+
+		return ChildrenIndices;
+	}
+
+	template <typename StateType>
+	inline EForEachResult RecursiveWalkBranchPoints_SplitPoints(
+		PV::FBranchChildrenAttributeConstView BranchChildrenAttribute,
+		PV::FBranchParentNumberAttributeConstView BranchParentNumberAttribute,
+		PV::FBranchNumberAttributeConstView BranchNumberAttribute,
+		PV::FBranchPointsAttributeConstView BranchPointsAttribute,
+		PV::FPointLengthFromRootAttributeConstView PointLengthFromRootAttribute,
+		int32 BranchIndex,
+		const StateType& RootPointState,
+		const TFunction<EForEachResult(int32 BranchIndex, int32 BranchPointIndex, int32 PointIndex, const TArray<int32>& ChildBranchIndices, StateType& PointState)>& Body
+	)
+	{
+		const TFunction<EForEachResult(int32, const StateType&)> WalkPoints = [&](int32 BranchIndex, const StateType& ParentState) -> EForEachResult
+		{
+			StateType PointState = ParentState;
+			const TArray<int32>& CurrentBranchPoints = BranchPointsAttribute[BranchIndex];
+
+			for (int32 Index = 0; Index < CurrentBranchPoints.Num(); ++Index)
+			{
+				const int32 PointIndex = CurrentBranchPoints[Index];
+				const TArray<int32> ChildBranchIndices = GetBranchPointChildBranchIndices_SplitPoints(
+					BranchChildrenAttribute,
+					BranchParentNumberAttribute,
+					BranchNumberAttribute,
+					BranchPointsAttribute,
+					PointLengthFromRootAttribute,
+					BranchIndex,
+					PointIndex
+				);
+
+				if (Body(BranchIndex, Index, PointIndex, ChildBranchIndices, PointState) == EForEachResult::Break)
+				{
+					return EForEachResult::Break;
+				}
+
+				for (const int32 ChildBranchIndex : ChildBranchIndices)
+				{
+					if (WalkPoints(ChildBranchIndex, PointState) == EForEachResult::Break)
+					{
+						return EForEachResult::Break;
+					}
+				}
+			}
+
+			return EForEachResult::Continue;
+		};
+
+		return WalkPoints(BranchIndex, RootPointState);
+	}
+
+	template <typename StateType>
+	inline EForEachResult RecursiveWalkBranchPoints_FusedPoints(
+		PV::FBranchChildrenAttributeConstView BranchChildrenAttribute,
+		PV::FBranchParentNumberAttributeConstView BranchParentNumberAttribute,
+		PV::FBranchNumberAttributeConstView BranchNumberAttribute,
+		PV::FBranchPointsAttributeConstView BranchPointsAttribute,
+		int32 BranchIndex,
+		const StateType& RootPointState,
+		const TFunction<EForEachResult(int32 BranchIndex, int32 BranchPointIndex, int32 PointIndex, const TArray<int32>& ChildBranchIndices, StateType& PointState)>& Body
+	)
+	{
+		const TFunction<EForEachResult(int32, bool, const StateType&)> WalkPoints = [&](int32 BranchIndex, bool bSkipFirstBranchPoint, const StateType& ParentState) -> EForEachResult
+		{
+			StateType PointState = ParentState;
+			const TArray<int32>& CurrentBranchPoints = BranchPointsAttribute[BranchIndex];
+
+			const int32 StartIndex = bSkipFirstBranchPoint ? 1 : 0;
+			for (int32 Index = StartIndex; Index < CurrentBranchPoints.Num(); ++Index)
+			{
+				const int32 PointIndex = CurrentBranchPoints[Index];
+				const TArray<int32> ChildBranchIndices = GetBranchPointChildBranchIndices_FusedPoints(
+					BranchChildrenAttribute,
+					BranchParentNumberAttribute,
+					BranchNumberAttribute,
+					BranchPointsAttribute,
+					BranchIndex,
+					PointIndex
+				);
+
+				if (Body(BranchIndex, Index, PointIndex, ChildBranchIndices, PointState) == EForEachResult::Break)
+				{
+					return EForEachResult::Break;
+				}
+
+				for (const int32 ChildBranchIndex : ChildBranchIndices)
+				{
+					if (WalkPoints(ChildBranchIndex, true, PointState) == EForEachResult::Break)
+					{
+						return EForEachResult::Break;
+					}
+				}
+			}
+
+			return EForEachResult::Continue;
+		};
+
+		return WalkPoints(BranchIndex, false, RootPointState);
+	}
+
+	inline EForEachResult RecursiveWalkBranchPoints_FusedPoints(
+		PV::FBranchChildrenAttributeConstView BranchChildrenAttribute,
+		PV::FBranchParentNumberAttributeConstView BranchParentNumberAttribute,
+		PV::FBranchNumberAttributeConstView BranchNumberAttribute,
+		PV::FBranchPointsAttributeConstView BranchPointsAttribute,
+		int32 BranchIndex,
+		const TFunction<EForEachResult(int32 BranchIndex, int32 BranchPointIndex, int32 PointIndex, const TArray<int32>& ChildPointIndices)>& Body
+	)
+	{
+		return RecursiveWalkBranchPoints_FusedPoints<int32>(
+			BranchChildrenAttribute,
+			BranchParentNumberAttribute,
+			BranchNumberAttribute,
+			BranchPointsAttribute,
+			BranchIndex,
+			0,
+			[&](int32 BranchIndex, int32 BranchPointIndex, int32 PointIndex, const TArray<int32>& ChildPointIndices, int32&) -> EForEachResult
+			{
+				return Body(BranchIndex, BranchPointIndex, PointIndex, ChildPointIndices);
+			}
+		);
+	}
+
+	template <typename StateType>
+	void RecursiveWalkPlantPoints_FusedPoints(
+		PV::FBranchChildrenAttributeConstView BranchChildrenAttribute,
+		PV::FBranchParentNumberAttributeConstView BranchParentNumberAttribute,
+		PV::FBranchNumberAttributeConstView BranchNumberAttribute,
+		PV::FBranchPointsAttributeConstView BranchPointsAttribute,
+		const StateType& RootPointState,
+		const TFunction<EForEachResult(int32 BranchIndex, int32 BranchPointIndex, int32 PointIndex, const TArray<int32>& ChildPointIndices, StateType& PointState)>& Body
+	)
+	{
+		for (const int32 TrunkIndex : GetTrunkIndices(BranchParentNumberAttribute))
+		{
+			auto Result = RecursiveWalkBranchPoints_FusedPoints(
+				BranchChildrenAttribute,
+				BranchParentNumberAttribute,
+				BranchNumberAttribute,
+				BranchPointsAttribute,
+				TrunkIndex,
+				RootPointState,
+				Body
+			);
+			if (Result == EForEachResult::Break)
+			{
+				return;
+			}
+		}
+	}
+
+	inline void RecursiveWalkPlantPoints_FusedPoints(
+		PV::FBranchChildrenAttributeConstView BranchChildrenAttribute,
+		PV::FBranchParentNumberAttributeConstView BranchParentNumberAttribute,
+		PV::FBranchNumberAttributeConstView BranchNumberAttribute,
+		PV::FBranchPointsAttributeConstView BranchPointsAttribute,
+		const TFunction<EForEachResult(int32 BranchIndex, int32 BranchPointIndex, int32 PointIndex, const TArray<int32>& ChildPointIndices)>& Body
+	)
+	{
+		RecursiveWalkPlantPoints_FusedPoints<int32>(
+			BranchChildrenAttribute,
+			BranchParentNumberAttribute,
+			BranchNumberAttribute,
+			BranchPointsAttribute,
+			INDEX_NONE,
+			[&](int32 BranchIndex, int32 BranchPointIndex, int32 PointIndex, const TArray<int32>& ChildPointIndices, int32&)
+			{
+				return Body(BranchIndex, BranchPointIndex, PointIndex, ChildPointIndices);
+			}
+		);
+	}
+
+	inline TArrayView<const int32> GetUniqueBranchPoints_FusedPoints(
+		PV::FBranchPointsAttributeConstView BranchPointsAttribute,
+		PV::FBranchParentNumberAttributeConstView BranchParentNumberAttribute,
+		PV::FBranchNumberAttributeConstView BranchNumberAttribute,
+		int32 BranchIndex
+	)
+	{
+		if (GetBranchParentIndex(BranchParentNumberAttribute, BranchNumberAttribute, BranchIndex) == INDEX_NONE)
+		{
+			return TArrayView<const int32>(BranchPointsAttribute[BranchIndex]);
+		}
+
+		const TArray<int32>& Points = BranchPointsAttribute[BranchIndex];
+		if (Points.Num() > 1)
+		{
+			return TArrayView<const int32>(&Points[1], Points.Num() - 1);
+		}
+
+		return TArrayView<const int32>();
+	}
+
+	inline void ForEachPlantPoint_SplitPoints(
+		PV::FBranchPlantNumberAttributeConstView BranchPlantNumberAttribute,
+		PV::FBranchPointsAttributeConstView BranchPointsAttribute,
+		int32 PlantNumber,
+		const TFunction<EForEachResult(int32 BranchIndex, int32 PointIndex)>& Body
+	)
+	{
+		const TArray<int32> BranchIndices = GetPlantBranchIndices(BranchPlantNumberAttribute, PlantNumber);
+		for (int32 BranchIndex : BranchIndices)
+		{
+			for (int32 PointIndex : BranchPointsAttribute[BranchIndex])
+			{
+				if (Body(BranchIndex, PointIndex) == EForEachResult::Break)
+				{
+					return;
+				}
+			}
+		}
+	}
+
+	inline void ForEachUniquePointOnBranches_FusedPoints(
+		PV::FBranchPointsAttributeConstView BranchPointsAttribute,
+		PV::FBranchParentNumberAttributeConstView BranchParentNumberAttribute,
+		PV::FBranchNumberAttributeConstView BranchNumberAttribute,
+		const TFunction<EForEachResult(int32 BranchIndex, int32 PointIndex)>& Body
+	)
+	{
+		for (int32 BranchIndex = 0; BranchIndex < BranchPointsAttribute.Num(); ++BranchIndex)
+		{
+			for (int32 PointIndex : GetUniqueBranchPoints_FusedPoints(BranchPointsAttribute, BranchParentNumberAttribute, BranchNumberAttribute, BranchIndex))
+			{
+				if (Body(BranchIndex, PointIndex) == EForEachResult::Break)
+				{
+					return;
+				}
+			}
+		}
+	}
+
+	inline void ForEachPlant(
+		PV::FBranchPlantNumberAttributeConstView BranchPlantNumberAttribute,
+		const TFunction<EForEachResult(int32 PlantNumber)>& Body
+	)
+	{
+		for (int32 PlantNumber : GetUniquePlantNumbers(BranchPlantNumberAttribute))
+		{
+			if (Body(PlantNumber) == EForEachResult::Break)
+			{
+				return;
+			}
+		}
+	}
+
+	inline void ForEachPlant(
+		const FManagedArrayCollection& Collection,
+		const TFunction<EForEachResult(int32 PlantNumber)>& Body
+	)
+	{
+		if (auto BranchPlantNumberAttribute = PV::FBranchPlantNumberAttribute::FindAttribute(Collection))
+		{
+			ForEachPlant(BranchPlantNumberAttribute, Body);
+		}
+	}
+
+	inline void ForEachPlantPoint(
+		PV::FBranchPlantNumberAttributeConstView BranchPlantNumberAttribute,
+		PV::FBranchPointsAttributeConstView BranchPointsAttribute,
+		int32 PlantNumber,
+		const TFunction<EForEachResult(int32 BranchIndex, int32 PointIndex)>& Body
+	)
+	{
+		const TArray<int32> BranchIndices = GetPlantBranchIndices(BranchPlantNumberAttribute, PlantNumber);
+		for (int32 BranchIndex : BranchIndices)
+		{
+			for (int32 PointIndex : BranchPointsAttribute[BranchIndex])
+			{
+				if (Body(BranchIndex, PointIndex) == EForEachResult::Break)
+				{
+					return;
+				}
+			}
+		}
+	}
+
+	inline void ForEachPlantPoint(
+		const FManagedArrayCollection& Collection,
+		int32 PlantNumber,
+		const TFunction<EForEachResult(int32 BranchIndex, int32 PointIndex)>& Body
+	)
+	{
+		const PV::FBranchPlantNumberAttributeConstView BranchPlantNumberAttribute = FBranchPlantNumberAttribute::FindAttribute(Collection);
+		const PV::FBranchPointsAttributeConstView BranchPointsAttribute = FBranchPointsAttribute::FindAttribute(Collection);
+
+		if (BranchPlantNumberAttribute && BranchPointsAttribute)
+		{
+			ForEachPlantPoint(BranchPlantNumberAttribute, BranchPointsAttribute, PlantNumber, Body);
+		}
+	}
+}
diff --git a/Engine/Plugins/Experimental/ProceduralVegetationEditor/Source/ProceduralVegetation/Public/Helpers/PVFacadeHelper.h b/Engine/Plugins/Experimental/ProceduralVegetationEditor/Source/ProceduralVegetation/Public/Helpers/PVFacadeHelper.h
deleted file mode 100644
diff --git a/Engine/Plugins/Experimental/ProceduralVegetationEditor/Source/ProceduralVegetation/Public/IPVEditorInterface.h b/Engine/Plugins/Experimental/ProceduralVegetationEditor/Source/ProceduralVegetation/Public/IPVEditorInterface.h
--- a/Engine/Plugins/Experimental/ProceduralVegetationEditor/Source/ProceduralVegetation/Public/IPVEditorInterface.h
+++ b/Engine/Plugins/Experimental/ProceduralVegetationEditor/Source/ProceduralVegetation/Public/IPVEditorInterface.h
@@ -20,7 +20,7 @@ namespace PV
 
 		virtual bool ImportGrowthDataFromStaticMesh(const FPVImportStaticMeshParams& InParams, FManagedArrayCollection& OutCollection) = 0;
 		virtual bool ImportGrowthDataFromTexture2D(const FPVImportTexture2DParams& InParams, const UMaterialInterface* TextureVisualizationMaterial, FPVImportTexture2DOutput& Output) = 0;
-		virtual TSharedRef<FManagedArrayCollection> CreateMockTreeGrowthDataCollection() const = 0;
+		virtual TSharedRef<FManagedArrayCollection> CreateMockTreeGrowthDataCollection(bool bSplitPoints) const = 0;
 	};
 
 	extern PROCEDURALVEGETATION_API IPVEditorInterface* GPVEditorInterface;
diff --git a/Engine/Plugins/Experimental/ProceduralVegetationEditor/Source/ProceduralVegetation/Public/Tests/PVMockTree.h b/Engine/Plugins/Experimental/ProceduralVegetationEditor/Source/ProceduralVegetation/Public/Tests/PVMockTree.h
--- a/Engine/Plugins/Experimental/ProceduralVegetationEditor/Source/ProceduralVegetation/Public/Tests/PVMockTree.h
+++ b/Engine/Plugins/Experimental/ProceduralVegetationEditor/Source/ProceduralVegetation/Public/Tests/PVMockTree.h
@@ -23,22 +23,22 @@ namespace PVMockTree
 	* 
 	* 
 	* Point Indices:
-	*      5*  11*
-	* 		|   /     *9
-	* *6    |  /     /
-	* \     | *10   /
-	*  \   4*  \   /
-	*   \   |   \ /
-	*    \  |    *8
-	*     \ |   /
-	*      3*  *7
-	*       | /
-	*       |/
-	*       *2
-	*       |
-	*       *1
-	*       |
-	*       *0
+	*       5*  11*
+	*        |   /     *9
+	* *6     |  /     /
+	*  \     | *10   /
+	*   \   4*  \   /
+	*    \   |   \ /
+	*     \  |    *8
+	*      \ |   /
+	*       3*  *7
+	*        | /
+	*        |/
+	*        *2
+	*        |
+	*        *1
+	*        |
+	*        *0
 	*/
 
 	const FVector3f Point0  = {  00.0000,  00.0000, 000.0000 };
@@ -80,17 +80,17 @@ namespace PVMockTree
 	const float Point10To11_Dist = 26.1638;
 
 	const float Point0_LengthFromRoot = 0;
-	const float Point1_LengthFromRoot = Point0To1_Dist;
-	const float Point2_LengthFromRoot = Point1_LengthFromRoot + Point1To2_Dist;
-	const float Point3_LengthFromRoot = Point2_LengthFromRoot + Point2To3_Dist;
-	const float Point4_LengthFromRoot = Point3_LengthFromRoot + Point3To4_Dist;
-	const float Point5_LengthFromRoot = Point4_LengthFromRoot + Point4To5_Dist;
-	const float Point6_LengthFromRoot = Point3_LengthFromRoot + Point3To6_Dist;
-	const float Point7_LengthFromRoot = Point2_LengthFromRoot + Point2To7_Dist;
-	const float Point8_LengthFromRoot = Point7_LengthFromRoot + Point7To8_Dist;
-	const float Point9_LengthFromRoot = Point8_LengthFromRoot + Point8To9_Dist;
-	const float Point10_LengthFromRoot = Point8_LengthFromRoot + Point8To10_Dist;
-	const float Point11_LengthFromRoot = Point10_LengthFromRoot + Point10To11_Dist;
+	const float Point1_LengthFromRoot = Point0To1_Dist * 0.01f/*cm->m*/;
+	const float Point2_LengthFromRoot = Point1_LengthFromRoot + Point1To2_Dist * 0.01f/*cm->m*/;
+	const float Point3_LengthFromRoot = Point2_LengthFromRoot + Point2To3_Dist * 0.01f/*cm->m*/;
+	const float Point4_LengthFromRoot = Point3_LengthFromRoot + Point3To4_Dist * 0.01f/*cm->m*/;
+	const float Point5_LengthFromRoot = Point4_LengthFromRoot + Point4To5_Dist * 0.01f/*cm->m*/;
+	const float Point6_LengthFromRoot = Point3_LengthFromRoot + Point3To6_Dist * 0.01f/*cm->m*/;
+	const float Point7_LengthFromRoot = Point2_LengthFromRoot + Point2To7_Dist * 0.01f/*cm->m*/;
+	const float Point8_LengthFromRoot = Point7_LengthFromRoot + Point7To8_Dist * 0.01f/*cm->m*/;
+	const float Point9_LengthFromRoot = Point8_LengthFromRoot + Point8To9_Dist * 0.01f/*cm->m*/;
+	const float Point10_LengthFromRoot = Point8_LengthFromRoot + Point8To10_Dist * 0.01f/*cm->m*/;
+	const float Point11_LengthFromRoot = Point10_LengthFromRoot + Point10To11_Dist * 0.01f/*cm->m*/;
 
 	const TArray<int32> Branch0_PointIndices = { 0, 1, 2, 3, 4, 5 };
 	const TArray<int32> Branch1_PointIndices = { 2, 7, 8, 9 };
@@ -122,7 +122,7 @@ namespace PVMockTreeCollection
 
 	PROCEDURALVEGETATION_API extern CompareResult Compare(const FManagedArrayCollection& Collection, bool bCheckPointPositions);
 
-	PROCEDURALVEGETATION_API extern TSharedRef<FManagedArrayCollection> CreateCollection();
+	PROCEDURALVEGETATION_API extern TSharedRef<FManagedArrayCollection> CreateCollection(bool bSplitPoints = true);
 }
 
 #endif // WITH_DEV_AUTOMATION_TESTS
\ No newline at end of file
diff --git a/Engine/Plugins/Experimental/ProceduralVegetationEditor/Source/ProceduralVegetationEditor/Private/Importer/PVImportHelpers.cpp b/Engine/Plugins/Experimental/ProceduralVegetationEditor/Source/ProceduralVegetationEditor/Private/Importer/PVImportHelpers.cpp
--- a/Engine/Plugins/Experimental/ProceduralVegetationEditor/Source/ProceduralVegetationEditor/Private/Importer/PVImportHelpers.cpp
+++ b/Engine/Plugins/Experimental/ProceduralVegetationEditor/Source/ProceduralVegetationEditor/Private/Importer/PVImportHelpers.cpp
@@ -8,12 +8,7 @@
 #include "GeometryCollectionToDynamicMesh.h"
 #include "GeometryCollection/GeometryCollection.h"
 
-#include "Helpers/PVFacadeHelper.h"
-#include "Facades/PVAttributesNames.h"
-#include "Facades/PVFoliageFacade.h"
-#include "Facades/PVPointFacade.h"
-#include "Facades/PVBranchFacade.h"
-#include "Facades/PVPlantFacade.h"
+#include "Helpers/PVAttributesHelper.h"
 
 #include "Engine/StaticMesh.h"
 
@@ -703,197 +698,155 @@ void PV::ImportHelper::CreateEmptyGrowthData(FManagedArrayCollection& OutCollect
 	// PointGroup
 	OutCollection.AddGroup(PV::GroupNames::PointGroup);
 	OutCollection.AddElements(NumBranchPoints, PV::GroupNames::PointGroup);
-	auto& PointPositionAttribute = OutCollection.AddAttribute<FVector3f>(PV::AttributeNames::PointPosition, PV::GroupNames::PointGroup);
-	auto& LengthFromRootAttribute = OutCollection.AddAttribute<float>(PV::AttributeNames::LengthFromRoot, PV::GroupNames::PointGroup);
-	auto& LengthFromSeedAttribute = OutCollection.AddAttribute<float>(PV::AttributeNames::LengthFromSeed, PV::GroupNames::PointGroup);
-	auto& PointScaleGradientAttribute = OutCollection.AddAttribute<float>(PV::AttributeNames::PointScaleGradient, PV::GroupNames::PointGroup);
-	auto& HullGradientAttribute = OutCollection.AddAttribute<float>(PV::AttributeNames::HullGradient, PV::GroupNames::PointGroup);
-	auto& MainTrunkGradientAttribute = OutCollection.AddAttribute<float>(PV::AttributeNames::MainTrunkGradient, PV::GroupNames::PointGroup);
-	auto& GroundGradientAttribute = OutCollection.AddAttribute<float>(PV::AttributeNames::GroundGradient, PV::GroupNames::PointGroup);
-	auto& PointScaleAttribute = OutCollection.AddAttribute<float>(PV::AttributeNames::PointScale, PV::GroupNames::PointGroup);
-	auto& BudDirectionAttribute = OutCollection.AddAttribute<TArray<FVector3f>>(PV::AttributeNames::BudDirection, PV::GroupNames::PointGroup);
-	auto& BranchGradientAttribute = OutCollection.AddAttribute<float>(PV::AttributeNames::BranchGradient, PV::GroupNames::PointGroup);
-	auto& BudHormoneLevelsAttribute = OutCollection.AddAttribute<TArray<float>>(PV::AttributeNames::BudHormoneLevels, PV::GroupNames::PointGroup);
-	auto& PlantGradientAttribute = OutCollection.AddAttribute<float>(PV::AttributeNames::PlantGradient, PV::GroupNames::PointGroup);
-	auto& BudNumberAttribute = OutCollection.AddAttribute<int32>(PV::AttributeNames::BudNumber, PV::GroupNames::PointGroup);
-	auto& NjordPixelIndexAttribute = OutCollection.AddAttribute<float>(PV::AttributeNames::NjordPixelIndex, PV::GroupNames::PointGroup);
-	auto& BudLightDetectedAttribute = OutCollection.AddAttribute<TArray<float>>(PV::AttributeNames::BudLightDetected, PV::GroupNames::PointGroup);
-	auto& BudDevelopmentAttribute = OutCollection.AddAttribute<TArray<int32>>(PV::AttributeNames::BudDevelopment, PV::GroupNames::PointGroup);
-
-	// Zero initialize vector attributes
-	for (int32 i = 0; i < NumBranchPoints; ++i)
-	{
-		PointPositionAttribute[i] = FVector3f::ZeroVector;
-		LengthFromRootAttribute[i] = 0;
-		LengthFromSeedAttribute[i] = 0;
-		PointScaleGradientAttribute[i] = 0;
-		HullGradientAttribute[i] = 0;
-		MainTrunkGradientAttribute[i] = 0;
-		GroundGradientAttribute[i] = 0;
-		PointScaleAttribute[i] = 0.f;
-		BudDirectionAttribute[i] = { FVector3f::UpVector, FVector3f::UpVector, FVector3f::UpVector, FVector3f::UpVector, FVector3f::UpVector, FVector3f::UpVector };
-		BranchGradientAttribute[i] = 0;
-		BudHormoneLevelsAttribute[i] = { 0.f, 0.f, 0.f, 0.f, 0.f };
-		PlantGradientAttribute[i] = 0;
-		BudNumberAttribute[i] = 0;
-		NjordPixelIndexAttribute[i] = 0;
-		BudLightDetectedAttribute[i] = { 1 };
-		BudDevelopmentAttribute[i] = { 0, 0, 0, 0, 0, 0 };
-	}
+	FPointPositionAttribute::AddAttribute(OutCollection, FVector3f::ZeroVector);
+	FPointLengthFromRootAttribute::AddAttribute(OutCollection, 0);
+	FPointLengthFromSeedAttribute::AddAttribute(OutCollection, 0);
+	FPointScaleGradientAttribute::AddAttribute(OutCollection, 0);
+	FPointHullGradientAttribute::AddAttribute(OutCollection, 0);
+	FPointMainTrunkGradientAttribute::AddAttribute(OutCollection, 0);
+	FPointGroundGradientAttribute::AddAttribute(OutCollection, 0);
+	FPointScaleAttribute::AddAttribute(OutCollection, 0.f);
+	FBudDirectionAttribute::AddAttribute(OutCollection);
+	FBudStatusAttribute::AddAttribute(OutCollection);
+	FBudLateralMeristemAttribute::AddAttribute(OutCollection);
+	FPointBranchGradientAttribute::AddAttribute(OutCollection, 0);
+	FBudHormoneLevelsAttribute::AddAttribute(OutCollection);
+	FPointPlantGradientAttribute::AddAttribute(OutCollection, 0);
+	FPointBudNumberAttribute::AddAttribute(OutCollection, 0);
+	FPointNjordPixelIndexAttribute::AddAttribute(OutCollection, 0);
+	FBudLightDetectedAttribute::AddAttribute(OutCollection);
+	FBudDevelopmentAttribute::AddAttribute(OutCollection);
 
 	// BranchGroup
 	OutCollection.AddGroup(PV::GroupNames::BranchGroup);
 	OutCollection.AddElements(NumBranches, PV::GroupNames::BranchGroup);
-	auto& BranchParentsAttribute = OutCollection.AddAttribute<TArray<int32>>(PV::AttributeNames::BranchParents, PV::GroupNames::BranchGroup);
-	auto& BranchChildrenAttribute = OutCollection.AddAttribute<TArray<int32>>(PV::AttributeNames::BranchChildren, PV::GroupNames::BranchGroup);
-	auto& BranchPointsAttribute = OutCollection.AddAttribute<TArray<int32>>(PV::AttributeNames::BranchPoints, PV::GroupNames::BranchGroup);
-	auto& BranchNumberAttribute = OutCollection.AddAttribute<int32>(PV::AttributeNames::BranchNumber, PV::GroupNames::BranchGroup);
-	auto& BranchSourceBudNumberAttribute = OutCollection.AddAttribute<int32>(PV::AttributeNames::BranchSourceBudNumber, PV::GroupNames::BranchGroup);
-	auto& BranchHierarchyNumberAttribute = OutCollection.AddAttribute<int32>(PV::AttributeNames::BranchHierarchyNumber, PV::GroupNames::BranchGroup);
-	auto& BranchParentNumberAttribute = OutCollection.AddAttribute<int32>(PV::AttributeNames::BranchParentNumber, PV::GroupNames::BranchGroup);
-	auto& PlantNumberAttribute = OutCollection.AddAttribute<int32>(PV::AttributeNames::PlantNumber, PV::GroupNames::BranchGroup);
-	auto& BranchFoliageIDsAttribute = OutCollection.AddAttribute<TArray<int32>>(PV::AttributeNames::BranchFoliageIDs, PV::GroupNames::BranchGroup);
-
-	for (int32 i = 0; i < NumBranches; ++i)
-	{
-		BranchParentsAttribute[i] = {};
-		BranchChildrenAttribute[i] = {};
-		BranchPointsAttribute[i] = {};
-		BranchNumberAttribute[i] = 0;
-		BranchSourceBudNumberAttribute[i] = 0;
-		BranchHierarchyNumberAttribute[i] = 0;
-		BranchParentNumberAttribute[i] = 0;
-		PlantNumberAttribute[i] = 0;
-		BranchFoliageIDsAttribute[i] = {};
-	}
+	FBranchParentsAttribute::AddAttribute(OutCollection);
+	FBranchChildrenAttribute::AddAttribute(OutCollection);
+	FBranchPointsAttribute::AddAttribute(OutCollection);
+	FBranchNumberAttribute::AddAttribute(OutCollection, 0);
+	FBranchSourceBudNumberAttribute::AddAttribute(OutCollection, 0);
+	FBranchHierarchyNumberAttribute::AddAttribute(OutCollection, 0);
+	FBranchParentNumberAttribute::AddAttribute(OutCollection, 0);
+	FBranchPlantNumberAttribute::AddAttribute(OutCollection, 0);
+	FBranchFoliageIDsAttribute::AddAttribute(OutCollection);
 
 	// FoliageNamesGroup
 	OutCollection.AddGroup(PV::GroupNames::FoliageNamesGroup);
 	OutCollection.AddElements(NumFoliageMeshes, PV::GroupNames::FoliageNamesGroup);
-	auto& FoliageNamesAttribute = OutCollection.AddAttribute<FString>(PV::AttributeNames::FoliageName, PV::GroupNames::FoliageNamesGroup);
-
-	for (int32 i = 0; i < NumFoliageMeshes; ++i)
-	{
-		FoliageNamesAttribute[i] = TEXT("");
-	}
+	FFoliageNameAttribute::AddAttribute(OutCollection, TEXT(""));
 	
 	// FoliageGroup
 	OutCollection.AddGroup(PV::GroupNames::FoliageGroup);
 	OutCollection.AddElements(NumFoliagePoints, PV::GroupNames::FoliageGroup);
-	auto& FoliageNameIDAttribute = OutCollection.AddAttribute<int32>(PV::AttributeNames::FoliageNameID, PV::GroupNames::FoliageGroup);
-	auto& FoliageBranchIDAttribute = OutCollection.AddAttribute<int32>(PV::AttributeNames::FoliageBranchID, PV::GroupNames::FoliageGroup);
-	auto& FoliagePivotPointAttribute = OutCollection.AddAttribute<FVector3f>(PV::AttributeNames::FoliagePivotPoint, PV::GroupNames::FoliageGroup);
-	auto& FoliageUPVectorAttribute = OutCollection.AddAttribute<FVector3f>(PV::AttributeNames::FoliageUPVector, PV::GroupNames::FoliageGroup);
-	auto& FoliageNormalVectorAttribute = OutCollection.AddAttribute<FVector3f>(PV::AttributeNames::FoliageNormalVector, PV::GroupNames::FoliageGroup);
-	auto& FoliageScaleAttribute = OutCollection.AddAttribute<float>(PV::AttributeNames::FoliageScale, PV::GroupNames::FoliageGroup);
-	auto& FoliageLengthFromRootAttribute = OutCollection.AddAttribute<float>(PV::AttributeNames::FoliageLengthFromRoot, PV::GroupNames::FoliageGroup);
-	auto& FoliageParentBoneIDAttribute = OutCollection.AddAttribute<int32>(PV::AttributeNames::FoliageParentBoneID, PV::GroupNames::FoliageGroup);
-
-	for (int32 i = 0; i < NumFoliagePoints; ++i)
-	{
-		FoliageNameIDAttribute[i] = 0;
-		FoliageBranchIDAttribute[i] = 0;
-		FoliagePivotPointAttribute[i] = FVector3f::ZeroVector;
-		FoliageUPVectorAttribute[i] = FVector3f::ZeroVector;
-		FoliageNormalVectorAttribute[i] = FVector3f::ZeroVector;
-		FoliageScaleAttribute[i] = 0;
-		FoliageLengthFromRootAttribute[i] = 0;
-		FoliageParentBoneIDAttribute[i] = 0;
-	}
+	FFoliageNameIDAttribute::AddAttribute(OutCollection, 0);
+	FFoliageBranchIDAttribute::AddAttribute(OutCollection, 0);
+	FFoliagePivotPointAttribute::AddAttribute(OutCollection, FVector3f::ZeroVector);
+	FFoliageUpVectorAttribute::AddAttribute(OutCollection, FVector3f::ZeroVector);
+	FFoliageNormalVectorAttribute::AddAttribute(OutCollection, FVector3f::ZeroVector);
+	FFoliageScaleAttribute::AddAttribute(OutCollection, 0);
+	FFoliageLengthFromRootAttribute::AddAttribute(OutCollection, 0);
+	FFoliageParentBoneIDAttribute::AddAttribute(OutCollection, 0);
 
 	// DetailsGroup
 	OutCollection.AddGroup(PV::GroupNames::DetailsGroup);
 	OutCollection.AddElements(1, PV::GroupNames::DetailsGroup);
-	auto& FoliagePathAttribute = OutCollection.AddAttribute<FString>(PV::AttributeNames::FoliagePath, PV::GroupNames::DetailsGroup);
-	auto& TrunkMaterialPathAttribute = OutCollection.AddAttribute<FString>(PV::AttributeNames::TrunkMaterialPath, PV::GroupNames::DetailsGroup);
-	auto& TrunkURangeAttribute = OutCollection.AddAttribute<TArray<FVector2f>>(PV::AttributeNames::TrunkURange, PV::GroupNames::DetailsGroup);
-	auto& LeafPhyllotaxyAttribute = OutCollection.AddAttribute<TArray<float>>(PV::AttributeNames::LeafPhyllotaxy, PV::GroupNames::DetailsGroup);
-
-	FoliagePathAttribute[0] = TEXT("");
-	TrunkMaterialPathAttribute[0] = 0;
-	TrunkURangeAttribute[0] = { };
-	LeafPhyllotaxyAttribute[0] = { 0, 200, 50, 1, 1, 0, 0, 0 };
+	FDetailFoliagePathAttribute::AddAttribute(OutCollection, TEXT(""));
+	FDetailTrunkMaterialPathAttribute::AddAttribute(OutCollection, TEXT(""));
+	FDetailTrunkURangeAttribute::AddAttribute(OutCollection, { });
+	FDetailLeafPhyllotaxyAttribute::AddAttribute(OutCollection, { 0, 200, 50, 1, 1, 0, 0, 0 });
 }
 
-void PV::ImportHelper::FillGrowthDataFromBranchHierarchyDescription(
-	FManagedArrayCollection& InOutCollection, 
+void PV::ImportHelper::FillAttributesFromBranchHierarchy(
+	FManagedArrayCollection& OutCollection,
 	TArrayView<const FPVBranchHierarchyDescription> InBranchHierarchies
 )
 {
-	PV::Facades::FPointFacade PointFacade(InOutCollection);
-	PV::Facades::FBranchFacade BranchFacade(InOutCollection);
+	FillAttributesFromBranchHierarchy(
+		FPointPositionAttribute::GetAttribute(OutCollection),
+		FPointScaleAttribute::GetAttribute(OutCollection),
+		FBranchPointsAttribute::GetAttribute(OutCollection),
+		FBranchParentsAttribute::GetAttribute(OutCollection),
+		FBranchChildrenAttribute::GetAttribute(OutCollection),
+		FBranchParentNumberAttribute::GetAttribute(OutCollection),
+		FBranchNumberAttribute::GetAttribute(OutCollection),
+		FBranchHierarchyNumberAttribute::GetAttribute(OutCollection),
+		FBranchPlantNumberAttribute::GetAttribute(OutCollection),
+		InBranchHierarchies
+	);
+}
 
+void PV::ImportHelper::FillAttributesFromBranchHierarchy(
+	PV::FPointPositionAttributeView PointPositionAttribute,
+	PV::FPointScaleAttributeView PointScaleAttribute,
+	PV::FBranchPointsAttributeView BranchPointsAttribute,
+	PV::FBranchParentsAttributeView BranchParentsAttribute,
+	PV::FBranchChildrenAttributeView BranchChildrenAttribute,
+	PV::FBranchParentNumberAttributeView BranchParentNumberAttribute,
+	PV::FBranchNumberAttributeView BranchNumberAttribute,
+	PV::FBranchHierarchyNumberAttributeView BranchHierarchyNumberAttribute,
+	PV::FBranchPlantNumberAttributeView BranchPlantNumberAttribute,
+	TArrayView<const FPVBranchHierarchyDescription> InBranchHierarchies
+)
+{
 	const static auto BranchIndexToBranchNumber = [](int32 BranchIndex) 
 	{ 
 		// Branch Number is BranchIndex + 1 since branch with parent number 0 is a root in PVE and our root has parent index -1.
 		return BranchIndex + 1; 
 	}; 
 
-	using FFillBranchesRecursive = TFunction<void(const FPVBranchHierarchyDescription&, int32, int32, int32, int32, int32, int32&)>;
+	using FFillBranchesRecursive = TFunction<void(const FPVBranchHierarchyDescription&, int32, int32, int32, int32, int32)>;
 
-	const FFillBranchesRecursive FillBranchesRecursive = [&PointFacade, &BranchFacade, &FillBranchesRecursive](
+	const FFillBranchesRecursive FillBranchesRecursive = [&](
 		const FPVBranchHierarchyDescription& BranchHierarchy, 
 		int32 PlantNumber,
 		int32 BranchIndex,
-		int32 FacadeBranchIndexOffset,
-		int32 FacadePointIndexOffset,
-		int32 Depth,
-		int32& BudNumber
+		int32 AttributeBranchIndexOffset,
+		int32 AttributePointIndexOffset,
+		int32 Depth
 	)
 	{
 		const FPVBranchDescription& Branch = BranchHierarchy.Branches[BranchIndex];
 
-		const int32 FacadeRootBranchIndex = BranchHierarchy.RootBranchIndex + FacadeBranchIndexOffset;
-		const int32 FacadeBranchIndex = BranchIndex + FacadeBranchIndexOffset;
-		const int32 FacadeParentBranchIndex = Branch.ParentBranchIndex != INDEX_NONE 
-			? Branch.ParentBranchIndex + FacadeBranchIndexOffset 
+		const int32 AttributeRootBranchIndex = BranchHierarchy.RootBranchIndex + AttributeBranchIndexOffset;
+		const int32 AttributeBranchIndex = BranchIndex + AttributeBranchIndexOffset;
+		const int32 AttributeParentBranchIndex = Branch.ParentBranchIndex != INDEX_NONE
+			? Branch.ParentBranchIndex + AttributeBranchIndexOffset
 			: INDEX_NONE;
 
-		TArray<int32> FacadeBranchPoints;
-		FacadeBranchPoints.Reserve(Branch.PointIndices.Num());
+		TArray<int32> AttributeBranchPoints;
+		AttributeBranchPoints.Reserve(Branch.PointIndices.Num());
 
 		for (int32 i = 0; i < Branch.PointIndices.Num(); ++i)
 		{
-			const int32 FacadePointIndex = Branch.PointIndices[i] + FacadePointIndexOffset;
-			FacadeBranchPoints.Add(FacadePointIndex);
-
-			PointFacade.SetBudNumber(FacadePointIndex, BudNumber);
-
-			// Setting the njord pixel index to the same as the bud number only works since we're traversing the branches recursivly
-			// This shold be set by the mesher but we set it here since it's easy to do so.
-			PointFacade.SetNjordPixelIndex(FacadePointIndex, BudNumber);
-			
-			BudNumber++;
-			
-			// TODO: Set these based on this doc https://docs.google.com/document/d/1fWNcORvDreHBOakcFfiKPhSXhE5yVTCabn4eWZU1jWg/edit?tab=t.0
-			//BudHormoneLevels[i] = { 0.f, 0.f, 0.f, 0.f, 0.f };
+			const int32 AttributePointIndex = Branch.PointIndices[i] + AttributePointIndexOffset;
+			AttributeBranchPoints.Add(AttributePointIndex);
 		}
-		BranchFacade.SetPoints(FacadeBranchIndex, MoveTemp(FacadeBranchPoints));
+		BranchPointsAttribute[AttributeBranchIndex] = MoveTemp(AttributeBranchPoints);
 
-		if (FacadeParentBranchIndex != INDEX_NONE)
+		if (AttributeParentBranchIndex != INDEX_NONE)
 		{
-			TArray<int32> Parents = BranchFacade.GetParents(FacadeParentBranchIndex);
-			Parents.Add(BranchIndexToBranchNumber(FacadeParentBranchIndex));
-			BranchFacade.SetParents(FacadeBranchIndex, MoveTemp(Parents));
+			TArray<int32> Parents = BranchParentsAttribute[AttributeParentBranchIndex];
+			Parents.Add(BranchIndexToBranchNumber(AttributeParentBranchIndex));
+			BranchParentsAttribute[AttributeBranchIndex] = MoveTemp(Parents);
 		}
 
-		TArray<int32> FacadeBranchChildren;
-		FacadeBranchChildren.SetNumUninitialized(Branch.ChildBranchIndices.Num());
+		TArray<int32> AttributeBranchChildren;
+		AttributeBranchChildren.SetNumUninitialized(Branch.ChildBranchIndices.Num());
 		for (int32 i = 0; i < Branch.ChildBranchIndices.Num(); ++i)
 		{
-			const int32 FacadeChildBranchIndex = Branch.ChildBranchIndices[i] + FacadeBranchIndexOffset;
+			const int32 AttributeChildBranchIndex = Branch.ChildBranchIndices[i] + AttributeBranchIndexOffset;
 
-			check(FacadeChildBranchIndex != FacadeRootBranchIndex);
-			FacadeBranchChildren[i] = BranchIndexToBranchNumber(FacadeChildBranchIndex);
+			check(AttributeChildBranchIndex != AttributeRootBranchIndex);
+			AttributeBranchChildren[i] = BranchIndexToBranchNumber(AttributeChildBranchIndex);
 		}
-		BranchFacade.SetChildren(FacadeBranchIndex, MoveTemp(FacadeBranchChildren));
 
-		BranchFacade.SetParentBranchNumber(FacadeBranchIndex, BranchIndexToBranchNumber(FacadeParentBranchIndex));
-		BranchFacade.SetBranchNumber(FacadeBranchIndex, BranchIndexToBranchNumber(FacadeBranchIndex));
+		BranchChildrenAttribute[AttributeBranchIndex] = MoveTemp(AttributeBranchChildren);
+
+		BranchParentNumberAttribute[AttributeBranchIndex] = BranchIndexToBranchNumber(AttributeParentBranchIndex);
+		BranchNumberAttribute[AttributeBranchIndex] = BranchIndexToBranchNumber(AttributeBranchIndex);
 
-		BranchFacade.SetBranchSourceBudNumber(FacadeBranchIndex, -1);
-		BranchFacade.SetBranchHierarchyNumber(FacadeBranchIndex, Depth);
-		BranchFacade.SetBranchPlantNumber(FacadeBranchIndex, PlantNumber);
+		BranchHierarchyNumberAttribute[AttributeBranchIndex] = Depth;
+		BranchPlantNumberAttribute[AttributeBranchIndex] = PlantNumber;
 
 		for (int32 ChildIndex : Branch.ChildBranchIndices)
 		{
@@ -901,77 +854,130 @@ void PV::ImportHelper::FillGrowthDataFromBranchHierarchyDescription(
 				BranchHierarchy, 
 				PlantNumber,
 				ChildIndex, 
-				FacadeBranchIndexOffset,
-				FacadePointIndexOffset,
-				Depth + 1, 
-				BudNumber
+				AttributeBranchIndexOffset,
+				AttributePointIndexOffset,
+				Depth + 1
 			);
 		}
 	};
 
-	int32 FacadeBranchIndexOffset = 0;
-	int32 FacadePointIndexOffset = 0;
+	int32 AttributeBranchIndexOffset = 0;
+	int32 AttributePointIndexOffset = 0;
 	int32 PlantNumber = 0;
 	for (const FPVBranchHierarchyDescription& BranchHierarchy : InBranchHierarchies)
 	{
-		int32 BudNumber = 1;
-
 		FillBranchesRecursive(
 			BranchHierarchy, 
 			PlantNumber,
 			BranchHierarchy.RootBranchIndex, 
-			FacadeBranchIndexOffset,
-			FacadePointIndexOffset,
-			0, 
-			BudNumber
+			AttributeBranchIndexOffset,
+			AttributePointIndexOffset,
+			0
 		);
 
-		TManagedArray<FVector3f>& Positions = PointFacade.ModifyPositions();
-		TManagedArray<float>& Scales = PointFacade.ModifyPointScales();
 		for (int32 i = 0; i < BranchHierarchy.Points.Num(); ++i)
 		{
-			Positions[i + FacadePointIndexOffset] = BranchHierarchy.Points[i];
-			Scales[i + FacadePointIndexOffset] = BranchHierarchy.PointsRadii[i] > 0 ? BranchHierarchy.PointsRadii[i] : 0;
+			PointPositionAttribute[i + AttributePointIndexOffset] = BranchHierarchy.Points[i];
+			PointScaleAttribute[i + AttributePointIndexOffset] = BranchHierarchy.PointsRadii[i] > 0 ? BranchHierarchy.PointsRadii[i] : 0;
 		}
 
-		FacadeBranchIndexOffset += BranchHierarchy.Branches.Num();
-		FacadePointIndexOffset += BranchHierarchy.Points.Num();
+		AttributeBranchIndexOffset += BranchHierarchy.Branches.Num();
+		AttributePointIndexOffset += BranchHierarchy.Points.Num();
 		PlantNumber += 1;
 	}
 }
 
-void PV::ImportHelper::ComputeLengthFromRoot(PV::Facades::FPointFacade& PointFacade, const PV::Facades::FBranchFacade& BranchFacade, const PV::Facades::FPlantFacade& PlantFacade)
+void PV::ImportHelper::ComputeBudNumbers(
+	FPointBudNumberAttributeView PointBudNumberAttribute,
+	FBranchSourceBudNumberAttributeView BranchSourceBudNumberAttribute,
+	FBranchPointsAttributeConstView BranchPointsAttribute
+)
 {
-	PV::FacadeHelper::WalkBranches(PlantFacade, BranchFacade, [&](int32 BranchIndex)->PV::FacadeHelper::EForEachResult
+	for (int32 PointIndex = 0; PointIndex < PointBudNumberAttribute.Num(); ++PointIndex)
+	{
+		PointBudNumberAttribute[PointIndex] = PointIndex + 1;
+	}
+
+	for (int32 BranchIndex = 0; BranchIndex < BranchSourceBudNumberAttribute.Num(); ++BranchIndex)
 	{
-		TManagedArray<float>& LengthFromRoots = PointFacade.ModifyLengthFromRoots();
-		TManagedArray<float>& LengthFromSeeds = PointFacade.ModifyLengthFromSeeds();
+		check(BranchPointsAttribute[BranchIndex].Num() > 1);
+		const int32 BranchRootPointIndex = BranchPointsAttribute[BranchIndex][0];
+		BranchSourceBudNumberAttribute[BranchIndex] = PointBudNumberAttribute[BranchRootPointIndex];
+	}
+}
+
+void PV::ImportHelper::ComputeNjordPixelIndex(
+	PV::FPointNjordPixelIndexAttributeView PointNjordPixelIndexAttribute,
+	PV::FBranchPointsAttributeConstView BranchPointsAttribute,
+	PV::FBranchParentNumberAttributeConstView BranchParentNumberAttribute,
+	PV::FBranchNumberAttributeConstView BranchNumberAttribute,
+	PV::FBranchChildrenAttributeConstView BranchChildrenAttribute
+)
+{
+	using namespace PV::AttributesHelper;
+
+	RecursiveWalkBranches<int32>(
+		BranchParentNumberAttribute,
+		BranchChildrenAttribute,
+		BranchNumberAttribute,
+		1, 
+		[&](int32 BranchIndex, int32& NjordPixelIndex)->EForEachResult
+		{
+			const bool bIsTrunk = GetBranchParentIndex(BranchParentNumberAttribute, BranchNumberAttribute, BranchIndex) == INDEX_NONE;
+			const TArray<int32>& BranchPoints = BranchPointsAttribute[BranchIndex];
+			for (int32 i = 0; i < BranchPoints.Num(); ++i)
+			{
+				if (bIsTrunk || i != 0)
+				{
+					PointNjordPixelIndexAttribute[BranchPoints[i]] = NjordPixelIndex;
+					NjordPixelIndex += 1;
+				}
+			}
+
+			return EForEachResult::Continue;
+		}
+	);
+}
 
-		const TArray<int32>& CurrentBranchPoints = BranchFacade.GetPoints(BranchIndex);
+void PV::ImportHelper::ComputeLengthFromRoot(
+	PV::FPointLengthFromRootAttributeView PointLengthFromRootAttribute,
+	PV::FPointLengthFromSeedAttributeView PointLengthFromSeedAttribute,
+	PV::FPointPositionAttributeConstView PointPositionAttribute,
+	PV::FBranchPointsAttributeConstView BranchPointsAttribute,
+	PV::FBranchParentNumberAttributeConstView BranchParentNumberAttribute,
+	PV::FBranchNumberAttributeConstView BranchNumberAttribute,
+	PV::FBranchChildrenAttributeConstView BranchChildrenAttribute
+)
+{
+	using namespace PV::AttributesHelper;
+
+	RecursiveWalkBranches(BranchParentNumberAttribute, BranchChildrenAttribute, BranchNumberAttribute, [&](int32 BranchIndex)->EForEachResult
+	{
+		const TArray<int32>& CurrentBranchPoints = BranchPointsAttribute[BranchIndex];
 		check(CurrentBranchPoints.Num() > 0);
 
 		const int32 FirstPointIndex = CurrentBranchPoints[0];
-		const int32 ParentBranchIndex = BranchFacade.GetParentBranchIndex(BranchIndex);
+		const int32 ParentBranchIndex = GetBranchParentIndex(BranchParentNumberAttribute, BranchNumberAttribute, BranchIndex);
 
 		if (ParentBranchIndex == INDEX_NONE)
 		{
-			LengthFromRoots[FirstPointIndex] = 0.f; // No parent, length to root is 0
-			LengthFromSeeds[FirstPointIndex] = 0.f;
+			PointLengthFromRootAttribute[FirstPointIndex] = 0.f; // No parent, length to root is 0
+			PointLengthFromSeedAttribute[FirstPointIndex] = 0.f;
 		}
 		else
 		{
-			const TArray<int32>& ParentBranchPoints = BranchFacade.GetPoints(ParentBranchIndex);
+			const TArray<int32>& ParentBranchPoints = BranchPointsAttribute[ParentBranchIndex];
 			check(ParentBranchPoints.Num() > 0);
 
-			const FVector3f& FirstPoint = PointFacade.GetPosition(FirstPointIndex);
+			const FVector3f& FirstPoint = PointPositionAttribute[FirstPointIndex];
 
 			int32 ClosestParentPointIndex = ParentBranchPoints[0];
-			float ClosestParentPointDistSquared = (PointFacade.GetPosition(ClosestParentPointIndex) - FirstPoint).SizeSquared();
+			float ClosestParentPointDistSquared = (PointPositionAttribute[ClosestParentPointIndex] - FirstPoint).SizeSquared();
 
 			for (int32 i = 1; i < ParentBranchPoints.Num(); ++i)
 			{
 				const int32 PointIndex = ParentBranchPoints[i];
-				const FVector3f ParentPoint = PointFacade.GetPosition(PointIndex);
+				const FVector3f ParentPoint = PointPositionAttribute[PointIndex];
 
 				const float DistSquared = (ParentPoint - FirstPoint).SizeSquared();
 				if (DistSquared > ClosestParentPointDistSquared)
@@ -984,13 +990,13 @@ void PV::ImportHelper::ComputeLengthFromRoot(PV::Facades::FPointFacade& PointFac
 					if (i < ParentBranchPoints.Num() - 1)
 					{
 						const int32 NextPointIndex = ParentBranchPoints[i + 1];
-						return (PointFacade.GetPosition(NextPointIndex) - PointFacade.GetPosition(PointIndex)).GetSafeNormal();
+						return (PointPositionAttribute[NextPointIndex] - PointPositionAttribute[PointIndex]).GetSafeNormal();
 					}
 
 					if (i > 0)
 					{
 						const int32 PrevPointIndex = ParentBranchPoints[i - 1];
-						return (PointFacade.GetPosition(PointIndex) - PointFacade.GetPosition(PrevPointIndex)).GetSafeNormal();
+						return (PointPositionAttribute[PointIndex] - PointPositionAttribute[PrevPointIndex]).GetSafeNormal();
 					}
 
 					return FVector3f::UpVector;
@@ -1004,8 +1010,8 @@ void PV::ImportHelper::ComputeLengthFromRoot(PV::Facades::FPointFacade& PointFac
 				}
 			}
 
-			LengthFromRoots[FirstPointIndex] = LengthFromRoots[ClosestParentPointIndex] + (ClosestParentPointDistSquared > 0 ? FMath::Sqrt(ClosestParentPointDistSquared) : 0);
-			LengthFromSeeds[FirstPointIndex] = LengthFromRoots[FirstPointIndex];
+			PointLengthFromRootAttribute[FirstPointIndex] = PointLengthFromRootAttribute[ClosestParentPointIndex] + (ClosestParentPointDistSquared > 0 ? FMath::Sqrt(ClosestParentPointDistSquared) : 0) * 0.01 /*cm->m*/;
+			PointLengthFromSeedAttribute[FirstPointIndex] = PointLengthFromRootAttribute[FirstPointIndex];
 		}
 
 		for (int32 i = 1; i < CurrentBranchPoints.Num(); ++i)
@@ -1013,96 +1019,99 @@ void PV::ImportHelper::ComputeLengthFromRoot(PV::Facades::FPointFacade& PointFac
 			const int32 PointIndex = CurrentBranchPoints[i];
 			const int32 PrevPointIndex = CurrentBranchPoints[i - 1];
 
-			const FVector3f& Position = PointFacade.GetPosition(PointIndex);
-			const FVector3f& PrevPosition = PointFacade.GetPosition(PrevPointIndex);
+			const FVector3f& Position = PointPositionAttribute[PointIndex];
+			const FVector3f& PrevPosition = PointPositionAttribute[PrevPointIndex];
 
-			LengthFromRoots[PointIndex] = LengthFromRoots[PrevPointIndex] + (Position - PrevPosition).Size();
-			LengthFromSeeds[PointIndex] = LengthFromRoots[PointIndex];
+			PointLengthFromRootAttribute[PointIndex] = PointLengthFromRootAttribute[PrevPointIndex] + (Position - PrevPosition).Size() * 0.01 /*cm->m*/;
+			PointLengthFromSeedAttribute[PointIndex] = PointLengthFromRootAttribute[PointIndex];
 		}
 
-		return PV::FacadeHelper::EForEachResult::Continue;
+		return EForEachResult::Continue;
 	});
 }
 
-void PV::ImportHelper::ComputeBudDevelopment(PV::Facades::FPointFacade& PointFacade, const PV::Facades::FBranchFacade& BranchFacade, const PV::Facades::FPlantFacade& PlantFacade)
+void PV::ImportHelper::ComputeBudDevelopment(
+	PV::FBudDevelopmentAttributeView BudDevelopmentAttribute,
+	PV::FPointPlantGradientAttributeConstView PointPlantGradientAttribute,
+	PV::FBranchPointsAttributeConstView BranchPointsAttribute,
+	PV::FBranchParentNumberAttributeConstView BranchParentNumberAttribute,
+	PV::FBranchParentsAttributeConstView BranchParentsAttribute,
+	PV::FBranchNumberAttributeConstView BranchNumberAttribute
+)
 {
-	/*
-		0_Generatation: Generation is assumed to be the number of parents in the parents array and does not account for codominant structures. 
-		1_Age: Assuming the hierarchy is already extracted, starting a loop at the seed point and adding a value of 1 to all the connected points will result in an attribute resembling the age. 
-		       In the case of PVE there is no connection between a child and its parents so other metrics need to be implemented E.g utilizing branchSourceBudNumber or getting all points within a given distance.
-		2_Branch_Age: Branch age is simply the age of the first point on each branch.
-		3_Age_Senescence: Is defaulted to zero due to a lack of data and zero being the safest assumption.
-		4_Light_Senescence: Same as Age senescence.
-		5_Relative_Age: Same as age.
-	*/
-
-	PV::FacadeHelper::WalkPlantPoints<int32, PV::FacadeHelper::EWalkBranchMode::FusedPoints>(
-		PlantFacade, 
-		BranchFacade, 
-		PointFacade,
-		0, 
-		[&](int32 BranchIndex, int32 BranchPointIndex, int32 PointIndex, const TArray<int32>& ChildPointIndices, int32& NumWalkedPoints)->PV::FacadeHelper::EForEachResult
-		{
-			NumWalkedPoints += 1;
-
-			const TArray<int32>& CurrentBranchPoints = BranchFacade.GetPoints(BranchIndex);
-			check(CurrentBranchPoints.Num() > 0);
-
-			const int32 FirstPointIndex = CurrentBranchPoints[0];
-			const int32 BranchAge = PointIndex == FirstPointIndex ? NumWalkedPoints : PointFacade.GetBudAge(FirstPointIndex);
-
-			TArray<int32> BudDevelopment = {
-				/*0_Generatation*/ BranchFacade.GetParents(BranchIndex).Num() + 1,
-				/*1_Age*/ NumWalkedPoints,
-				/*2_Branch_Age*/ BranchAge,
-				/*3_Age_Senescence*/ 0,
-				/*4_Light_Senescence*/ 0,
-				/*5_Relative_Age*/ NumWalkedPoints
-			};
-
-			PointFacade.SetBudDevelopment(PointIndex, MoveTemp(BudDevelopment));
-
-			return PV::FacadeHelper::EForEachResult::Continue;
+	using namespace PV::AttributesHelper;
+
+	int32 MaxBudAge = 0;
+	for (int32 BranchIndex = 0; BranchIndex < BranchPointsAttribute.Num(); ++BranchIndex)
+	{
+		MaxBudAge = FMath::Max(MaxBudAge, BranchPointsAttribute[BranchIndex].Num());
+	}
+
+	ForEachUniquePointOnBranches_FusedPoints(
+		BranchPointsAttribute,
+		BranchParentNumberAttribute,
+		BranchNumberAttribute,
+		[&](int32 BranchIndex, int32 PointIndex)
+		{
+			const float Gradient = PointPlantGradientAttribute[PointIndex];
+			const int32 BudAge = FMath::RoundToInt(FMath::GetMappedRangeValueClamped(FVector2f(0, 1), FVector2f(1, float(MaxBudAge)), Gradient));
+
+			BudDevelopmentAttribute[PointIndex].BudAge = BudAge;
+			BudDevelopmentAttribute[PointIndex].RelativeBudAge = BudAge;
+			BudDevelopmentAttribute[PointIndex].Generation = BranchParentsAttribute[BranchIndex].Num() + 1;
+
+			return EForEachResult::Continue;
+		}
+	);
+
+	ForEachUniquePointOnBranches_FusedPoints(
+		BranchPointsAttribute,
+		BranchParentNumberAttribute,
+		BranchNumberAttribute,
+		[&](int32 BranchIndex, int32 PointIndex)
+		{
+			const bool bIsRootBranch = GetBranchParentIndex(BranchParentNumberAttribute, BranchNumberAttribute, BranchIndex) == INDEX_NONE;
+			const int32 RootPointIndex = bIsRootBranch ? BranchPointsAttribute[BranchIndex][0] : BranchPointsAttribute[BranchIndex][1];
+			const int32 BranchAge = BudDevelopmentAttribute[RootPointIndex].BudAge;
+			BudDevelopmentAttribute[PointIndex].BranchAge = BranchAge;
+
+			return EForEachResult::Continue;
 		}
 	);
 }
 
-void PV::ImportHelper::ComputeBudDirections(PV::Facades::FPointFacade& PointFacade, const PV::Facades::FBranchFacade& BranchFacade, const PV::Facades::FPlantFacade& PlantFacade)
+void PV::ImportHelper::ComputeBudDirections(
+	PV::FBudDirectionAttributeView BudDirectionAttribute,
+	PV::FPointPositionAttributeConstView PointPositionAttribute,
+	PV::FBranchChildrenAttributeConstView BranchChildrenAttribute,
+	PV::FBranchParentNumberAttributeConstView BranchParentNumberAttribute,
+	PV::FBranchNumberAttributeConstView BranchNumberAttribute,
+	PV::FBranchPointsAttributeConstView BranchPointsAttribute
+)
 {
-	/*
-		0_Apical: Apical is generated by getting the direction from the previous point to the current point, unless its the last.
-		          The last point is the direction of the current point to the next.
-		1_Axillary: This is set by using a phyllotaxy looping over the points. If the current point in an iteration has a child, 
-		            the axillary direction will be set to the orthogonalized apical direction of the child and the phyllotaxy will reset at that point.
-		2_Light_Optimal: Is set to Apical
-		3_Light_Sub_Optimal: Is set to Apical.
-		4_Guide_Curve: Is set to Apical.
-		5 Up_Vector: This is created by operating over the prims in hierarchical order. For the root point of each plant we compute an orthogonal vector 
-		             between the ground plane and the Apical direction. From here we run a parallel transport of the vector throughout the branch. 
-					 For the direct children we compute the dihedral matrix between the parent apical and the child apical direction, transform the parent 
-					 up vector by this matrix and continue the parallel transport. This is repeated until the hierarchy is exhausted.
-	*/
-
-	PV::FacadeHelper::WalkPlantPoints<int32, PV::FacadeHelper::EWalkBranchMode::FusedPoints>(
-		PlantFacade,
-		BranchFacade,
-		PointFacade,
+	using namespace PV::AttributesHelper;
+
+	RecursiveWalkPlantPoints_FusedPoints<int32>(
+		BranchChildrenAttribute,
+		BranchParentNumberAttribute,
+		BranchNumberAttribute,
+		BranchPointsAttribute,
 		INDEX_NONE,
-		[&](int32 BranchIndex, int32 BranchPointIndex, int32 PointIndex, const TArray<int32>& ChildPointIndices, int32& PrevPointIndex)->PV::FacadeHelper::EForEachResult
+		[&](int32 BranchIndex, int32 BranchPointIndex, int32 PointIndex, const TArray<int32>& ChildPointIndices, int32& PrevPointIndex)->EForEachResult
 		{
-			const TArray<int32>& BranchPoints = BranchFacade.GetPoints(BranchIndex);
+			const TArray<int32>& BranchPoints = BranchPointsAttribute[BranchIndex];
 
 			const FVector3f Apical = [&]()->FVector3f
 			{
 				if (BranchPointIndex < BranchPoints.Num() - 1)
 				{
 					const int32 NextPointIndex = BranchPoints[BranchPointIndex + 1];
-					return (PointFacade.GetPosition(NextPointIndex) - PointFacade.GetPosition(PointIndex)).GetSafeNormal();
+					return (PointPositionAttribute[NextPointIndex] - PointPositionAttribute[PointIndex]).GetSafeNormal();
 				}
 				else if (BranchPointIndex > 0)
 				{
 					const int32 PrevPointIndex = BranchPoints[BranchPointIndex - 1];
-					return (PointFacade.GetPosition(PointIndex) - PointFacade.GetPosition(PrevPointIndex)).GetSafeNormal();
+					return (PointPositionAttribute[PointIndex] - PointPositionAttribute[PrevPointIndex]).GetSafeNormal();
 				}
 
 				ensure(false); // We should never be able to get here unless we have a branch with only 1 point
@@ -1110,99 +1119,343 @@ void PV::ImportHelper::ComputeBudDirections(PV::Facades::FPointFacade& PointFaca
 			}();
 
 			const FVector3f UpVector = FRotationMatrix44f::MakeFromZX(Apical, FVector3f::ForwardVector).ToQuat().GetRightVector() * -1;
-			const FVector3f PrevAxillary = PrevPointIndex == INDEX_NONE ? FVector3f::ForwardVector : PointFacade.GetBudDirection(PrevPointIndex)[1];
+			const FVector3f PrevAxillary = PrevPointIndex == INDEX_NONE ? FVector3f::ForwardVector : BudDirectionAttribute[PrevPointIndex].Axillary;
 			const FVector3f Axillary = FQuat4f(Apical, FMath::DegreesToRadians(135.f)).RotateVector(PrevAxillary);
 
-			TArray<FVector3f> BudDirections = {
-				/*0_Apical*/Apical,
-				/*1_Axillary*/Axillary,
-				/*2_Light_Optimal*/Apical,
-				/*3_Light_Sub_Optimal*/Apical,
-				/*4_Guide_Curve*/Apical,
-				/*Up_Vector*/UpVector
-			};
-
-			PointFacade.SetBudDirections(PointIndex, MoveTemp(BudDirections));
-
+			BudDirectionAttribute[PointIndex].Apical = Apical;
+			BudDirectionAttribute[PointIndex].Axillary = Axillary;
+			BudDirectionAttribute[PointIndex].LightOptimal = Apical;
+			BudDirectionAttribute[PointIndex].LightSubOptimal = Apical;
+			BudDirectionAttribute[PointIndex].GuideCurve = Apical;
+			BudDirectionAttribute[PointIndex].UpVector = UpVector;
+			
 			PrevPointIndex = PointIndex;
 
-			return PV::FacadeHelper::EForEachResult::Continue;
+			return EForEachResult::Continue;
 		}
 	);
 }
 
-void PV::ImportHelper::InterpolatePointsScales(PV::Facades::FPointFacade& PointFacade, const PV::Facades::FBranchFacade& BranchFacade, const PV::Facades::FPlantFacade& PlantFacade)
+void PV::ImportHelper::ComputeBudHormoneLevels(
+	PV::FBudHormoneLevelsAttributeView BudHormoneLevelsAttribute,
+	PV::FPointPlantGradientAttributeConstView PointPlantGradientAttribute
+)
 {
-	PV::FacadeHelper::WalkBranches(PlantFacade, BranchFacade, [&](int32 BranchIndex)->PV::FacadeHelper::EForEachResult
+	using namespace PV::AttributesHelper;
+		
+	for (int32 PointIndex = 0; PointIndex < BudHormoneLevelsAttribute.Num(); PointIndex++)
 	{
-		const TArray<int32>& BranchPoints = BranchFacade.GetPoints(BranchIndex);
+		const float Gradient = PointPlantGradientAttribute[PointIndex];
+		const float InverseGradient = 1.f - Gradient;
+
+		BudHormoneLevelsAttribute[PointIndex].Apical = FMath::GetMappedRangeValueClamped(FVector2f(0.75, 1.0), FVector2f(0, 1), InverseGradient);
+		BudHormoneLevelsAttribute[PointIndex].Axillary = 0;
+		BudHormoneLevelsAttribute[PointIndex].AxillaryInhibition = 0;
+		BudHormoneLevelsAttribute[PointIndex].Radical = 0;
+		BudHormoneLevelsAttribute[PointIndex].Ethylene = Gradient;
+		BudHormoneLevelsAttribute[PointIndex].Cykininin = 0;
+	}
+}
 
-		float CurrentPointScale = PointFacade.GetPointScale(BranchPoints.Last());
-		float CurrentLengthFromRoot = PointFacade.GetLengthFromRoot(BranchPoints.Last());
+void PV::ImportHelper::ComputeGradients(
+	PV::FPointGroundGradientAttributeView PointGroundGradientAttribute,
+	PV::FPointScaleGradientAttributeView PointScaleGradientAttribute,
+	PV::FPointPlantGradientAttributeView PointPlantGradientAttribute,
+	PV::FPointBranchGradientAttributeView PointBranchGradientAttribute,
+	PV::FPointLengthFromRootAttributeConstView PointLengthFromRootAttribute,
+	PV::FPointScaleAttributeConstView PointScaleAttribute,
+	PV::FBranchPointsAttributeConstView BranchPointsAttribute,
+	PV::FBranchParentNumberAttributeConstView BranchParentNumberAttribute,
+	PV::FBranchNumberAttributeConstView BranchNumberAttribute,
+	PV::FBranchChildrenAttributeConstView BranchChildrenAttribute
+)
+{
+	using namespace PV::AttributesHelper;
 
-		for (int32 i = BranchPoints.Num() - 1; i >= 0; --i)
+	const float* MaxLengthFromRoot = Algo::MaxElement(PointLengthFromRootAttribute);
+	const float* MaxPointScale = Algo::MaxElement(PointScaleAttribute);
+	if (MaxLengthFromRoot == nullptr || MaxPointScale == nullptr || *MaxLengthFromRoot == 0 || *MaxPointScale == 0)
+	{
+		return;
+	}
+
+	for (int32 PointIndex = 0; PointIndex < PointGroundGradientAttribute.Num(); PointIndex++)
+	{
+		const float LengthFromRoot = PointLengthFromRootAttribute[PointIndex];
+		const float LengthFromRootAlpha = LengthFromRoot / *MaxLengthFromRoot;
+		PointGroundGradientAttribute[PointIndex] = 1.0 - LengthFromRootAlpha;
+
+		const float PointScale = PointScaleAttribute[PointIndex];
+		const float PointScaleAlpha = PointScale / *MaxPointScale;
+		PointScaleGradientAttribute[PointIndex] = PointScaleAlpha;
+	}
+
+	// PlantGradient - A gradient from 1-0 where the tip of each branch is 0, and the base of each branch is whatever the gradient of the parent branch is at that point.
+	// BranchGradient - A gradient from 1-0 where 1 = base of the branch and 0 = tip of the branch
+	RecursiveWalkBranches(
+		BranchParentNumberAttribute,
+		BranchChildrenAttribute,
+		BranchNumberAttribute,
+		[&](int32 BranchIndex)->EForEachResult
 		{
-			float NextValidScale = -1;
-			float NextValidScale_LengthFromRoot = -1;
-			int32 NextValidScale_Index = INDEX_NONE;
+			const TArray<int32>& BranchPoints = BranchPointsAttribute[BranchIndex];
+			check(BranchPoints.Num() > 1);
+
+			const float BasePointPlantGradient = GetBranchParentIndex(BranchParentNumberAttribute, BranchNumberAttribute, BranchIndex) == INDEX_NONE
+				? 1.0
+				: PointPlantGradientAttribute[BranchPoints[0]];
 
-			for (int32 j = i - 1; j >= 0; --j)
+			for (int32 i = 0; i < BranchPoints.Num(); ++i)
 			{
-				const int32 PointIndex = BranchPoints[j];
-				const float PointScale = PointFacade.GetPointScale(PointIndex);
-				if (PointScale > 0)
-				{
-					NextValidScale = PointScale;
-					NextValidScale_LengthFromRoot = PointFacade.GetLengthFromRoot(PointIndex);
-					NextValidScale_Index = j;
-					break;
-				}
+				const float BranchAlpha = i / float(BranchPoints.Num() - 1);
+				const float PointPlantGradient = FMath::Lerp(BasePointPlantGradient, 0, BranchAlpha);
+				PointPlantGradientAttribute[BranchPoints[i]] = PointPlantGradient;
+				PointBranchGradientAttribute[BranchPoints[i]] = 1.0 - BranchAlpha;
 			}
 
-			if (NextValidScale_Index == INDEX_NONE)
+			return EForEachResult::Continue;
+		}
+	);
+}
+
+void PV::ImportHelper::ComputeBudStatus(
+	PV::FBudStatusAttributeView BudStatusAttribute,
+	PV::FBranchPointsAttributeConstView BranchPointsAttribute,
+	PV::FBranchParentNumberAttributeConstView BranchParentNumberAttribute,
+	PV::FBranchNumberAttributeConstView BranchNumberAttribute,
+	PV::FBranchChildrenAttributeConstView BranchChildrenAttribute
+)
+{
+	using namespace PV::AttributesHelper;
+
+	ForEachUniquePointOnBranches_FusedPoints(
+		BranchPointsAttribute,
+		BranchParentNumberAttribute,
+		BranchNumberAttribute,
+		[&](int32 BranchIndex, int32 PointIndex)->EForEachResult
+		{
+			PV::FBudStatusView BudStatusView = BudStatusAttribute[PointIndex];
+			BudStatusView.ApicalMeristem = 0;
+			BudStatusView.Codominant = 0;
+			BudStatusView.Axillary = 0;
+			BudStatusView.Seed = 0;
+			BudStatusView.Dormant = 0;
+			BudStatusView.Triggered = 1;
+			BudStatusView.NumTriggered = 1;
+			BudStatusView.Inactive = 1;
+			BudStatusView.BrokenTip = 0;
+			BudStatusView.Broken = 0;
+
+			if (BranchPointsAttribute[BranchIndex].Last() == PointIndex) // IsTip
 			{
-				for (int32 j = i - 1; j >= 0; --j)
+				BudStatusView.ApicalMeristem = 1;
+				BudStatusView.Dormant = 1;
+				BudStatusView.Triggered = 0;
+				BudStatusView.NumTriggered = 0;
+				BudStatusView.Inactive = 0;
+			}
+			else if (BranchPointsAttribute[BranchIndex][0] == PointIndex) // Is branch root point.
+				// Due to ForEachUniquePointOnBranches_FusedPoints only iterating shared points on the parent branch, we don't have to check whether this is the root branch
+				// (as the root branch is the only branch which will have the root point iterated on by ForEachUniquePointOnBranches_FusedPoints).
+			{
+				BudStatusView.Seed = 1;
+			}
+			else
+			{
+				const int32 NumChildren = GetBranchPointChildBranchIndices_FusedPoints(
+					BranchChildrenAttribute,
+					BranchParentNumberAttribute,
+					BranchNumberAttribute,
+					BranchPointsAttribute,
+					BranchIndex, 
+					PointIndex
+				).Num();
+
+				if (NumChildren > 0) 
 				{
-					PointFacade.SetPointScale(BranchPoints[j], CurrentPointScale);
+					BudStatusView.NumTriggered = NumChildren;
+					BudStatusView.Axillary = 1;
 				}
-				break;
 			}
 
-			if (NextValidScale_Index == i - 1)
+			return EForEachResult::Continue;
+		}
+	);
+}
+
+void PV::ImportHelper::ComputeBudLateralMeristem(
+	PV::FBudLateralMeristemAttributeView BudLateralMeristemAttribute,
+	PV::FPointPositionAttributeConstView PointPositionAttribute,
+	PV::FPointScaleAttributeConstView PointScaleAttribute,
+	PV::FPointLengthFromRootAttributeConstView PointLengthFromRootAttribute,
+	PV::FBudDirectionAttributeConstView BudDirectionAttribute,
+	PV::FBranchPointsAttributeConstView BranchPointsAttribute,
+	PV::FBranchParentNumberAttributeConstView BranchParentNumberAttribute,
+	PV::FBranchChildrenAttributeConstView BranchChildrenAttribute,
+	PV::FBranchNumberAttributeConstView BranchNumberAttribute
+)
+{
+	using namespace PV::AttributesHelper;
+
+	RecursiveWalkBranches_Reversed(
+		BranchParentNumberAttribute,
+		BranchChildrenAttribute,
+		BranchNumberAttribute,
+		[&](int32 BranchIndex)->EForEachResult
+		{
+			const TArray<int32>& BranchPointIndices = BranchPointsAttribute[BranchIndex];
+			check(BranchPointIndices.Num() > 1);
+
+			const FVector3f& FirstBranchPointPos = PointPositionAttribute[BranchPointIndices[0]];
+			const FVector3f& SecondBranchPointPos = PointPositionAttribute[BranchPointIndices[1]];
+
+			const FVector3f BranchDir = (FirstBranchPointPos - SecondBranchPointPos).GetSafeNormal();
+			const FVector3f FirstPointNormal = BudDirectionAttribute[BranchPointIndices[0]].Apical;;
+			
+			const float ParentDot = FVector3f::DotProduct(BranchDir, FirstPointNormal);
+
+			float PreviousPointDavinciValue = 0;
+			for (int32 i = BranchPointIndices.Num() - 1; i>= 0; --i)
 			{
-				CurrentPointScale = PointFacade.GetPointScale(BranchPoints[NextValidScale_Index]);
-				continue;
+				const int32 PointIndex = BranchPointIndices[i];
+				const float ExisitingDavinciValue = BudLateralMeristemAttribute[PointIndex].Davinci;
+				const float PointScale = PointScaleAttribute[PointIndex];
+				const float AccumulatedDavinciValue = PreviousPointDavinciValue + ExisitingDavinciValue + PointScale;
+				PreviousPointDavinciValue = AccumulatedDavinciValue;
+
+				BudLateralMeristemAttribute[PointIndex].LateralMeristem = 0.0;
+				BudLateralMeristemAttribute[PointIndex].Multiplier = 1;
+				BudLateralMeristemAttribute[PointIndex].Inactive = 0.0;
+				BudLateralMeristemAttribute[PointIndex].Davinci = AccumulatedDavinciValue;
+				BudLateralMeristemAttribute[PointIndex].ParentDot = ParentDot;
+				BudLateralMeristemAttribute[PointIndex].RootDistance = PointLengthFromRootAttribute[PointIndex];
+				BudLateralMeristemAttribute[PointIndex].Degredation = 0.0;
 			}
 
-			for (int32 j = i - 1; j >= NextValidScale_Index; --j)
+			return EForEachResult::Continue;
+		}
+	);
+}
+
+void PV::ImportHelper::ComputeBudLightDetected(PV::FBudLightDetectedAttributeView BudLightDetectedAttribute)
+{
+	for (int32 i = 0; i < BudLightDetectedAttribute.Num(); ++i)
+	{
+		BudLightDetectedAttribute[i].Availible = 1.0;
+		BudLightDetectedAttribute[i].Resource = 1.0;
+		BudLightDetectedAttribute[i].Branch = 1.0;
+		BudLightDetectedAttribute[i].Collision = 0.0;
+	}
+}
+
+void PV::ImportHelper::InterpolatePointsScales(
+	PV::FPointScaleAttributeView PointScaleAttribute,
+	PV::FPointLengthFromRootAttributeConstView PointLengthFromRootAttribute,
+	PV::FBranchPointsAttributeConstView BranchPointsAttribute,
+	PV::FBranchParentNumberAttributeConstView BranchParentNumberAttribute,
+	PV::FBranchChildrenAttributeConstView BranchChildrenAttribute,
+	PV::FBranchNumberAttributeConstView BranchNumberAttribute
+)
+{
+	using namespace PV::AttributesHelper;
+
+	RecursiveWalkBranches(
+		BranchParentNumberAttribute,
+		BranchChildrenAttribute,
+		BranchNumberAttribute,
+		[&](int32 BranchIndex)->EForEachResult
+		{
+			const TArray<int32>& BranchPoints = BranchPointsAttribute[BranchIndex];
+
+			float CurrentPointScale = PointScaleAttribute[BranchPoints.Last()];
+			float CurrentLengthFromRoot = PointLengthFromRootAttribute[BranchPoints.Last()];
+
+			for (int32 i = BranchPoints.Num() - 1; i >= 0; --i)
 			{
-				const int32 PointIndex = BranchPoints[j];
-				const float PointLengthFromRoot = PointFacade.GetLengthFromRoot(PointIndex);
-				const float Alpha = FMath::GetMappedRangeValueClamped(FVector2D(CurrentLengthFromRoot, NextValidScale_LengthFromRoot), FVector2D(0, 1), PointLengthFromRoot);
-				const float PointScale = FMath::Lerp(CurrentPointScale, NextValidScale, Alpha);
-				PointFacade.SetPointScale(PointIndex, PointScale);
+				float NextValidScale = -1;
+				float NextValidScale_LengthFromRoot = -1;
+				int32 NextValidScale_Index = INDEX_NONE;
+
+				for (int32 j = i - 1; j >= 0; --j)
+				{
+					const int32 PointIndex = BranchPoints[j];
+					const float PointScale = PointScaleAttribute[PointIndex];
+					if (PointScale > 0)
+					{
+						NextValidScale = PointScale;
+						NextValidScale_LengthFromRoot = PointLengthFromRootAttribute[PointIndex];
+						NextValidScale_Index = j;
+						break;
+					}
+				}
+
+				if (NextValidScale_Index == INDEX_NONE)
+				{
+					for (int32 j = i - 1; j >= 0; --j)
+					{
+						PointScaleAttribute[BranchPoints[j]] = CurrentPointScale;
+					}
+					break;
+				}
+
+				if (NextValidScale_Index == i - 1)
+				{
+					CurrentPointScale = PointScaleAttribute[BranchPoints[NextValidScale_Index]];
+					continue;
+				}
+
+				for (int32 j = i - 1; j >= NextValidScale_Index; --j)
+				{
+					const int32 PointIndex = BranchPoints[j];
+					const float PointLengthFromRoot = PointLengthFromRootAttribute[PointIndex];
+					const float Alpha = FMath::GetMappedRangeValueClamped(FVector2D(CurrentLengthFromRoot, NextValidScale_LengthFromRoot), FVector2D(0, 1), PointLengthFromRoot);
+					const float PointScale = FMath::Lerp(CurrentPointScale, NextValidScale, Alpha);
+					PointScaleAttribute[PointIndex] = PointScale;
+				}
+
+				CurrentPointScale = NextValidScale;
+				CurrentLengthFromRoot = NextValidScale_LengthFromRoot;
+				i = NextValidScale_Index + 1;
 			}
 
-			CurrentPointScale = NextValidScale;
-			CurrentLengthFromRoot = NextValidScale_LengthFromRoot;
-			i = NextValidScale_Index + 1;
+			return EForEachResult::Continue;
 		}
-
-		return PV::FacadeHelper::EForEachResult::Continue;
-	});
+	);
 }
 
 bool PV::ImportHelper::SplitFusedPoints(FManagedArrayCollection& Collection)
 {
-	PV::Facades::FBranchFacade BranchFacade(Collection);
-
 	if (!Collection.HasGroup(PV::GroupNames::PointGroup) || !Collection.HasGroup(PV::GroupNames::BranchGroup))
 	{
 		return false;
 	}
 
-	auto& BranchPointsAttribute = Collection.ModifyAttribute<TArray<int32>>(PV::AttributeNames::BranchPoints, PV::GroupNames::BranchGroup);
-	auto& BranchParentNumberAttribute = Collection.ModifyAttribute<int32>(PV::AttributeNames::BranchParentNumber, PV::GroupNames::BranchGroup);
+	auto BranchPointsAttribute = FBranchPointsAttribute::GetAttribute(Collection);
+	auto BranchParentNumberAttribute = FBranchParentNumberAttribute::GetAttribute(Collection);
+
+	auto PointPositionAttribute = FPointPositionAttribute::GetAttribute(Collection);
+	auto LengthFromRootAttribute = FPointLengthFromRootAttribute::GetAttribute(Collection);
+	auto LengthFromSeedAttribute = FPointLengthFromSeedAttribute::GetAttribute(Collection);
+	auto PointScaleGradientAttribute = FPointScaleGradientAttribute::GetAttribute(Collection);
+	auto PointHullGradientAttribute = FPointHullGradientAttribute::GetAttribute(Collection);
+	auto PointMainTrunkGradientAttribute = FPointMainTrunkGradientAttribute::GetAttribute(Collection);
+	auto GroundGradientAttribute = FPointGroundGradientAttribute::GetAttribute(Collection);
+	auto PointScaleAttribute = FPointScaleAttribute::GetAttribute(Collection);
+	auto BudDirectionAttribute = FBudDirectionAttribute::GetAttribute(Collection);
+	auto BudStatusAttribute = FBudStatusAttribute::GetAttribute(Collection);
+	auto BudLateralMeristemAttribute = FBudLateralMeristemAttribute::GetAttribute(Collection);
+	auto PointBranchGradientAttribute = FPointBranchGradientAttribute::GetAttribute(Collection);
+	auto BudHormoneLevelsAttribute = FBudHormoneLevelsAttribute::GetAttribute(Collection);
+	auto PlantGradientAttribute = FPointPlantGradientAttribute::GetAttribute(Collection);
+	auto BudNumberAttribute = FPointBudNumberAttribute::GetAttribute(Collection);
+	auto NjordPixelIndexAttribute = FPointNjordPixelIndexAttribute::GetAttribute(Collection);
+	auto BudLightDetectedAttribute = FBudLightDetectedAttribute::GetAttribute(Collection);
+	auto BudDevelopmentAttribute = FBudDevelopmentAttribute::GetAttribute(Collection);
+
+	const int32* MaxBudNumber = Algo::MaxElement(BudNumberAttribute);
+	check(MaxBudNumber != nullptr);
+
+	int32 NewBudNumber = *MaxBudNumber + 1;
 
 	for (int32 i = 0; i < BranchPointsAttribute.Num(); ++i)
 	{
@@ -1212,23 +1465,6 @@ bool PV::ImportHelper::SplitFusedPoints(FManagedArrayCollection& Collection)
 			continue;
 		}
 
-		auto& PointPositionAttribute = Collection.ModifyAttribute<FVector3f>(PV::AttributeNames::PointPosition, PV::GroupNames::PointGroup);
-		auto& LengthFromRootAttribute = Collection.ModifyAttribute<float>(PV::AttributeNames::LengthFromRoot, PV::GroupNames::PointGroup);
-		auto& LengthFromSeedAttribute = Collection.ModifyAttribute<float>(PV::AttributeNames::LengthFromSeed, PV::GroupNames::PointGroup);
-		auto& PointScaleGradientAttribute = Collection.ModifyAttribute<float>(PV::AttributeNames::PointScaleGradient, PV::GroupNames::PointGroup);
-		auto& HullGradientAttribute = Collection.ModifyAttribute<float>(PV::AttributeNames::HullGradient, PV::GroupNames::PointGroup);
-		auto& MainTrunkGradientAttribute = Collection.ModifyAttribute<float>(PV::AttributeNames::MainTrunkGradient, PV::GroupNames::PointGroup);
-		auto& GroundGradientAttribute = Collection.ModifyAttribute<float>(PV::AttributeNames::GroundGradient, PV::GroupNames::PointGroup);
-		auto& PointScaleAttribute = Collection.ModifyAttribute<float>(PV::AttributeNames::PointScale, PV::GroupNames::PointGroup);
-		auto& BudDirectionAttribute = Collection.ModifyAttribute<TArray<FVector3f>>(PV::AttributeNames::BudDirection, PV::GroupNames::PointGroup);
-		auto& BranchGradientAttribute = Collection.ModifyAttribute<float>(PV::AttributeNames::BranchGradient, PV::GroupNames::PointGroup);
-		auto& BudHormoneLevelsAttribute = Collection.ModifyAttribute<TArray<float>>(PV::AttributeNames::BudHormoneLevels, PV::GroupNames::PointGroup);
-		auto& PlantGradientAttribute = Collection.ModifyAttribute<float>(PV::AttributeNames::PlantGradient, PV::GroupNames::PointGroup);
-		auto& BudNumberAttribute = Collection.ModifyAttribute<int32>(PV::AttributeNames::BudNumber, PV::GroupNames::PointGroup);
-		auto& NjordPixelIndexAttribute = Collection.ModifyAttribute<float>(PV::AttributeNames::NjordPixelIndex, PV::GroupNames::PointGroup);
-		auto& BudLightDetectedAttribute = Collection.ModifyAttribute<TArray<float>>(PV::AttributeNames::BudLightDetected, PV::GroupNames::PointGroup);
-		auto& BudDevelopmentAttribute = Collection.ModifyAttribute<TArray<int32>>(PV::AttributeNames::BudDevelopment, PV::GroupNames::PointGroup);
-
 		const int32 NumPoints = PointPositionAttribute.Num();
 		Collection.Resize(NumPoints + 1, PV::GroupNames::PointGroup);
 
@@ -1242,18 +1478,22 @@ bool PV::ImportHelper::SplitFusedPoints(FManagedArrayCollection& Collection)
 		LengthFromRootAttribute[NewPointIndex] = LengthFromRootAttribute[PrevPointIndex];
 		LengthFromSeedAttribute[NewPointIndex] = LengthFromSeedAttribute[PrevPointIndex];
 		PointScaleGradientAttribute[NewPointIndex] = PointScaleGradientAttribute[PrevPointIndex];
-		HullGradientAttribute[NewPointIndex] = HullGradientAttribute[PrevPointIndex];
-		MainTrunkGradientAttribute[NewPointIndex] = MainTrunkGradientAttribute[PrevPointIndex];
+		PointHullGradientAttribute[NewPointIndex] = PointHullGradientAttribute[PrevPointIndex];
+		PointMainTrunkGradientAttribute[NewPointIndex] = PointMainTrunkGradientAttribute[PrevPointIndex];
 		GroundGradientAttribute[NewPointIndex] = GroundGradientAttribute[PrevPointIndex];
 		PointScaleAttribute[NewPointIndex] = PointScaleAttribute[PrevPointIndex];
-		BudDirectionAttribute[NewPointIndex] = BudDirectionAttribute[PrevPointIndex];
-		BranchGradientAttribute[NewPointIndex] = BranchGradientAttribute[PrevPointIndex];
-		BudHormoneLevelsAttribute[NewPointIndex] = BudHormoneLevelsAttribute[PrevPointIndex];
+		BudDirectionAttribute[NewPointIndex].Array = BudDirectionAttribute[PrevPointIndex].Array;
+		BudStatusAttribute[NewPointIndex].Array = BudStatusAttribute[PrevPointIndex].Array;
+		BudLateralMeristemAttribute[NewPointIndex].Array = BudLateralMeristemAttribute[PrevPointIndex].Array;
+		PointBranchGradientAttribute[NewPointIndex] = PointBranchGradientAttribute[PrevPointIndex];
+		BudHormoneLevelsAttribute[NewPointIndex].Array = BudHormoneLevelsAttribute[PrevPointIndex].Array;
 		PlantGradientAttribute[NewPointIndex] = PlantGradientAttribute[PrevPointIndex];
-		BudNumberAttribute[NewPointIndex] = BudNumberAttribute[PrevPointIndex];
+		BudNumberAttribute[NewPointIndex] = NewBudNumber;
 		NjordPixelIndexAttribute[NewPointIndex] = NjordPixelIndexAttribute[PrevPointIndex];
-		BudLightDetectedAttribute[NewPointIndex] = BudLightDetectedAttribute[PrevPointIndex];
-		BudDevelopmentAttribute[NewPointIndex] = BudDevelopmentAttribute[PrevPointIndex];
+		BudLightDetectedAttribute[NewPointIndex].Array = BudLightDetectedAttribute[PrevPointIndex].Array;
+		BudDevelopmentAttribute[NewPointIndex].Array = BudDevelopmentAttribute[PrevPointIndex].Array;
+
+		NewBudNumber++;
 	}
 
 	return true;
@@ -1283,9 +1523,9 @@ void PV::ImportHelper::ConnectDisconnectedBranches(FPVBranchHierarchyDescription
 		}
 
 		check(Branch.PointIndices.Num() > 1);
-			
+		
 		const FPVBranchDescription& ParentBranch = InOutBranchHierarchy.Branches[Branch.ParentBranchIndex];
-			
+		
 		const int32 FirstPointIndex = Branch.PointIndices[0];
 		const FVector3f& FirstPointLocation = InOutBranchHierarchy.Points[FirstPointIndex];
 
@@ -1354,7 +1594,9 @@ void PV::ImportHelper::ConnectDisconnectedBranches(FPVBranchHierarchyDescription
 		int32 ClosestApicalIndex = ClosestPointIndex;
 		for (int32 i = 0; i < ParentBranch.PointIndices.Num(); ++i)
 		{
-			const float PercentageDiff = FMath::Abs(PointDistances[ClosestPointIndex] - PointDistances[i]) / (PointDistances[i] > 0 ? PointDistances[i] : 1.f);
+			const float PercentageDiff = FMath::IsNearlyZero(PointDistances[i])
+				? PointDistances[ClosestPointIndex]
+				: FMath::Abs(PointDistances[ClosestPointIndex] - PointDistances[i]) / PointDistances[i];
 
 			if (PercentageDiff < 0.5)
 			{
@@ -1559,7 +1801,7 @@ void PV::ImportHelper::FuseCollinearBranches(FPVBranchHierarchyDescription& InOu
 	}
 }
 
-bool PV::ImportHelper::GenerateGrowthDataFromBranchHierarchies(FManagedArrayCollection& OutCollection, TArrayView<const FPVBranchHierarchyDescription> InBranchHierarchies)
+bool PV::ImportHelper::GenerateGrowthDataFromBranchHierarchies(FManagedArrayCollection& OutCollection, TArrayView<const FPVBranchHierarchyDescription> InBranchHierarchies, bool bSplitPoints)
 {
 	if (InBranchHierarchies.Num() < 1)
 	{
@@ -1580,28 +1822,148 @@ bool PV::ImportHelper::GenerateGrowthDataFromBranchHierarchies(FManagedArrayColl
 		}
 	}
 
-
 	CreateEmptyGrowthData(OutCollection, NumBranchPoints, 0, NumBranches, 0);
-	FillGrowthDataFromBranchHierarchyDescription(OutCollection, InBranchHierarchies);
+	
+	auto PointPositionAttribute = FPointPositionAttribute::GetAttribute(OutCollection);
+	auto PointScaleAttribute = FPointScaleAttribute::GetAttribute(OutCollection);
+	auto BranchPointsAttribute = FBranchPointsAttribute::GetAttribute(OutCollection);
+	auto BranchParentNumberAttribute = FBranchParentNumberAttribute::GetAttribute(OutCollection);
+	auto BranchParentsAttribute = FBranchParentsAttribute::GetAttribute(OutCollection);
+	auto BranchNumberAttribute = FBranchNumberAttribute::GetAttribute(OutCollection);
+	auto BranchChildrenAttribute = FBranchChildrenAttribute::GetAttribute(OutCollection);
+	auto BranchHierarchyNumberAttribute = FBranchHierarchyNumberAttribute::GetAttribute(OutCollection);
+	auto BranchPlantNumberAttribute = FBranchPlantNumberAttribute::GetAttribute(OutCollection);
+	FillAttributesFromBranchHierarchy(
+		PointPositionAttribute,
+		PointScaleAttribute,
+		BranchPointsAttribute,
+		BranchParentsAttribute,
+		BranchChildrenAttribute,
+		BranchParentNumberAttribute,
+		BranchNumberAttribute,
+		BranchHierarchyNumberAttribute,
+		BranchPlantNumberAttribute,
+		InBranchHierarchies
+	);
 
-	PV::Facades::FPointFacade PointFacade(OutCollection);
-	const PV::Facades::FBranchFacade BranchFacade(OutCollection);
-	const PV::Facades::FPlantFacade PlantFacade(OutCollection);
+	auto PointBudNumberAttribute = FPointBudNumberAttribute::GetAttribute(OutCollection);
+	auto BranchSourceBudNumberAttribute = FBranchSourceBudNumberAttribute::GetAttribute(OutCollection);
+	ComputeBudNumbers(
+		PointBudNumberAttribute,
+		BranchSourceBudNumberAttribute, 
+		BranchPointsAttribute
+	);
 
-	ComputeLengthFromRoot(PointFacade, BranchFacade, PlantFacade);
-	ComputeBudDevelopment(PointFacade, BranchFacade, PlantFacade);
-	ComputeBudDirections(PointFacade, BranchFacade, PlantFacade);
-	InterpolatePointsScales(PointFacade, BranchFacade, PlantFacade);
+	auto PointNjordPixelIndexAttribute = FPointNjordPixelIndexAttribute::GetAttribute(OutCollection);
+	ComputeNjordPixelIndex(
+		PointNjordPixelIndexAttribute,
+		BranchPointsAttribute,
+		BranchParentNumberAttribute, 
+		BranchNumberAttribute,
+		BranchChildrenAttribute
+	);
+
+	auto PointLengthFromRootAttribute = FPointLengthFromRootAttribute::GetAttribute(OutCollection);
+	auto PointLengthFromSeedAttribute = FPointLengthFromSeedAttribute::GetAttribute(OutCollection);
+	ComputeLengthFromRoot(
+		PointLengthFromRootAttribute,
+		PointLengthFromSeedAttribute,
+		PointPositionAttribute,
+		BranchPointsAttribute,
+		BranchParentNumberAttribute,
+		BranchNumberAttribute,
+		BranchChildrenAttribute
+	);
+
+	auto BudDirectionAttribute = FBudDirectionAttribute::GetAttribute(OutCollection);
+	ComputeBudDirections(
+		BudDirectionAttribute,
+		PointPositionAttribute,
+		BranchChildrenAttribute,
+		BranchParentNumberAttribute,
+		BranchNumberAttribute,
+		BranchPointsAttribute
+	);
+
+	auto BudLightDetectedAttribute = FBudLightDetectedAttribute::GetAttribute(OutCollection);
+	ComputeBudLightDetected(BudLightDetectedAttribute);
+
+	InterpolatePointsScales(
+		PointScaleAttribute,
+		PointLengthFromRootAttribute,
+		BranchPointsAttribute,
+		BranchParentNumberAttribute,
+		BranchChildrenAttribute,
+		BranchNumberAttribute
+	);
+
+	auto BudLateralMeristemAttribute = FBudLateralMeristemAttribute::GetAttribute(OutCollection);
+	ComputeBudLateralMeristem(
+		BudLateralMeristemAttribute,
+		PointPositionAttribute,
+		PointScaleAttribute,
+		PointLengthFromRootAttribute,
+		BudDirectionAttribute,
+		BranchPointsAttribute,
+		BranchParentNumberAttribute,
+		BranchChildrenAttribute,
+		BranchNumberAttribute
+	);
 
-	SplitFusedPoints(OutCollection); // NOTE: Important that this is done last as ComputeBudDevelopment and ComputeBudDirections relies on fused points to traverse the tree
+	auto GroundGradientAttribute = PV::FPointGroundGradientAttribute::GetAttribute(OutCollection);
+	auto PointScaleGradientAttribute = PV::FPointScaleGradientAttribute::GetAttribute(OutCollection);
+	auto PointPlantGradientAttribute = PV::FPointPlantGradientAttribute::GetAttribute(OutCollection);
+	auto PointBranchGradientAttribute = PV::FPointBranchGradientAttribute::GetAttribute(OutCollection);
+	ComputeGradients(
+		GroundGradientAttribute,
+		PointScaleGradientAttribute,
+		PointPlantGradientAttribute,
+		PointBranchGradientAttribute,
+		PointLengthFromRootAttribute,
+		PointScaleAttribute,
+		BranchPointsAttribute,
+		BranchParentNumberAttribute,
+		BranchNumberAttribute,
+		BranchChildrenAttribute
+	);
+
+	auto BudDevelopmentAttribute = FBudDevelopmentAttribute::GetAttribute(OutCollection);
+	ComputeBudDevelopment(
+		BudDevelopmentAttribute,
+		PointPlantGradientAttribute,
+		BranchPointsAttribute,
+		BranchParentNumberAttribute,
+		BranchParentsAttribute,
+		BranchNumberAttribute
+	);
+
+	auto BudHormoneLevelsAttribute = FBudHormoneLevelsAttribute::GetAttribute(OutCollection);
+	ComputeBudHormoneLevels(
+		BudHormoneLevelsAttribute,
+		PointPlantGradientAttribute
+	);
+
+	auto BudStatusAttribute = FBudStatusAttribute::GetAttribute(OutCollection);
+	ComputeBudStatus(
+		BudStatusAttribute,
+		BranchPointsAttribute,
+		BranchParentNumberAttribute,
+		BranchNumberAttribute,
+		BranchChildrenAttribute
+	);
+
+	if (bSplitPoints)
+	{
+		SplitFusedPoints(OutCollection); // NOTE: Important that this is done last as the compute steps assumes fused points
+	}
 
 	return true;
 }
 
-bool PV::ImportHelper::GenerateGrowthDataFromBranchHierarchy(FManagedArrayCollection& OutCollection, const FPVBranchHierarchyDescription& InBranchHierarchy)
+bool PV::ImportHelper::GenerateGrowthDataFromBranchHierarchy(FManagedArrayCollection& OutCollection, const FPVBranchHierarchyDescription& InBranchHierarchy, bool bSplitPoints)
 {
 	TArrayView<const FPVBranchHierarchyDescription> BranchHierarchyArrayView(&InBranchHierarchy, 1);
-	return GenerateGrowthDataFromBranchHierarchies(OutCollection, BranchHierarchyArrayView);
+	return GenerateGrowthDataFromBranchHierarchies(OutCollection, BranchHierarchyArrayView, bSplitPoints);
 }
 
 void PV::ImportHelper::DistToBranch(const FVector3f& Point, const TArray<FVector3f>& Points, const TArray<int32>& BranchPointIndices, FVector3f& OutClosestPoint, float& OutDistSquared, int32& OutPointIndex)
diff --git a/Engine/Plugins/Experimental/ProceduralVegetationEditor/Source/ProceduralVegetationEditor/Private/Importer/PVImportHelpers.h b/Engine/Plugins/Experimental/ProceduralVegetationEditor/Source/ProceduralVegetationEditor/Private/Importer/PVImportHelpers.h
--- a/Engine/Plugins/Experimental/ProceduralVegetationEditor/Source/ProceduralVegetationEditor/Private/Importer/PVImportHelpers.h
+++ b/Engine/Plugins/Experimental/ProceduralVegetationEditor/Source/ProceduralVegetationEditor/Private/Importer/PVImportHelpers.h
@@ -3,6 +3,8 @@
 
 #include "CoreMinimal.h"
 
+#include "Facades/PVAttributes.h"
+
 namespace PV::Facades
 {
 	class FPointFacade;
@@ -117,30 +119,149 @@ namespace PV::DynamicMeshHelper
 
 namespace PV::ImportHelper
 {
+	// Creates a collection containing all attributes (of the supplied size) required for creating PV growth data.
 	void CreateEmptyGrowthData(FManagedArrayCollection& OutCollection, int32 NumBranchPoints, int32 NumFoliagePoints, int32 NumBranches, int32 NumFoliageMeshes);
 
-	void FillGrowthDataFromBranchHierarchyDescription(
-		FManagedArrayCollection& Collection, 
+	void FillAttributesFromBranchHierarchy(
+		FManagedArrayCollection& OutCollection,
 		TArrayView<const FPVBranchHierarchyDescription> InBranchHierarchies
 	);
 
-	void ComputeLengthFromRoot(PV::Facades::FPointFacade& PointFacade, const PV::Facades::FBranchFacade& BranchFacade, const PV::Facades::FPlantFacade& PlantFacade);
+	// Copies the branch hierarchies from the input BranchHierarchies to the attributes.
+	// Assumes the attributes are large enough to receive the hierarchies.
+	void FillAttributesFromBranchHierarchy(
+		PV::FPointPositionAttributeView PointPositionAttribute,
+		PV::FPointScaleAttributeView PointScaleAttribute,
+		PV::FBranchPointsAttributeView BranchPointsAttribute,
+		PV::FBranchParentsAttributeView BranchParentsAttribute,
+		PV::FBranchChildrenAttributeView BranchChildrenAttribute,
+		PV::FBranchParentNumberAttributeView BranchParentNumberAttribute,
+		PV::FBranchNumberAttributeView BranchNumberAttribute,
+		PV::FBranchHierarchyNumberAttributeView BranchHierarchyNumberAttribute,
+		PV::FBranchPlantNumberAttributeView BranchPlantNumberAttribute,
+		TArrayView<const FPVBranchHierarchyDescription> InBranchHierarchies
+	);
+
+	// Re-computes the point bud numbers and writes them to the BudNumber and BranchSourceBudNumber attributes.
+	void ComputeBudNumbers(
+		FPointBudNumberAttributeView PointBudNumberAttribute,
+		FBranchSourceBudNumberAttributeView BranchSourceBudNumberAttribute,
+		FBranchPointsAttributeConstView BranchPointsAttribute
+	);
+
+	// Re-computes all njord pixel indices and writes them to the NjordPixelIndexAttribute.
+	void ComputeNjordPixelIndex(
+		PV::FPointNjordPixelIndexAttributeView PointNjordPixelIndexAttribute,
+		PV::FBranchPointsAttributeConstView BranchPointsAttribute,
+		PV::FBranchParentNumberAttributeConstView BranchParentNumberAttribute,
+		PV::FBranchNumberAttributeConstView BranchNumberAttribute,
+		PV::FBranchChildrenAttributeConstView BranchChildrenAttribute
+	);
+
+	// Re-computes the length from root and length from seed for all points.
+	// NOTE: Assumes fused points
+	void ComputeLengthFromRoot(
+		PV::FPointLengthFromRootAttributeView PointLengthFromRootAttribute,
+		PV::FPointLengthFromSeedAttributeView PointLengthFromSeedAttribute,
+		PV::FPointPositionAttributeConstView PointPositionAttribute,
+		PV::FBranchPointsAttributeConstView BranchPointsAttribute,
+		PV::FBranchParentNumberAttributeConstView BranchParentNumberAttribute,
+		PV::FBranchNumberAttributeConstView BranchNumberAttribute,
+		PV::FBranchChildrenAttributeConstView BranchChildrenAttribute
+	);
+
+	// Re-computes bud development for all points using the existing branch hierarchy and plant gradient. 
+	// NOTE: Assumes fused points
+	void ComputeBudDevelopment(
+		PV::FBudDevelopmentAttributeView BudDevelopmentAttribute,
+		PV::FPointPlantGradientAttributeConstView PointPlantGradientAttribute,
+		PV::FBranchPointsAttributeConstView BranchPointsAttribute,
+		PV::FBranchParentNumberAttributeConstView BranchParentNumberAttribute,
+		PV::FBranchParentsAttributeConstView BranchParentsAttribute,
+		PV::FBranchNumberAttributeConstView BranchNumberAttribute
+	);
+
+	// Re-computes the bud directions for all points.
+	// NOTE: Assumes fused points
+	void ComputeBudDirections(
+		PV::FBudDirectionAttributeView BudDirectionAttribute,
+		PV::FPointPositionAttributeConstView PointPositionAttribute,
+		PV::FBranchChildrenAttributeConstView BranchChildrenAttribute,
+		PV::FBranchParentNumberAttributeConstView BranchParentNumberAttribute,
+		PV::FBranchNumberAttributeConstView BranchNumberAttribute,
+		PV::FBranchPointsAttributeConstView BranchPointsAttribute
+	);
+
+	// Re-computes the bud hormone levels using the plant gradient for all points.
+	void ComputeBudHormoneLevels(
+		PV::FBudHormoneLevelsAttributeView BudHormoneLevelsAttribute,
+		PV::FPointPlantGradientAttributeConstView PointPlantGradientAttribute
+	);
 
-	void ComputeBudDevelopment(PV::Facades::FPointFacade& PointFacade, const PV::Facades::FBranchFacade& BranchFacade, const PV::Facades::FPlantFacade& PlantFacade);
+	// Re-computes the point- ground gradient, scale gradient, plant gradient, and branch gradient.
+	void ComputeGradients(
+		PV::FPointGroundGradientAttributeView PointGroundGradientAttribute,
+		PV::FPointScaleGradientAttributeView PointScaleGradientAttribute,
+		PV::FPointPlantGradientAttributeView PointPlantGradientAttribute,
+		PV::FPointBranchGradientAttributeView PointBranchGradientAttribute,
+		PV::FPointLengthFromRootAttributeConstView PointLengthFromRootAttribute,
+		PV::FPointScaleAttributeConstView PointScaleAttribute,
+		PV::FBranchPointsAttributeConstView BranchPointsAttribute,
+		PV::FBranchParentNumberAttributeConstView BranchParentNumberAttribute,
+		PV::FBranchNumberAttributeConstView BranchNumberAttribute,
+		PV::FBranchChildrenAttributeConstView BranchChildrenAttribute
+	);
 
-	void ComputeBudDirections(PV::Facades::FPointFacade& PointFacade, const PV::Facades::FBranchFacade& BranchFacade, const PV::Facades::FPlantFacade& PlantFacade);
+	// Re-computes the bud status for all points.
+	// NOTE: Assumes fused points
+	void ComputeBudStatus(
+		PV::FBudStatusAttributeView BudStatusAttribute,
+		PV::FBranchPointsAttributeConstView BranchPointsAttribute,
+		PV::FBranchParentNumberAttributeConstView BranchParentNumberAttribute,
+		PV::FBranchNumberAttributeConstView BranchNumberAttribute,
+		PV::FBranchChildrenAttributeConstView BranchChildrenAttribute
+	);
 
-	void InterpolatePointsScales(PV::Facades::FPointFacade& PointFacade, const PV::Facades::FBranchFacade& BranchFacade, const PV::Facades::FPlantFacade& PlantFacade);
+	// Re-computes the bud lateral meristem for all points.
+	void ComputeBudLateralMeristem(
+		PV::FBudLateralMeristemAttributeView BudLateralMeristemAttribute,
+		PV::FPointPositionAttributeConstView PointPositionAttribute,
+		PV::FPointScaleAttributeConstView PointScaleAttribute,
+		PV::FPointLengthFromRootAttributeConstView PointLengthFromRootAttribute,
+		PV::FBudDirectionAttributeConstView BudDirectionAttribute,
+		PV::FBranchPointsAttributeConstView BranchPointsAttribute,
+		PV::FBranchParentNumberAttributeConstView BranchParentNumberAttribute,
+		PV::FBranchChildrenAttributeConstView BranchChildrenAttribute,
+		PV::FBranchNumberAttributeConstView BranchNumberAttribute
+	);
+
+	// Re-computes the bud light detected for all points.
+	void ComputeBudLightDetected(PV::FBudLightDetectedAttributeView BudLightDetectedAttribute);
+
+	// This function ensures all points have a valid point scale by walking the point hierarchy
+	// and finding points with missing point scales. It then interpolates the missing point scales
+	// using the point scale of the parent/child points.
+	void InterpolatePointsScales(
+		PV::FPointScaleAttributeView PointScaleAttribute,
+		PV::FPointLengthFromRootAttributeConstView PointLengthFromRootAttribute,
+		PV::FBranchPointsAttributeConstView BranchPointsAttribute,
+		PV::FBranchParentNumberAttributeConstView BranchParentNumberAttribute,
+		PV::FBranchChildrenAttributeConstView BranchChildrenAttribute,
+		PV::FBranchNumberAttributeConstView BranchNumberAttribute
+	);
 
+	// Converts a growth data collection with fused points to a split points collection.
+	// Will not validate that the input collection has split points.
+	// NOTE: Assumes input collection is valid growth data.
 	bool SplitFusedPoints(FManagedArrayCollection& Collection);
 
 	void ConnectDisconnectedBranches(FPVBranchHierarchyDescription& InOutBranchHierarchy);
 
 	void FuseCollinearBranches(FPVBranchHierarchyDescription& InOutBranchHierarchy);
 
-	bool GenerateGrowthDataFromBranchHierarchies(FManagedArrayCollection& OutCollection, TArrayView<const FPVBranchHierarchyDescription> InBranchHierarchies);
+	bool GenerateGrowthDataFromBranchHierarchies(FManagedArrayCollection& OutCollection, TArrayView<const FPVBranchHierarchyDescription> InBranchHierarchies, bool bSplitPoints = true);
 
-	bool GenerateGrowthDataFromBranchHierarchy(FManagedArrayCollection& OutCollection, const FPVBranchHierarchyDescription& InBranchHierarchy);
+	bool GenerateGrowthDataFromBranchHierarchy(FManagedArrayCollection& OutCollection, const FPVBranchHierarchyDescription& InBranchHierarchy, bool bSplitPoints = true);
 
 	void DistToBranch(const FVector3f& Point, const TArray<FVector3f>& Points, const TArray<int32>& BranchPointIndices, FVector3f& OutClosestPoint, float& OutDistSquared, int32& OutPointIndex);
 
diff --git a/Engine/Plugins/Experimental/ProceduralVegetationEditor/Source/ProceduralVegetationEditor/Private/Importer/PVImporter_Texture2D.cpp b/Engine/Plugins/Experimental/ProceduralVegetationEditor/Source/ProceduralVegetationEditor/Private/Importer/PVImporter_Texture2D.cpp
--- a/Engine/Plugins/Experimental/ProceduralVegetationEditor/Source/ProceduralVegetationEditor/Private/Importer/PVImporter_Texture2D.cpp
+++ b/Engine/Plugins/Experimental/ProceduralVegetationEditor/Source/ProceduralVegetationEditor/Private/Importer/PVImporter_Texture2D.cpp
@@ -7,7 +7,7 @@
 #include "IntVectorTypes.h"
 #include "PCGSettings.h"
 #include "GeometryCollection/ManagedArrayCollection.h"
-#include "Helpers/PVFacadeHelper.h"
+#include "Helpers/PVAttributesHelper.h"
 
 #include "Engine/Texture2D.h"
 #include "Polygon2.h"
@@ -259,9 +259,7 @@ namespace Internal
 		const FVector& InOffset = FVector::ZeroVector
 	)
 	{
-		const PV::Facades::FPlantFacade PlantFacade(InCollection);
-		const PV::Facades::FBranchFacade BranchFacade(InCollection);
-		PV::Facades::FPointFacade PointFacade(InCollection);
+		using namespace PV::AttributesHelper;
 
 		TArray<TTuple<FVector3f, FString, float>> Labels;
 		Labels.Reserve(PlantNumberToPlantIndices.Num());
@@ -273,17 +271,23 @@ namespace Internal
 		};
 		TArray<FPlantOffsetInfo> PlantsOffsetInfo;
 
-		PV::FacadeHelper::ForEachPlant(PlantFacade, [&](int32 PlantNumber)
+		PV::FPointPositionAttributeView PointPositionAttribute = PV::FPointPositionAttribute::FindAttribute(InCollection);
+		if (!ensure(PointPositionAttribute.IsValid()))
+		{
+			return;
+		}
+
+		ForEachPlant(InCollection, [&](int32 PlantNumber)
 		{
 			FPlantOffsetInfo& PlantOffsetInfo = PlantsOffsetInfo.AddDefaulted_GetRef();
 			PlantOffsetInfo.PlantNumber = PlantNumber;
 
-			PV::FacadeHelper::ForEachPlantPoint(PlantNumber, PlantFacade, BranchFacade, [&](int32 BranchIndex, int32 PointIndex)
+			ForEachPlantPoint(InCollection, PlantNumber, [&](int32 BranchIndex, int32 PointIndex)
 			{
-				PlantOffsetInfo.Bounds += PointFacade.GetPosition(PointIndex);
-				return PV::FacadeHelper::EForEachResult::Continue;
+				PlantOffsetInfo.Bounds += PointPositionAttribute[PointIndex];
+				return EForEachResult::Continue;
 			});
-			return PV::FacadeHelper::EForEachResult::Continue;
+			return EForEachResult::Continue;
 		});
 
 		FVector3f Offset = FVector3f(InOffset);
@@ -291,11 +295,11 @@ namespace Internal
 		{
 			Offset.Y -= PlantsOffsetInfo[i].Bounds.Max.Y + 10.f;
 
-			PV::FacadeHelper::ForEachPlantPoint(PlantsOffsetInfo[i].PlantNumber, PlantFacade, BranchFacade, [&](int32 BranchIndex, int32 PointIndex)
+			ForEachPlantPoint(InCollection, PlantsOffsetInfo[i].PlantNumber, [&](int32 BranchIndex, int32 PointIndex)
 			{
-				FVector3f& PointPosition = PointFacade.ModifyPositions()[PointIndex];
+				FVector3f& PointPosition = PointPositionAttribute[PointIndex];
 				PointPosition += Offset;
-				return PV::FacadeHelper::EForEachResult::Continue;
+				return EForEachResult::Continue;
 			});
 
 			const int32 PlantIndex = PlantNumberToPlantIndices.IsValidIndex(PlantsOffsetInfo[i].PlantNumber)
diff --git a/Engine/Plugins/Experimental/ProceduralVegetationEditor/Source/ProceduralVegetationEditor/Private/PVEditorInterfaceImpl.cpp b/Engine/Plugins/Experimental/ProceduralVegetationEditor/Source/ProceduralVegetationEditor/Private/PVEditorInterfaceImpl.cpp
--- a/Engine/Plugins/Experimental/ProceduralVegetationEditor/Source/ProceduralVegetationEditor/Private/PVEditorInterfaceImpl.cpp
+++ b/Engine/Plugins/Experimental/ProceduralVegetationEditor/Source/ProceduralVegetationEditor/Private/PVEditorInterfaceImpl.cpp
@@ -25,12 +25,12 @@ bool PV::FPVEditorInterfaceImpl::ImportGrowthDataFromTexture2D(const FPVImportTe
 	return ImportResult == PV::Texture2DImport::EImportResult::Success;
 }
 
-TSharedRef<FManagedArrayCollection> PV::FPVEditorInterfaceImpl::CreateMockTreeGrowthDataCollection() const
+TSharedRef<FManagedArrayCollection> PV::FPVEditorInterfaceImpl::CreateMockTreeGrowthDataCollection(bool bSplitPoints) const
 {
 	const FPVBranchHierarchyDescription BranchHierarchy = PV::ImportHelper::CreateBranchHierarchyFromPoints(PVMockTree::Points, PVMockTree::PointsRadii, PVMockTree::BranchPointIndices);
 
 	TSharedRef<FManagedArrayCollection> OutCollection = MakeShared<FManagedArrayCollection>();
-	const bool bResult = PV::ImportHelper::GenerateGrowthDataFromBranchHierarchy(*OutCollection, BranchHierarchy);
+	const bool bResult = PV::ImportHelper::GenerateGrowthDataFromBranchHierarchy(*OutCollection, BranchHierarchy, bSplitPoints);
 	check(bResult != false); // Mock tree is staticly defined to always be valid, so it should always return valid growth data
 
 	return OutCollection;
diff --git a/Engine/Plugins/Experimental/ProceduralVegetationEditor/Source/ProceduralVegetationEditor/Private/PVEditorInterfaceImpl.h b/Engine/Plugins/Experimental/ProceduralVegetationEditor/Source/ProceduralVegetationEditor/Private/PVEditorInterfaceImpl.h
--- a/Engine/Plugins/Experimental/ProceduralVegetationEditor/Source/ProceduralVegetationEditor/Private/PVEditorInterfaceImpl.h
+++ b/Engine/Plugins/Experimental/ProceduralVegetationEditor/Source/ProceduralVegetationEditor/Private/PVEditorInterfaceImpl.h
@@ -14,7 +14,7 @@ namespace PV
 	public:
 		bool ImportGrowthDataFromStaticMesh(const FPVImportStaticMeshParams& InParams, FManagedArrayCollection& OutCollection) override;
 		bool ImportGrowthDataFromTexture2D(const FPVImportTexture2DParams& InParams, const UMaterialInterface* TextureVisualizationMaterial, FPVImportTexture2DOutput& Output) override;
-		TSharedRef<FManagedArrayCollection> CreateMockTreeGrowthDataCollection() const override;
+		TSharedRef<FManagedArrayCollection> CreateMockTreeGrowthDataCollection(bool bSplitPoints) const override;
 	};
 };
 #endif
diff --git a/Engine/Plugins/Experimental/ProceduralVegetationEditor/Source/ProceduralVegetationEditor/Private/Tests/Importer/PVImportTests.cpp b/Engine/Plugins/Experimental/ProceduralVegetationEditor/Source/ProceduralVegetationEditor/Private/Tests/Importer/PVImportTests.cpp
--- a/Engine/Plugins/Experimental/ProceduralVegetationEditor/Source/ProceduralVegetationEditor/Private/Tests/Importer/PVImportTests.cpp
+++ b/Engine/Plugins/Experimental/ProceduralVegetationEditor/Source/ProceduralVegetationEditor/Private/Tests/Importer/PVImportTests.cpp
@@ -12,9 +12,6 @@
 #include "Importer/PVImportHelpers.h"
 
 #include "Facades/PVAttributesNames.h"
-#include "Facades/PVPointFacade.h"
-#include "Facades/PVBranchFacade.h"
-#include "Facades/PVPlantFacade.h"
 
 #define PV_IMPORTER_TEST(TestName) PV_SIMPLE_AUTOMATION_TEST(Importer, TestName)
 
@@ -26,7 +23,7 @@ namespace PVMockTree
 	}
 }
 
-PV_IMPORTER_TEST(FillGrowthDataFromBranchHierarchyDescription)
+PV_IMPORTER_TEST(FillAttributesFromBranchHierarchy)
 {
 	const FPVBranchHierarchyDescription BranchHierarchy = PVMockTree::CreateBranchHierarchyDescription();
 
@@ -38,48 +35,49 @@ PV_IMPORTER_TEST(FillGrowthDataFromBranchHierarchyDescription)
 	FManagedArrayCollection ExpectedCollection;
 	PV::ImportHelper::CreateEmptyGrowthData(ExpectedCollection, NumBranchPoints, NumFoliagePoints, NumBranches, NumFoliageMeshes);
 
-	PV::Facades::FPointFacade ExpectedPointFacade(ExpectedCollection);
-	PV::Facades::FBranchFacade ExpectedBranchFacade(ExpectedCollection);
-	
-	// Init ExpectedPointFacade
+	auto ExpectedPointPositionAttribute = PV::FPointPositionAttribute::GetAttribute(ExpectedCollection);
+	auto ExpectedBranchPointsAttribute = PV::FBranchPointsAttribute::GetAttribute(ExpectedCollection);
+	auto ExpectedBranchParentNumberAttribute = PV::FBranchParentNumberAttribute::GetAttribute(ExpectedCollection);
+	auto ExpectedBranchParentsAttribute = PV::FBranchParentsAttribute::GetAttribute(ExpectedCollection);
+	auto ExpectedBranchChildrenAttribute = PV::FBranchChildrenAttribute::GetAttribute(ExpectedCollection);
+
+	// Init Expected Point Attributes
 	{
-		TManagedArray<FVector3f>& ExpectedPositions = ExpectedPointFacade.ModifyPositions();
-
-		ExpectedPositions[0] = PVMockTree::Point0;
-		ExpectedPositions[1] = PVMockTree::Point1;
-		ExpectedPositions[2] = PVMockTree::Point2;
-		ExpectedPositions[3] = PVMockTree::Point3;
-		ExpectedPositions[4] = PVMockTree::Point4;
-		ExpectedPositions[5] = PVMockTree::Point5;
-		ExpectedPositions[6] = PVMockTree::Point6;
-		ExpectedPositions[7] = PVMockTree::Point7;
-		ExpectedPositions[8] = PVMockTree::Point8;
-		ExpectedPositions[9] = PVMockTree::Point9;
-		ExpectedPositions[10] = PVMockTree::Point10;
-		ExpectedPositions[11] = PVMockTree::Point11;
+		ExpectedPointPositionAttribute[0] = PVMockTree::Point0;
+		ExpectedPointPositionAttribute[1] = PVMockTree::Point1;
+		ExpectedPointPositionAttribute[2] = PVMockTree::Point2;
+		ExpectedPointPositionAttribute[3] = PVMockTree::Point3;
+		ExpectedPointPositionAttribute[4] = PVMockTree::Point4;
+		ExpectedPointPositionAttribute[5] = PVMockTree::Point5;
+		ExpectedPointPositionAttribute[6] = PVMockTree::Point6;
+		ExpectedPointPositionAttribute[7] = PVMockTree::Point7;
+		ExpectedPointPositionAttribute[8] = PVMockTree::Point8;
+		ExpectedPointPositionAttribute[9] = PVMockTree::Point9;
+		ExpectedPointPositionAttribute[10] = PVMockTree::Point10;
+		ExpectedPointPositionAttribute[11] = PVMockTree::Point11;
 	}
 
-	// Init ExpectedBranchFacade
+	// Init Expected Branch attributes
 	{
-		ExpectedBranchFacade.SetPoints(0, PVMockTree::Branch0_PointIndices);
-		ExpectedBranchFacade.SetPoints(1, PVMockTree::Branch1_PointIndices);
-		ExpectedBranchFacade.SetPoints(2, PVMockTree::Branch2_PointIndices);
-		ExpectedBranchFacade.SetPoints(3, PVMockTree::Branch3_PointIndices);
-
-		ExpectedBranchFacade.SetParentBranchNumber(0, 0);
-		ExpectedBranchFacade.SetParentBranchNumber(1, PVMockTree::Branch0_BranchNumber);
-		ExpectedBranchFacade.SetParentBranchNumber(2, PVMockTree::Branch1_BranchNumber);
-		ExpectedBranchFacade.SetParentBranchNumber(3, PVMockTree::Branch0_BranchNumber);
-
-		ExpectedBranchFacade.SetParents(0, { });
-		ExpectedBranchFacade.SetParents(1, { PVMockTree::Branch0_BranchNumber });
-		ExpectedBranchFacade.SetParents(2, { PVMockTree::Branch0_BranchNumber, PVMockTree::Branch1_BranchNumber });
-		ExpectedBranchFacade.SetParents(3, { PVMockTree::Branch0_BranchNumber });
-
-		ExpectedBranchFacade.SetChildren(0, { PVMockTree::Branch1_BranchNumber, PVMockTree::Branch3_BranchNumber });
-		ExpectedBranchFacade.SetChildren(1, { PVMockTree::Branch2_BranchNumber });
-		ExpectedBranchFacade.SetChildren(2, { });
-		ExpectedBranchFacade.SetChildren(3, { });
+		ExpectedBranchPointsAttribute[0] = PVMockTree::Branch0_PointIndices;
+		ExpectedBranchPointsAttribute[1] = PVMockTree::Branch1_PointIndices;
+		ExpectedBranchPointsAttribute[2] = PVMockTree::Branch2_PointIndices;
+		ExpectedBranchPointsAttribute[3] = PVMockTree::Branch3_PointIndices;
+
+		ExpectedBranchParentNumberAttribute[0] = 0;
+		ExpectedBranchParentNumberAttribute[1] = PVMockTree::Branch0_BranchNumber;
+		ExpectedBranchParentNumberAttribute[2] = PVMockTree::Branch1_BranchNumber;
+		ExpectedBranchParentNumberAttribute[3] = PVMockTree::Branch0_BranchNumber;
+
+		ExpectedBranchParentsAttribute[0] = { };
+		ExpectedBranchParentsAttribute[1] = { PVMockTree::Branch0_BranchNumber };
+		ExpectedBranchParentsAttribute[2] = { PVMockTree::Branch0_BranchNumber, PVMockTree::Branch1_BranchNumber };
+		ExpectedBranchParentsAttribute[3] = { PVMockTree::Branch0_BranchNumber };
+
+		ExpectedBranchChildrenAttribute[0] = { PVMockTree::Branch1_BranchNumber, PVMockTree::Branch3_BranchNumber };
+		ExpectedBranchChildrenAttribute[1] = { PVMockTree::Branch2_BranchNumber };
+		ExpectedBranchChildrenAttribute[2] = { };
+		ExpectedBranchChildrenAttribute[3] = { };
 	}
 
 	FManagedArrayCollection ActualCollection;
@@ -88,74 +86,90 @@ PV_IMPORTER_TEST(FillGrowthDataFromBranchHierarchyDescription)
 	UTEST_EQUAL("CorrectNrOfPointsInGroup", ActualCollection.NumElements(PV::GroupNames::PointGroup), NumBranchPoints);
 	UTEST_EQUAL("CorrectNrOfBranchesInGroup", ActualCollection.NumElements(PV::GroupNames::BranchGroup), NumBranches);
 
-	PV::ImportHelper::FillGrowthDataFromBranchHierarchyDescription(ActualCollection, { BranchHierarchy });
+	PV::ImportHelper::FillAttributesFromBranchHierarchy(ActualCollection, { BranchHierarchy });
 
-	const PV::Facades::FPointFacade ActualPointFacade(ActualCollection);
-	const PV::Facades::FBranchFacade ActualBranchFacade(ActualCollection);
+	const auto& ActualPointPositionAttribute = PV::FPointPositionAttribute::GetAttribute(ActualCollection);
+	const auto& ActualBranchPointsAttribute = PV::FBranchPointsAttribute::GetAttribute(ActualCollection);
+	const auto& ActualBranchParentNumberAttribute = PV::FBranchParentNumberAttribute::GetAttribute(ActualCollection);
+	const auto& ActualBranchParentsAttribute = PV::FBranchParentsAttribute::GetAttribute(ActualCollection);
+	const auto& ActualBranchChildrenAttribute = PV::FBranchChildrenAttribute::GetAttribute(ActualCollection);
 
 	// Validate that the points are correctly indexed and in the correct location
-	UTEST_TRUE("PointPosition0", ActualPointFacade.GetPosition(0).Equals(ExpectedPointFacade.GetPosition(0)));
-	UTEST_TRUE("PointPosition1", ActualPointFacade.GetPosition(1).Equals(ExpectedPointFacade.GetPosition(1)));
-	UTEST_TRUE("PointPosition2", ActualPointFacade.GetPosition(2).Equals(ExpectedPointFacade.GetPosition(2)));
-	UTEST_TRUE("PointPosition3", ActualPointFacade.GetPosition(3).Equals(ExpectedPointFacade.GetPosition(3)));
-	UTEST_TRUE("PointPosition4", ActualPointFacade.GetPosition(4).Equals(ExpectedPointFacade.GetPosition(4)));
-	UTEST_TRUE("PointPosition5", ActualPointFacade.GetPosition(5).Equals(ExpectedPointFacade.GetPosition(5)));
-	UTEST_TRUE("PointPosition6", ActualPointFacade.GetPosition(6).Equals(ExpectedPointFacade.GetPosition(6)));
-	UTEST_TRUE("PointPosition7", ActualPointFacade.GetPosition(7).Equals(ExpectedPointFacade.GetPosition(7)));
-	UTEST_TRUE("PointPosition8", ActualPointFacade.GetPosition(8).Equals(ExpectedPointFacade.GetPosition(8)));
-	UTEST_TRUE("PointPosition9", ActualPointFacade.GetPosition(9).Equals(ExpectedPointFacade.GetPosition(9)));
-	UTEST_TRUE("PointPosition10", ActualPointFacade.GetPosition(10).Equals(ExpectedPointFacade.GetPosition(10)));
-	UTEST_TRUE("PointPosition11", ActualPointFacade.GetPosition(11).Equals(ExpectedPointFacade.GetPosition(11)));
+	UTEST_TRUE("PointPosition0",  ActualPointPositionAttribute[0].Equals(ExpectedPointPositionAttribute[0]));
+	UTEST_TRUE("PointPosition1",  ActualPointPositionAttribute[1].Equals(ExpectedPointPositionAttribute[1]));
+	UTEST_TRUE("PointPosition2",  ActualPointPositionAttribute[2].Equals(ExpectedPointPositionAttribute[2]));
+	UTEST_TRUE("PointPosition3",  ActualPointPositionAttribute[3].Equals(ExpectedPointPositionAttribute[3]));
+	UTEST_TRUE("PointPosition4",  ActualPointPositionAttribute[4].Equals(ExpectedPointPositionAttribute[4]));
+	UTEST_TRUE("PointPosition5",  ActualPointPositionAttribute[5].Equals(ExpectedPointPositionAttribute[5]));
+	UTEST_TRUE("PointPosition6",  ActualPointPositionAttribute[6].Equals(ExpectedPointPositionAttribute[6]));
+	UTEST_TRUE("PointPosition7",  ActualPointPositionAttribute[7].Equals(ExpectedPointPositionAttribute[7]));
+	UTEST_TRUE("PointPosition8",  ActualPointPositionAttribute[8].Equals(ExpectedPointPositionAttribute[8]));
+	UTEST_TRUE("PointPosition9",  ActualPointPositionAttribute[9].Equals(ExpectedPointPositionAttribute[9]));
+	UTEST_TRUE("PointPosition10", ActualPointPositionAttribute[10].Equals(ExpectedPointPositionAttribute[10]));
+	UTEST_TRUE("PointPosition11", ActualPointPositionAttribute[11].Equals(ExpectedPointPositionAttribute[11]));
 
 	// Validate that the branch point indices are correct
-	UTEST_EQUAL("BranchPoints0", ActualBranchFacade.GetPoints(0), ExpectedBranchFacade.GetPoints(0));
-	UTEST_EQUAL("BranchPoints1", ActualBranchFacade.GetPoints(1), ExpectedBranchFacade.GetPoints(1));
-	UTEST_EQUAL("BranchPoints2", ActualBranchFacade.GetPoints(2), ExpectedBranchFacade.GetPoints(2));
-	UTEST_EQUAL("BranchPoints3", ActualBranchFacade.GetPoints(3), ExpectedBranchFacade.GetPoints(3));
+	UTEST_EQUAL("BranchPoints0", ActualBranchPointsAttribute[0], ExpectedBranchPointsAttribute[0]);
+	UTEST_EQUAL("BranchPoints1", ActualBranchPointsAttribute[1], ExpectedBranchPointsAttribute[1]);
+	UTEST_EQUAL("BranchPoints2", ActualBranchPointsAttribute[2], ExpectedBranchPointsAttribute[2]);
+	UTEST_EQUAL("BranchPoints3", ActualBranchPointsAttribute[3], ExpectedBranchPointsAttribute[3]);
 
 	// Validate that the branch parents are correct
-	UTEST_EQUAL("ParentBranchNumber0", ActualBranchFacade.GetParentBranchNumber(0), ExpectedBranchFacade.GetParentBranchNumber(0));
-	UTEST_EQUAL("ParentBranchNumber1", ActualBranchFacade.GetParentBranchNumber(1), ExpectedBranchFacade.GetParentBranchNumber(1));
-	UTEST_EQUAL("ParentBranchNumber2", ActualBranchFacade.GetParentBranchNumber(2), ExpectedBranchFacade.GetParentBranchNumber(2));
-	UTEST_EQUAL("ParentBranchNumber3", ActualBranchFacade.GetParentBranchNumber(3), ExpectedBranchFacade.GetParentBranchNumber(3));
-	UTEST_EQUAL("BranchParents0", ActualBranchFacade.GetParents(0), ExpectedBranchFacade.GetParents(0));
-	UTEST_EQUAL("BranchParents1", ActualBranchFacade.GetParents(1), ExpectedBranchFacade.GetParents(1));
-	UTEST_EQUAL("BranchParents2", ActualBranchFacade.GetParents(2), ExpectedBranchFacade.GetParents(2));
-	UTEST_EQUAL("BranchParents3", ActualBranchFacade.GetParents(3), ExpectedBranchFacade.GetParents(3));
+	UTEST_EQUAL("ParentBranchNumber0", ActualBranchParentNumberAttribute[0], ExpectedBranchParentNumberAttribute[0]);
+	UTEST_EQUAL("ParentBranchNumber1", ActualBranchParentNumberAttribute[1], ExpectedBranchParentNumberAttribute[1]);
+	UTEST_EQUAL("ParentBranchNumber2", ActualBranchParentNumberAttribute[2], ExpectedBranchParentNumberAttribute[2]);
+	UTEST_EQUAL("ParentBranchNumber3", ActualBranchParentNumberAttribute[3], ExpectedBranchParentNumberAttribute[3]);
+	UTEST_EQUAL("BranchParents0", ActualBranchParentsAttribute[0], ExpectedBranchParentsAttribute[0]);
+	UTEST_EQUAL("BranchParents1", ActualBranchParentsAttribute[1], ExpectedBranchParentsAttribute[1]);
+	UTEST_EQUAL("BranchParents2", ActualBranchParentsAttribute[2], ExpectedBranchParentsAttribute[2]);
+	UTEST_EQUAL("BranchParents3", ActualBranchParentsAttribute[3], ExpectedBranchParentsAttribute[3]);
+
+	// Validate branch children
+	UTEST_EQUAL("BranchChildren0", ActualBranchChildrenAttribute[0], ExpectedBranchChildrenAttribute[0]);
+	UTEST_EQUAL("BranchChildren1", ActualBranchChildrenAttribute[1], ExpectedBranchChildrenAttribute[1]);
+	UTEST_EQUAL("BranchChildren2", ActualBranchChildrenAttribute[2], ExpectedBranchChildrenAttribute[2]);
+	UTEST_EQUAL("BranchChildren3", ActualBranchChildrenAttribute[3], ExpectedBranchChildrenAttribute[3]);
 
 	return true;
 }
 
 PV_IMPORTER_TEST(ComputeLengthFromRoot)
 {
-	const FPVBranchHierarchyDescription BranchHierarchy = PVMockTree::CreateBranchHierarchyDescription();
-
-	const int32 NumBranchPoints = BranchHierarchy.GetNumBranchPoints();
-	const int32 NumFoliagePoints = 0;
-	const int32 NumBranches = BranchHierarchy.GetNumBranches();
-	const int32 NumFoliageMeshes = 0;
-
-	FManagedArrayCollection Collection;
-	PV::ImportHelper::CreateEmptyGrowthData(Collection, NumBranchPoints, NumFoliagePoints, NumBranches, NumFoliageMeshes);
-	PV::ImportHelper::FillGrowthDataFromBranchHierarchyDescription(Collection, { BranchHierarchy });
-
-	PV::Facades::FPointFacade PointFacade(Collection);
-	const PV::Facades::FBranchFacade BranchFacade(Collection);
-	const PV::Facades::FPlantFacade PlantFacade(Collection);
-	PV::ImportHelper::ComputeLengthFromRoot(PointFacade, BranchFacade, PlantFacade);
-
-	UTEST_NEARLY_EQUAL("Point0_LengthFromRoot", PointFacade.GetLengthFromRoot(0), PVMockTree::Point0_LengthFromRoot, 0.01f);
-	UTEST_NEARLY_EQUAL("Point1_LengthFromRoot", PointFacade.GetLengthFromRoot(1), PVMockTree::Point1_LengthFromRoot, 0.01f);
-	UTEST_NEARLY_EQUAL("Point2_LengthFromRoot", PointFacade.GetLengthFromRoot(2), PVMockTree::Point2_LengthFromRoot, 0.01f);
-	UTEST_NEARLY_EQUAL("Point3_LengthFromRoot", PointFacade.GetLengthFromRoot(3), PVMockTree::Point3_LengthFromRoot, 0.01f);
-	UTEST_NEARLY_EQUAL("Point4_LengthFromRoot", PointFacade.GetLengthFromRoot(4), PVMockTree::Point4_LengthFromRoot, 0.01f);
-	UTEST_NEARLY_EQUAL("Point5_LengthFromRoot", PointFacade.GetLengthFromRoot(5), PVMockTree::Point5_LengthFromRoot, 0.01f);
-	UTEST_NEARLY_EQUAL("Point6_LengthFromRoot", PointFacade.GetLengthFromRoot(6), PVMockTree::Point6_LengthFromRoot, 0.01f);
-	UTEST_NEARLY_EQUAL("Point7_LengthFromRoot", PointFacade.GetLengthFromRoot(7), PVMockTree::Point7_LengthFromRoot, 0.01f);
-	UTEST_NEARLY_EQUAL("Point8_LengthFromRoot", PointFacade.GetLengthFromRoot(8), PVMockTree::Point8_LengthFromRoot, 0.01f);
-	UTEST_NEARLY_EQUAL("Point9_LengthFromRoot", PointFacade.GetLengthFromRoot(9), PVMockTree::Point9_LengthFromRoot, 0.01f);
-	UTEST_NEARLY_EQUAL("Point10_LengthFromRoot", PointFacade.GetLengthFromRoot(10), PVMockTree::Point10_LengthFromRoot, 0.01f);
-	UTEST_NEARLY_EQUAL("Point11_LengthFromRoot", PointFacade.GetLengthFromRoot(11), PVMockTree::Point11_LengthFromRoot, 0.01f);
+	const TSharedRef<FManagedArrayCollection> Collection = PVMockTreeCollection::CreateCollection();
+	
+	PV::FPointLengthFromRootAttributeConstView LengthFromRootAttribute = PV::FPointLengthFromRootAttribute::FindAttribute(*Collection);
+	UTEST_TRUE("HasAttribute", LengthFromRootAttribute.IsValid());
+	UTEST_TRUE("Has 12 elements", LengthFromRootAttribute.Num() >= 12);
+
+	UTEST_NEARLY_EQUAL("Point0_LengthFromRoot",  LengthFromRootAttribute[0],  PVMockTree::Point0_LengthFromRoot, 0.01f);
+	UTEST_NEARLY_EQUAL("Point1_LengthFromRoot",  LengthFromRootAttribute[1],  PVMockTree::Point1_LengthFromRoot, 0.01f);
+	UTEST_NEARLY_EQUAL("Point2_LengthFromRoot",  LengthFromRootAttribute[2],  PVMockTree::Point2_LengthFromRoot, 0.01f);
+	UTEST_NEARLY_EQUAL("Point3_LengthFromRoot",  LengthFromRootAttribute[3],  PVMockTree::Point3_LengthFromRoot, 0.01f);
+	UTEST_NEARLY_EQUAL("Point4_LengthFromRoot",  LengthFromRootAttribute[4],  PVMockTree::Point4_LengthFromRoot, 0.01f);
+	UTEST_NEARLY_EQUAL("Point5_LengthFromRoot",  LengthFromRootAttribute[5],  PVMockTree::Point5_LengthFromRoot, 0.01f);
+	UTEST_NEARLY_EQUAL("Point6_LengthFromRoot",  LengthFromRootAttribute[6],  PVMockTree::Point6_LengthFromRoot, 0.01f);
+	UTEST_NEARLY_EQUAL("Point7_LengthFromRoot",  LengthFromRootAttribute[7],  PVMockTree::Point7_LengthFromRoot, 0.01f);
+	UTEST_NEARLY_EQUAL("Point8_LengthFromRoot",  LengthFromRootAttribute[8],  PVMockTree::Point8_LengthFromRoot, 0.01f);
+	UTEST_NEARLY_EQUAL("Point9_LengthFromRoot",  LengthFromRootAttribute[9],  PVMockTree::Point9_LengthFromRoot, 0.01f);
+	UTEST_NEARLY_EQUAL("Point10_LengthFromRoot", LengthFromRootAttribute[10], PVMockTree::Point10_LengthFromRoot, 0.01f);
+	UTEST_NEARLY_EQUAL("Point11_LengthFromRoot", LengthFromRootAttribute[11], PVMockTree::Point11_LengthFromRoot, 0.01f);
+
+	const PV::FPointLengthFromSeedAttributeView LengthFromSeedAttribute = PV::FPointLengthFromSeedAttribute::FindAttribute(*Collection);
+	UTEST_TRUE("HasAttribute", LengthFromSeedAttribute.IsValid());
+	UTEST_TRUE("Has 12 elements", LengthFromSeedAttribute.Num() >= 12);
+
+	UTEST_NEARLY_EQUAL("Point0_LengthFromRoot",  LengthFromSeedAttribute[0],  PVMockTree::Point0_LengthFromRoot,  0.01f);
+	UTEST_NEARLY_EQUAL("Point1_LengthFromRoot",  LengthFromSeedAttribute[1],  PVMockTree::Point1_LengthFromRoot,  0.01f);
+	UTEST_NEARLY_EQUAL("Point2_LengthFromRoot",  LengthFromSeedAttribute[2],  PVMockTree::Point2_LengthFromRoot,  0.01f);
+	UTEST_NEARLY_EQUAL("Point3_LengthFromRoot",  LengthFromSeedAttribute[3],  PVMockTree::Point3_LengthFromRoot,  0.01f);
+	UTEST_NEARLY_EQUAL("Point4_LengthFromRoot",  LengthFromSeedAttribute[4],  PVMockTree::Point4_LengthFromRoot,  0.01f);
+	UTEST_NEARLY_EQUAL("Point5_LengthFromRoot",  LengthFromSeedAttribute[5],  PVMockTree::Point5_LengthFromRoot,  0.01f);
+	UTEST_NEARLY_EQUAL("Point6_LengthFromRoot",  LengthFromSeedAttribute[6],  PVMockTree::Point6_LengthFromRoot,  0.01f);
+	UTEST_NEARLY_EQUAL("Point7_LengthFromRoot",  LengthFromSeedAttribute[7],  PVMockTree::Point7_LengthFromRoot,  0.01f);
+	UTEST_NEARLY_EQUAL("Point8_LengthFromRoot",  LengthFromSeedAttribute[8],  PVMockTree::Point8_LengthFromRoot,  0.01f);
+	UTEST_NEARLY_EQUAL("Point9_LengthFromRoot",  LengthFromSeedAttribute[9],  PVMockTree::Point9_LengthFromRoot,  0.01f);
+	UTEST_NEARLY_EQUAL("Point10_LengthFromRoot", LengthFromSeedAttribute[10], PVMockTree::Point10_LengthFromRoot, 0.01f);
+	UTEST_NEARLY_EQUAL("Point11_LengthFromRoot", LengthFromSeedAttribute[11], PVMockTree::Point11_LengthFromRoot, 0.01f);
 
 	return true;
 }
@@ -164,190 +178,536 @@ PV_IMPORTER_TEST(ComputeBudDevelopment)
 {
 	struct FBudDevelopmentDataLayout
 	{
-		int32 Generatation;
+		int32 Generation;
 		int32 Age;
 		int32 Branch_Age;
 		int32 AgeSenescence;
 		int32 Light_Senescence;
 		int32 Relative_Age;
+
+		bool operator==(const TArray<int32>& Other) const
+		{
+			return Other.Num() == 6
+				&& Other[0] == Generation
+				&& Other[1] == Age
+				&& Other[2] == Branch_Age
+				&& Other[3] == AgeSenescence
+				&& Other[4] == Light_Senescence
+				&& Other[5] == Relative_Age;
+		}
 	};
 
 	const FBudDevelopmentDataLayout Point0_ExpectedBudDevelopment = {
-		.Generatation = 1, 
-		.Age = 1,
-		.Branch_Age = 1,
+		.Generation = 1,
+		.Age = 6,
+		.Branch_Age = 6,
 		.AgeSenescence = 0,
 		.Light_Senescence = 0,
-		.Relative_Age = 1
+		.Relative_Age = 6
 	};
 	const FBudDevelopmentDataLayout Point1_ExpectedBudDevelopment = {
-		.Generatation = 1,
-		.Age = 2,
-		.Branch_Age = 1,
+		.Generation = 1,
+		.Age = 5,
+		.Branch_Age = 6,
 		.AgeSenescence = 0,
 		.Light_Senescence = 0,
-		.Relative_Age = 2
+		.Relative_Age = 5
 	};
 	const FBudDevelopmentDataLayout Point2_ExpectedBudDevelopment = {
-		.Generatation = 1,
-		.Age = 3,
-		.Branch_Age = 1,
+		.Generation = 1,
+		.Age = 4,
+		.Branch_Age = 6,
 		.AgeSenescence = 0,
 		.Light_Senescence = 0,
-		.Relative_Age = 3
+		.Relative_Age = 4
 	};
 	const FBudDevelopmentDataLayout Point3_ExpectedBudDevelopment = {
-		.Generatation = 1,
-		.Age = 4,
-		.Branch_Age = 1,
+		.Generation = 1,
+		.Age = 3,
+		.Branch_Age = 6,
 		.AgeSenescence = 0,
 		.Light_Senescence = 0,
-		.Relative_Age = 4
+		.Relative_Age = 3
 	};
 	const FBudDevelopmentDataLayout Point4_ExpectedBudDevelopment = {
-		.Generatation = 1,
-		.Age = 5,
-		.Branch_Age = 1,
+		.Generation = 1,
+		.Age = 2,
+		.Branch_Age = 6,
 		.AgeSenescence = 0,
 		.Light_Senescence = 0,
-		.Relative_Age = 5
+		.Relative_Age = 2
 	};
 	const FBudDevelopmentDataLayout Point5_ExpectedBudDevelopment = {
-		.Generatation = 1,
-		.Age = 6,
-		.Branch_Age = 1,
+		.Generation = 1,
+		.Age = 1,
+		.Branch_Age = 6,
 		.AgeSenescence = 0,
 		.Light_Senescence = 0,
-		.Relative_Age = 6
+		.Relative_Age = 1
 	};
 	const FBudDevelopmentDataLayout Point6_ExpectedBudDevelopment = {
-		.Generatation = 2,
-		.Age = 5,
-		.Branch_Age = 4,
+		.Generation = 2,
+		.Age = 1,
+		.Branch_Age = 1,
 		.AgeSenescence = 0,
 		.Light_Senescence = 0,
-		.Relative_Age = 5
+		.Relative_Age = 1
 	};
 	const FBudDevelopmentDataLayout Point7_ExpectedBudDevelopment = {
-		.Generatation = 2,
-		.Age = 4,
+		.Generation = 2,
+		.Age = 3,
 		.Branch_Age = 3,
 		.AgeSenescence = 0,
 		.Light_Senescence = 0,
-		.Relative_Age = 4
+		.Relative_Age = 3
 	};
 	const FBudDevelopmentDataLayout Point8_ExpectedBudDevelopment = {
-		.Generatation = 2,
-		.Age = 5,
+		.Generation = 2,
+		.Age = 2,
 		.Branch_Age = 3,
 		.AgeSenescence = 0,
 		.Light_Senescence = 0,
-		.Relative_Age = 5
+		.Relative_Age = 2
 	};
 	const FBudDevelopmentDataLayout Point9_ExpectedBudDevelopment = {
-		.Generatation = 2,
-		.Age = 6,
+		.Generation = 2,
+		.Age = 1,
 		.Branch_Age = 3,
 		.AgeSenescence = 0,
 		.Light_Senescence = 0,
-		.Relative_Age = 6
+		.Relative_Age = 1
 	};
 	const FBudDevelopmentDataLayout Point10_ExpectedBudDevelopment = {
-		.Generatation = 3,
-		.Age = 6,
-		.Branch_Age = 5,
+		.Generation = 3,
+		.Age = 2,
+		.Branch_Age = 2,
 		.AgeSenescence = 0,
 		.Light_Senescence = 0,
-		.Relative_Age = 6
+		.Relative_Age = 2
 	};
 	const FBudDevelopmentDataLayout Point11_ExpectedBudDevelopment = {
-		.Generatation = 3,
-		.Age = 7,
-		.Branch_Age = 5,
+		.Generation = 3,
+		.Age = 1,
+		.Branch_Age = 2,
 		.AgeSenescence = 0,
 		.Light_Senescence = 0,
-		.Relative_Age = 7
+		.Relative_Age = 1
 	};
 
-	const int32 ExpectedNumBudDevelopmentValues = 6;
+	const TSharedRef<FManagedArrayCollection> Collection = PVMockTreeCollection::CreateCollection();
+	const PV::FBudDevelopmentAttributeConstView BudDevelopmentAttribute = PV::FBudDevelopmentAttribute::FindAttribute(*Collection);
+	UTEST_TRUE("HasAttribute", BudDevelopmentAttribute);
+	UTEST_TRUE("Has 12 elements", BudDevelopmentAttribute.Num() >= 12);
+
+	UTEST_TRUE("Point0_BudDevelopment",  BudDevelopmentAttribute[0].Array  == Point0_ExpectedBudDevelopment);
+	UTEST_TRUE("Point1_BudDevelopment",  BudDevelopmentAttribute[1].Array  == Point1_ExpectedBudDevelopment);
+	UTEST_TRUE("Point2_BudDevelopment",  BudDevelopmentAttribute[2].Array  == Point2_ExpectedBudDevelopment);
+	UTEST_TRUE("Point3_BudDevelopment",  BudDevelopmentAttribute[3].Array  == Point3_ExpectedBudDevelopment);
+	UTEST_TRUE("Point4_BudDevelopment",  BudDevelopmentAttribute[4].Array  == Point4_ExpectedBudDevelopment);
+	UTEST_TRUE("Point5_BudDevelopment",  BudDevelopmentAttribute[5].Array  == Point5_ExpectedBudDevelopment);
+	UTEST_TRUE("Point6_BudDevelopment",  BudDevelopmentAttribute[6].Array  == Point6_ExpectedBudDevelopment);
+	UTEST_TRUE("Point7_BudDevelopment",  BudDevelopmentAttribute[7].Array  == Point7_ExpectedBudDevelopment);
+	UTEST_TRUE("Point8_BudDevelopment",  BudDevelopmentAttribute[8].Array  == Point8_ExpectedBudDevelopment);
+	UTEST_TRUE("Point9_BudDevelopment",  BudDevelopmentAttribute[9].Array  == Point9_ExpectedBudDevelopment);
+	UTEST_TRUE("Point10_BudDevelopment", BudDevelopmentAttribute[10].Array == Point10_ExpectedBudDevelopment);
+	UTEST_TRUE("Point11_BudDevelopment", BudDevelopmentAttribute[11].Array == Point11_ExpectedBudDevelopment);
 
-	const FPVBranchHierarchyDescription BranchHierarchy = PVMockTree::CreateBranchHierarchyDescription();
-	const int32 NumBranchPoints = BranchHierarchy.GetNumBranchPoints();
-	const int32 NumFoliagePoints = 0;
-	const int32 NumBranches = BranchHierarchy.GetNumBranches();
-	const int32 NumFoliageMeshes = 0;
+	return true;
+}
+
+PV_IMPORTER_TEST(ComputeGradients)
+{
+	const float Point0_ExpectedPlantGradient  = 1;
+	const float Point1_ExpectedPlantGradient  = 0.8;
+	const float Point2_ExpectedPlantGradient  = 0.6;
+	const float Point3_ExpectedPlantGradient  = 0.4;
+	const float Point4_ExpectedPlantGradient  = 0.2;
+	const float Point5_ExpectedPlantGradient  = 0;
+	const float Point6_ExpectedPlantGradient  = 0;
+	const float Point7_ExpectedPlantGradient  = FMath::Lerp(Point2_ExpectedPlantGradient, 0.f, 1.f / 3);
+	const float Point8_ExpectedPlantGradient  = FMath::Lerp(Point2_ExpectedPlantGradient, 0.f, 2.f / 3);
+	const float Point9_ExpectedPlantGradient  = 0;
+	const float Point10_ExpectedPlantGradient = FMath::Lerp(Point8_ExpectedPlantGradient, 0.f, 1.f / 2);
+	const float Point11_ExpectedPlantGradient = 0;
+
+	const TSharedRef<FManagedArrayCollection> Collection = PVMockTreeCollection::CreateCollection();
+	
+	const PV::FPointPlantGradientAttributeView PointPlantGradientAttribute = PV::FPointPlantGradientAttribute::FindAttribute(*Collection);
+	UTEST_TRUE("HasAttribute", PointPlantGradientAttribute.IsValid());
+	UTEST_TRUE("Has 12 elements", PointPlantGradientAttribute.Num() >= 12);
+
+	UTEST_TRUE("Point0_ExpectedPlantGradient",  FMath::IsNearlyEqual(PointPlantGradientAttribute[0],  Point0_ExpectedPlantGradient,  0.01f));
+	UTEST_TRUE("Point1_ExpectedPlantGradient",  FMath::IsNearlyEqual(PointPlantGradientAttribute[1],  Point1_ExpectedPlantGradient,  0.01f));
+	UTEST_TRUE("Point2_ExpectedPlantGradient",  FMath::IsNearlyEqual(PointPlantGradientAttribute[2],  Point2_ExpectedPlantGradient,  0.01f));
+	UTEST_TRUE("Point3_ExpectedPlantGradient",  FMath::IsNearlyEqual(PointPlantGradientAttribute[3],  Point3_ExpectedPlantGradient,  0.01f));
+	UTEST_TRUE("Point4_ExpectedPlantGradient",  FMath::IsNearlyEqual(PointPlantGradientAttribute[4],  Point4_ExpectedPlantGradient,  0.01f));
+	UTEST_TRUE("Point5_ExpectedPlantGradient",  FMath::IsNearlyEqual(PointPlantGradientAttribute[5],  Point5_ExpectedPlantGradient,  0.01f));
+	UTEST_TRUE("Point6_ExpectedPlantGradient",  FMath::IsNearlyEqual(PointPlantGradientAttribute[6],  Point6_ExpectedPlantGradient,  0.01f));
+	UTEST_TRUE("Point7_ExpectedPlantGradient",  FMath::IsNearlyEqual(PointPlantGradientAttribute[7],  Point7_ExpectedPlantGradient,  0.01f));
+	UTEST_TRUE("Point8_ExpectedPlantGradient",  FMath::IsNearlyEqual(PointPlantGradientAttribute[8],  Point8_ExpectedPlantGradient,  0.01f));
+	UTEST_TRUE("Point9_ExpectedPlantGradient",  FMath::IsNearlyEqual(PointPlantGradientAttribute[9],  Point9_ExpectedPlantGradient,  0.01f));
+	UTEST_TRUE("Point10_ExpectedPlantGradient", FMath::IsNearlyEqual(PointPlantGradientAttribute[10], Point10_ExpectedPlantGradient, 0.01f));
+	UTEST_TRUE("Point11_ExpectedPlantGradient", FMath::IsNearlyEqual(PointPlantGradientAttribute[11], Point11_ExpectedPlantGradient, 0.01f));
+
+	return true;
+}
+
+PV_IMPORTER_TEST(ComputeBudStatus)
+{
+	struct FBudStatusDataLayout {
+		int32 ApicalMeristem;
+		int32 Codominant;
+		int32 Axillary;
+		int32 Seed;
+		int32 Dormant;
+		int32 Triggered;
+		int32 NumTriggered;
+		int32 Inactive;
+		int32 BrokenTip;
+		int32 Broken;
+
+		bool operator==(const TArray<int32>& Other) const
+		{
+			return Other.Num() == 10
+				&& ApicalMeristem == Other[0]
+				&& Codominant == Other[1]
+				&& Axillary == Other[2]
+				&& Seed == Other[3]
+				&& Dormant == Other[4]
+				&& Triggered == Other[5]
+				&& NumTriggered == Other[6]
+				&& Inactive == Other[7]
+				&& BrokenTip == Other[8]
+				&& Broken == Other[9];
+		}
+	};
+
+	const FBudStatusDataLayout Point0_ExpectedBudStatus = {
+		.ApicalMeristem = 0,
+		.Codominant = 0,
+		.Axillary = 0,
+		.Seed = 1,
+		.Dormant = 0,
+		.Triggered = 1,
+		.NumTriggered = 1,
+		.Inactive = 1,
+		.BrokenTip = 0,
+		.Broken = 0,
+	};
+
+	const FBudStatusDataLayout Point1_ExpectedBudStatus = {
+		.ApicalMeristem = 0,
+		.Codominant = 0,
+		.Axillary = 0,
+		.Seed = 0,
+		.Dormant = 0,
+		.Triggered = 1,
+		.NumTriggered = 1,
+		.Inactive = 1,
+		.BrokenTip = 0,
+		.Broken = 0,
+	};
+
+	const FBudStatusDataLayout Point2_ExpectedBudStatus = {
+		.ApicalMeristem = 0,
+		.Codominant = 0,
+		.Axillary = 1,
+		.Seed = 0,
+		.Dormant = 0,
+		.Triggered = 1,
+		.NumTriggered = 1,
+		.Inactive = 1,
+		.BrokenTip = 0,
+		.Broken = 0,
+	};
+
+	const FBudStatusDataLayout Point3_ExpectedBudStatus = {
+		.ApicalMeristem = 0,
+		.Codominant = 0,
+		.Axillary = 1,
+		.Seed = 0,
+		.Dormant = 0,
+		.Triggered = 1,
+		.NumTriggered = 1,
+		.Inactive = 1,
+		.BrokenTip = 0,
+		.Broken = 0,
+	};
+
+	const FBudStatusDataLayout Point4_ExpectedBudStatus = {
+		.ApicalMeristem = 0,
+		.Codominant = 0,
+		.Axillary = 0,
+		.Seed = 0,
+		.Dormant = 0,
+		.Triggered = 1,
+		.NumTriggered = 1,
+		.Inactive = 1,
+		.BrokenTip = 0,
+		.Broken = 0,
+	};
+
+	const FBudStatusDataLayout Point5_ExpectedBudStatus = {
+		.ApicalMeristem = 1,
+		.Codominant = 0,
+		.Axillary = 0,
+		.Seed = 0,
+		.Dormant = 1,
+		.Triggered = 0,
+		.NumTriggered = 0,
+		.Inactive = 0,
+		.BrokenTip = 0,
+		.Broken = 0,
+	};
+
+	const FBudStatusDataLayout Point6_ExpectedBudStatus = {
+		.ApicalMeristem = 1,
+		.Codominant = 0,
+		.Axillary = 0,
+		.Seed = 0,
+		.Dormant = 1,
+		.Triggered = 0,
+		.NumTriggered = 0,
+		.Inactive = 0,
+		.BrokenTip = 0,
+		.Broken = 0,
+	};
+
+	const FBudStatusDataLayout Point7_ExpectedBudStatus = {
+		.ApicalMeristem = 0,
+		.Codominant = 0,
+		.Axillary = 0,
+		.Seed = 0,
+		.Dormant = 0,
+		.Triggered = 1,
+		.NumTriggered = 1,
+		.Inactive = 1,
+		.BrokenTip = 0,
+		.Broken = 0,
+	};
+
+	const FBudStatusDataLayout Point8_ExpectedBudStatus = {
+		.ApicalMeristem = 0,
+		.Codominant = 0,
+		.Axillary = 1,
+		.Seed = 0,
+		.Dormant = 0,
+		.Triggered = 1,
+		.NumTriggered = 1,
+		.Inactive = 1,
+		.BrokenTip = 0,
+		.Broken = 0,
+	};
+
+	const FBudStatusDataLayout Point9_ExpectedBudStatus = {
+		.ApicalMeristem = 1,
+		.Codominant = 0,
+		.Axillary = 0,
+		.Seed = 0,
+		.Dormant = 1,
+		.Triggered = 0,
+		.NumTriggered = 0,
+		.Inactive = 0,
+		.BrokenTip = 0,
+		.Broken = 0,
+	};
+
+	const FBudStatusDataLayout Point10_ExpectedBudStatus = {
+		.ApicalMeristem = 0,
+		.Codominant = 0,
+		.Axillary = 0,
+		.Seed = 0,
+		.Dormant = 0,
+		.Triggered = 1,
+		.NumTriggered = 1,
+		.Inactive = 1,
+		.BrokenTip = 0,
+		.Broken = 0,
+	};
+
+	const FBudStatusDataLayout Point11_ExpectedBudStatus = {
+		.ApicalMeristem = 1,
+		.Codominant = 0,
+		.Axillary = 0,
+		.Seed = 0,
+		.Dormant = 1,
+		.Triggered = 0,
+		.NumTriggered = 0,
+		.Inactive = 0,
+		.BrokenTip = 0,
+		.Broken = 0,
+	};
+
+	const TSharedRef<FManagedArrayCollection> Collection = PVMockTreeCollection::CreateCollection();
+
+	const PV::FBudStatusAttributeConstView BudStatusAttribute = PV::FBudStatusAttribute::FindAttribute(*Collection);
+	UTEST_TRUE("HasAttribute", BudStatusAttribute);
+	UTEST_TRUE("Has 12 elements", BudStatusAttribute.Num() >= 12);
+
+	UTEST_TRUE("Point0_BudStatus",  BudStatusAttribute[0].Array  == Point0_ExpectedBudStatus);
+	UTEST_TRUE("Point1_BudStatus",  BudStatusAttribute[1].Array  == Point1_ExpectedBudStatus);
+	UTEST_TRUE("Point2_BudStatus",  BudStatusAttribute[2].Array  == Point2_ExpectedBudStatus);
+	UTEST_TRUE("Point3_BudStatus",  BudStatusAttribute[3].Array  == Point3_ExpectedBudStatus);
+	UTEST_TRUE("Point4_BudStatus",  BudStatusAttribute[4].Array  == Point4_ExpectedBudStatus);
+	UTEST_TRUE("Point5_BudStatus",  BudStatusAttribute[5].Array  == Point5_ExpectedBudStatus);
+	UTEST_TRUE("Point6_BudStatus",  BudStatusAttribute[6].Array  == Point6_ExpectedBudStatus);
+	UTEST_TRUE("Point7_BudStatus",  BudStatusAttribute[7].Array  == Point7_ExpectedBudStatus);
+	UTEST_TRUE("Point8_BudStatus",  BudStatusAttribute[8].Array  == Point8_ExpectedBudStatus);
+	UTEST_TRUE("Point9_BudStatus",  BudStatusAttribute[9].Array  == Point9_ExpectedBudStatus);
+	UTEST_TRUE("Point10_BudStatus", BudStatusAttribute[10].Array == Point10_ExpectedBudStatus);
+	UTEST_TRUE("Point11_BudStatus", BudStatusAttribute[11].Array == Point11_ExpectedBudStatus);
+
+	return true;
+}
+
+PV_IMPORTER_TEST(ComputeBudHormoneLevels)
+{
+	struct FBudHormoneLevelsDataLayout
+	{
+		float Apical;
+		float Axillary;
+		float AxillaryInhibition;
+		float Radical;
+		float Ethylene;
+		float Cykininin;
+
+		bool operator==(const TArray<float>& Other) const
+		{
+			return Other.Num() == 6
+				&& FMath::IsNearlyEqual(Other[0], Apical, 0.01f)
+				&& FMath::IsNearlyEqual(Other[1], Axillary, 0.01f)
+				&& FMath::IsNearlyEqual(Other[2], AxillaryInhibition, 0.01f)
+				&& FMath::IsNearlyEqual(Other[3], Radical, 0.01f)
+				&& FMath::IsNearlyEqual(Other[4], Ethylene, 0.01f)
+				&& FMath::IsNearlyEqual(Other[5], Cykininin, 0.01f);
+		}
+	};
+
+	const FBudHormoneLevelsDataLayout Point0_ExpectedBudHormoneLevels = {
+		.Apical = 0,
+		.Axillary = 0,
+		.AxillaryInhibition = 0,
+		.Radical = 0,
+		.Ethylene = 1,
+		.Cykininin = 0,
+	};
+	const FBudHormoneLevelsDataLayout Point5_ExpectedBudHormoneLevels = {
+		.Apical = 1,
+		.Axillary = 0,
+		.AxillaryInhibition = 0,
+		.Radical = 0,
+		.Ethylene = 0,
+		.Cykininin = 0,
+	};
+	const FBudHormoneLevelsDataLayout Point6_ExpectedBudHormoneLevels = {
+		.Apical = 1,
+		.Axillary = 0,
+		.AxillaryInhibition = 0,
+		.Radical = 0,
+		.Ethylene = 0,
+		.Cykininin = 0,
+	};
+	const FBudHormoneLevelsDataLayout Point9_ExpectedBudHormoneLevels = {
+		.Apical = 1,
+		.Axillary = 0,
+		.AxillaryInhibition = 0,
+		.Radical = 0,
+		.Ethylene = 0,
+		.Cykininin = 0,
+	};
+	const FBudHormoneLevelsDataLayout Point11_ExpectedBudHormoneLevels = {
+		.Apical = 1,
+		.Axillary = 0,
+		.AxillaryInhibition = 0,
+		.Radical = 0,
+		.Ethylene = 0,
+		.Cykininin = 0,
+	};
+
+	const TSharedRef<FManagedArrayCollection> Collection = PVMockTreeCollection::CreateCollection();
+
+	const PV::FBudHormoneLevelsAttributeConstView BudHormoneLevelsAttribute = PV::FBudHormoneLevelsAttribute::FindAttribute(*Collection);
+	UTEST_TRUE("HasAttribute", BudHormoneLevelsAttribute.IsValid());
+	UTEST_TRUE("Has 12 elements", BudHormoneLevelsAttribute.Num() >= 12);
+
+	UTEST_TRUE("Point0_ExpectedBudHormoneLevels",  BudHormoneLevelsAttribute[0].Array  == Point0_ExpectedBudHormoneLevels);
+	UTEST_TRUE("Point5_ExpectedBudHormoneLevels",  BudHormoneLevelsAttribute[5].Array  == Point5_ExpectedBudHormoneLevels);
+	UTEST_TRUE("Point6_ExpectedBudHormoneLevels",  BudHormoneLevelsAttribute[6].Array  == Point6_ExpectedBudHormoneLevels);
+	UTEST_TRUE("Point9_ExpectedBudHormoneLevels",  BudHormoneLevelsAttribute[9].Array  == Point9_ExpectedBudHormoneLevels);
+	UTEST_TRUE("Point11_ExpectedBudHormoneLevels", BudHormoneLevelsAttribute[11].Array == Point11_ExpectedBudHormoneLevels);
+
+	return true;
+}
+
+PV_IMPORTER_TEST(ComputeBudLateralMeristem)
+{
+	const TSharedRef<FManagedArrayCollection> Collection = PVMockTreeCollection::CreateCollection();
+	
+	const PV::FBudLateralMeristemAttributeConstView BudLateralMeristemAttribute = PV::FBudLateralMeristemAttribute::FindAttribute(*Collection);
+	UTEST_TRUE("HasAttribute", BudLateralMeristemAttribute.IsValid());
+	UTEST_TRUE("Has 12 elements", BudLateralMeristemAttribute.Num() >= 12);
+
+	UTEST_TRUE("Point0_DavinciIsNotZero",  !FMath::IsNearlyZero(BudLateralMeristemAttribute[0].Davinci));
+	UTEST_TRUE("Point1_DavinciIsNotZero",  !FMath::IsNearlyZero(BudLateralMeristemAttribute[1].Davinci));
+	UTEST_TRUE("Point2_DavinciIsNotZero",  !FMath::IsNearlyZero(BudLateralMeristemAttribute[2].Davinci));
+	UTEST_TRUE("Point3_DavinciIsNotZero",  !FMath::IsNearlyZero(BudLateralMeristemAttribute[3].Davinci));
+	UTEST_TRUE("Point4_DavinciIsNotZero",  !FMath::IsNearlyZero(BudLateralMeristemAttribute[4].Davinci));
+	UTEST_TRUE("Point7_DavinciIsNotZero",  !FMath::IsNearlyZero(BudLateralMeristemAttribute[7].Davinci));
+	UTEST_TRUE("Point8_DavinciIsNotZero",  !FMath::IsNearlyZero(BudLateralMeristemAttribute[8].Davinci));
+	UTEST_TRUE("Point10_DavinciIsNotZero", !FMath::IsNearlyZero(BudLateralMeristemAttribute[10].Davinci));
+
+	UTEST_TRUE("Point5_DavinciIsZero",  FMath::IsNearlyZero(BudLateralMeristemAttribute[5].Davinci));
+	UTEST_TRUE("Point6_DavinciIsZero",  FMath::IsNearlyZero(BudLateralMeristemAttribute[6].Davinci));
+	UTEST_TRUE("Point9_DavinciIsZero",  FMath::IsNearlyZero(BudLateralMeristemAttribute[9].Davinci));
+	UTEST_TRUE("Point11_DavinciIsZero", FMath::IsNearlyZero(BudLateralMeristemAttribute[11].Davinci));
+
+	return true;
+}
+
+PV_IMPORTER_TEST(ComputeBudLightDetected)
+{
+	struct FBudLightDetectedDataLayout
+	{
+		float Availible;
+		float Resource;
+		float Branch;
+		float Collision;
+
+		bool operator==(const TArray<float>& Other) const
+		{
+			return Other.Num() == 4
+				&& Other[0] == Availible
+				&& Other[1] == Resource
+				&& Other[2] == Branch
+				&& Other[3] == Collision;
+		}
+	};
+
+	FBudLightDetectedDataLayout ExpectedBudLightDetected = {
+		.Availible = 1.0,
+		.Resource = 1.0,
+		.Branch = 1.0,
+		.Collision = 0.0,
+	};
 
-	FManagedArrayCollection Collection;
-	PV::ImportHelper::CreateEmptyGrowthData(Collection, NumBranchPoints, NumFoliagePoints, NumBranches, NumFoliageMeshes);
-	PV::ImportHelper::FillGrowthDataFromBranchHierarchyDescription(Collection, { BranchHierarchy });
-
-	PV::Facades::FPointFacade PointFacade(Collection);
-	const PV::Facades::FBranchFacade BranchFacade(Collection);
-	const PV::Facades::FPlantFacade PlantFacade(Collection);
-	PV::ImportHelper::ComputeBudDevelopment(PointFacade, BranchFacade, PlantFacade);
-
-	UTEST_TRUE("Point0_BudDevelopment_Num", PointFacade.GetBudDevelopment(0).Num() == ExpectedNumBudDevelopmentValues);
-	UTEST_TRUE("Point1_BudDevelopment_Num", PointFacade.GetBudDevelopment(1).Num() == ExpectedNumBudDevelopmentValues);
-	UTEST_TRUE("Point2_BudDevelopment_Num", PointFacade.GetBudDevelopment(2).Num() == ExpectedNumBudDevelopmentValues);
-	UTEST_TRUE("Point3_BudDevelopment_Num", PointFacade.GetBudDevelopment(3).Num() == ExpectedNumBudDevelopmentValues);
-	UTEST_TRUE("Point4_BudDevelopment_Num", PointFacade.GetBudDevelopment(4).Num() == ExpectedNumBudDevelopmentValues);
-	UTEST_TRUE("Point5_BudDevelopment_Num", PointFacade.GetBudDevelopment(5).Num() == ExpectedNumBudDevelopmentValues);
-	UTEST_TRUE("Point6_BudDevelopment_Num", PointFacade.GetBudDevelopment(6).Num() == ExpectedNumBudDevelopmentValues);
-	UTEST_TRUE("Point7_BudDevelopment_Num", PointFacade.GetBudDevelopment(7).Num() == ExpectedNumBudDevelopmentValues);
-	UTEST_TRUE("Point8_BudDevelopment_Num", PointFacade.GetBudDevelopment(8).Num() == ExpectedNumBudDevelopmentValues);
-	UTEST_TRUE("Point9_BudDevelopment_Num", PointFacade.GetBudDevelopment(9).Num() == ExpectedNumBudDevelopmentValues);
-	UTEST_TRUE("Point10_BudDevelopment_Num", PointFacade.GetBudDevelopment(10).Num() == ExpectedNumBudDevelopmentValues);
-	UTEST_TRUE("Point11_BudDevelopment_Num", PointFacade.GetBudDevelopment(11).Num() == ExpectedNumBudDevelopmentValues);
-
-	UTEST_TRUE(
-		"Point0_BudDevelopment", 
-		FMemory::Memcmp((void*)PointFacade.GetBudDevelopment(0).GetData(), (void*)&Point0_ExpectedBudDevelopment, ExpectedNumBudDevelopmentValues * sizeof(int32)) == 0
-	);
-	UTEST_TRUE(
-		"Point1_BudDevelopment",
-		FMemory::Memcmp((void*)PointFacade.GetBudDevelopment(1).GetData(), (void*)&Point1_ExpectedBudDevelopment, ExpectedNumBudDevelopmentValues * sizeof(int32)) == 0
-	);
-	UTEST_TRUE(
-		"Point2_BudDevelopment",
-		FMemory::Memcmp((void*)PointFacade.GetBudDevelopment(2).GetData(), (void*)&Point2_ExpectedBudDevelopment, ExpectedNumBudDevelopmentValues * sizeof(int32)) == 0
-	);
-	UTEST_TRUE(
-		"Point3_BudDevelopment",
-		FMemory::Memcmp((void*)PointFacade.GetBudDevelopment(3).GetData(), (void*)&Point3_ExpectedBudDevelopment, ExpectedNumBudDevelopmentValues * sizeof(int32)) == 0
-	);
-	UTEST_TRUE(
-		"Point4_BudDevelopment",
-		FMemory::Memcmp((void*)PointFacade.GetBudDevelopment(4).GetData(), (void*)&Point4_ExpectedBudDevelopment, ExpectedNumBudDevelopmentValues * sizeof(int32)) == 0
-	);
-	UTEST_TRUE(
-		"Point5_BudDevelopment",
-		FMemory::Memcmp((void*)PointFacade.GetBudDevelopment(5).GetData(), (void*)&Point5_ExpectedBudDevelopment, ExpectedNumBudDevelopmentValues * sizeof(int32)) == 0
-	);
-	UTEST_TRUE(
-		"Point6_BudDevelopment",
-		FMemory::Memcmp((void*)PointFacade.GetBudDevelopment(6).GetData(), (void*)&Point6_ExpectedBudDevelopment, ExpectedNumBudDevelopmentValues * sizeof(int32)) == 0
-	);
-	UTEST_TRUE(
-		"Point7_BudDevelopment",
-		FMemory::Memcmp((void*)PointFacade.GetBudDevelopment(7).GetData(), (void*)&Point7_ExpectedBudDevelopment, ExpectedNumBudDevelopmentValues * sizeof(int32)) == 0
-	);
-	UTEST_TRUE(
-		"Point8_BudDevelopment",
-		FMemory::Memcmp((void*)PointFacade.GetBudDevelopment(8).GetData(), (void*)&Point8_ExpectedBudDevelopment, ExpectedNumBudDevelopmentValues * sizeof(int32)) == 0
-	);
-	UTEST_TRUE(
-		"Point9_BudDevelopment",
-		FMemory::Memcmp((void*)PointFacade.GetBudDevelopment(9).GetData(), (void*)&Point9_ExpectedBudDevelopment, ExpectedNumBudDevelopmentValues * sizeof(int32)) == 0
-	);
-	UTEST_TRUE(
-		"Point10_BudDevelopment",
-		FMemory::Memcmp((void*)PointFacade.GetBudDevelopment(10).GetData(), (void*)&Point10_ExpectedBudDevelopment, ExpectedNumBudDevelopmentValues * sizeof(int32)) == 0
-	);
-
-	UTEST_TRUE(
-		"Point11_BudDevelopment",
-		FMemory::Memcmp((void*)PointFacade.GetBudDevelopment(11).GetData(), (void*)&Point11_ExpectedBudDevelopment, ExpectedNumBudDevelopmentValues * sizeof(int32)) == 0
-	);
+	const TSharedRef<FManagedArrayCollection> Collection = PVMockTreeCollection::CreateCollection();
+
+	const PV::FBudLightDetectedAttributeConstView BudLightDetectedAttribute = PV::FBudLightDetectedAttribute::FindAttribute(*Collection);
+	UTEST_TRUE("HasAttribute", BudLightDetectedAttribute.IsValid());
+	UTEST_TRUE("Has 12 elements", BudLightDetectedAttribute.Num() >= 12);
+
+	UTEST_TRUE("Point0_BudLightDetected",  BudLightDetectedAttribute[0].Array  == ExpectedBudLightDetected);
+	UTEST_TRUE("Point1_BudLightDetected",  BudLightDetectedAttribute[1].Array  == ExpectedBudLightDetected);
+	UTEST_TRUE("Point2_BudLightDetected",  BudLightDetectedAttribute[2].Array  == ExpectedBudLightDetected);
+	UTEST_TRUE("Point3_BudLightDetected",  BudLightDetectedAttribute[3].Array  == ExpectedBudLightDetected);
+	UTEST_TRUE("Point4_BudLightDetected",  BudLightDetectedAttribute[4].Array  == ExpectedBudLightDetected);
+	UTEST_TRUE("Point5_BudLightDetected",  BudLightDetectedAttribute[5].Array  == ExpectedBudLightDetected);
+	UTEST_TRUE("Point6_BudLightDetected",  BudLightDetectedAttribute[6].Array  == ExpectedBudLightDetected);
+	UTEST_TRUE("Point7_BudLightDetected",  BudLightDetectedAttribute[7].Array  == ExpectedBudLightDetected);
+	UTEST_TRUE("Point8_BudLightDetected",  BudLightDetectedAttribute[8].Array  == ExpectedBudLightDetected);
+	UTEST_TRUE("Point9_BudLightDetected",  BudLightDetectedAttribute[9].Array  == ExpectedBudLightDetected);
+	UTEST_TRUE("Point10_BudLightDetected", BudLightDetectedAttribute[10].Array == ExpectedBudLightDetected);
+	UTEST_TRUE("Point11_BudLightDetected", BudLightDetectedAttribute[11].Array == ExpectedBudLightDetected);
 
 	return true;
 }
@@ -369,46 +729,24 @@ PV_IMPORTER_TEST(ComputeBudDirections)
 
 	const int32 ExpectedNumBudDirectionValues = 6;
 
-	const FPVBranchHierarchyDescription BranchHierarchy = PVMockTree::CreateBranchHierarchyDescription();
-	const int32 NumBranchPoints = BranchHierarchy.GetNumBranchPoints();
-	const int32 NumFoliagePoints = 0;
-	const int32 NumBranches = BranchHierarchy.GetNumBranches();
-	const int32 NumFoliageMeshes = 0;
-
-	FManagedArrayCollection Collection;
-	PV::ImportHelper::CreateEmptyGrowthData(Collection, NumBranchPoints, NumFoliagePoints, NumBranches, NumFoliageMeshes);
-	PV::ImportHelper::FillGrowthDataFromBranchHierarchyDescription(Collection, { BranchHierarchy });
-
-	PV::Facades::FPointFacade PointFacade(Collection);
-	const PV::Facades::FBranchFacade BranchFacade(Collection);
-	const PV::Facades::FPlantFacade PlantFacade(Collection);
-	PV::ImportHelper::ComputeBudDirections(PointFacade, BranchFacade, PlantFacade);
-
-	UTEST_TRUE("Point0_BudDirections_Num", PointFacade.GetBudDirection(0).Num() == ExpectedNumBudDirectionValues);
-	UTEST_TRUE("Point1_BudDirections_Num", PointFacade.GetBudDirection(1).Num() == ExpectedNumBudDirectionValues);
-	UTEST_TRUE("Point2_BudDirections_Num", PointFacade.GetBudDirection(2).Num() == ExpectedNumBudDirectionValues);
-	UTEST_TRUE("Point3_BudDirections_Num", PointFacade.GetBudDirection(3).Num() == ExpectedNumBudDirectionValues);
-	UTEST_TRUE("Point4_BudDirections_Num", PointFacade.GetBudDirection(4).Num() == ExpectedNumBudDirectionValues);
-	UTEST_TRUE("Point5_BudDirections_Num", PointFacade.GetBudDirection(5).Num() == ExpectedNumBudDirectionValues);
-	UTEST_TRUE("Point6_BudDirections_Num", PointFacade.GetBudDirection(6).Num() == ExpectedNumBudDirectionValues);
-	UTEST_TRUE("Point7_BudDirections_Num", PointFacade.GetBudDirection(7).Num() == ExpectedNumBudDirectionValues);
-	UTEST_TRUE("Point8_BudDirections_Num", PointFacade.GetBudDirection(8).Num() == ExpectedNumBudDirectionValues);
-	UTEST_TRUE("Point9_BudDirections_Num", PointFacade.GetBudDirection(9).Num() == ExpectedNumBudDirectionValues);
-	UTEST_TRUE("Point10_BudDirections_Num", PointFacade.GetBudDirection(10).Num() == ExpectedNumBudDirectionValues);
-	UTEST_TRUE("Point11_BudDirections_Num", PointFacade.GetBudDirection(11).Num() == ExpectedNumBudDirectionValues);
-
-	UTEST_NEARLY_EQUAL("Point0_Apical", FVector(PointFacade.GetBudDirection(0)[0]), FVector(ExpectedApical0), UE_KINDA_SMALL_NUMBER);
-	UTEST_NEARLY_EQUAL("Point1_Apical", FVector(PointFacade.GetBudDirection(1)[0]), FVector(ExpectedApical1), UE_KINDA_SMALL_NUMBER);
-	UTEST_NEARLY_EQUAL("Point2_Apical", FVector(PointFacade.GetBudDirection(2)[0]), FVector(ExpectedApical2), UE_KINDA_SMALL_NUMBER);
-	UTEST_NEARLY_EQUAL("Point3_Apical", FVector(PointFacade.GetBudDirection(3)[0]), FVector(ExpectedApical3), UE_KINDA_SMALL_NUMBER);
-	UTEST_NEARLY_EQUAL("Point4_Apical", FVector(PointFacade.GetBudDirection(4)[0]), FVector(ExpectedApical4), UE_KINDA_SMALL_NUMBER);
-	UTEST_NEARLY_EQUAL("Point5_Apical", FVector(PointFacade.GetBudDirection(5)[0]), FVector(ExpectedApical5), UE_KINDA_SMALL_NUMBER);
-	UTEST_NEARLY_EQUAL("Point6_Apical", FVector(PointFacade.GetBudDirection(6)[0]), FVector(ExpectedApical6), UE_KINDA_SMALL_NUMBER);
-	UTEST_NEARLY_EQUAL("Point7_Apical", FVector(PointFacade.GetBudDirection(7)[0]), FVector(ExpectedApical7), UE_KINDA_SMALL_NUMBER);
-	UTEST_NEARLY_EQUAL("Point8_Apical", FVector(PointFacade.GetBudDirection(8)[0]), FVector(ExpectedApical8), UE_KINDA_SMALL_NUMBER);
-	UTEST_NEARLY_EQUAL("Point9_Apical", FVector(PointFacade.GetBudDirection(9)[0]), FVector(ExpectedApical9), UE_KINDA_SMALL_NUMBER);
-	UTEST_NEARLY_EQUAL("Point10_Apical", FVector(PointFacade.GetBudDirection(10)[0]), FVector(ExpectedApical10), UE_KINDA_SMALL_NUMBER);
-	UTEST_NEARLY_EQUAL("Point11_Apical", FVector(PointFacade.GetBudDirection(11)[0]), FVector(ExpectedApical11), UE_KINDA_SMALL_NUMBER);
+	const TSharedRef<FManagedArrayCollection> Collection = PVMockTreeCollection::CreateCollection();
+
+	const PV::FBudDirectionAttributeConstView BudDirectionAttribute = PV::FBudDirectionAttribute::FindAttribute(*Collection);
+	UTEST_TRUE("HasAttribute", BudDirectionAttribute.IsValid());
+	UTEST_TRUE("Has 12 elements", BudDirectionAttribute.Num() >= 12);
+
+	UTEST_NEARLY_EQUAL("Point0_Apical", FVector(BudDirectionAttribute[0].Apical), FVector(ExpectedApical0), UE_KINDA_SMALL_NUMBER);
+	UTEST_NEARLY_EQUAL("Point1_Apical", FVector(BudDirectionAttribute[1].Apical), FVector(ExpectedApical1), UE_KINDA_SMALL_NUMBER);
+	UTEST_NEARLY_EQUAL("Point2_Apical", FVector(BudDirectionAttribute[2].Apical), FVector(ExpectedApical2), UE_KINDA_SMALL_NUMBER);
+	UTEST_NEARLY_EQUAL("Point3_Apical", FVector(BudDirectionAttribute[3].Apical), FVector(ExpectedApical3), UE_KINDA_SMALL_NUMBER);
+	UTEST_NEARLY_EQUAL("Point4_Apical", FVector(BudDirectionAttribute[4].Apical), FVector(ExpectedApical4), UE_KINDA_SMALL_NUMBER);
+	UTEST_NEARLY_EQUAL("Point5_Apical", FVector(BudDirectionAttribute[5].Apical), FVector(ExpectedApical5), UE_KINDA_SMALL_NUMBER);
+	UTEST_NEARLY_EQUAL("Point6_Apical", FVector(BudDirectionAttribute[6].Apical), FVector(ExpectedApical6), UE_KINDA_SMALL_NUMBER);
+	UTEST_NEARLY_EQUAL("Point7_Apical", FVector(BudDirectionAttribute[7].Apical), FVector(ExpectedApical7), UE_KINDA_SMALL_NUMBER);
+	UTEST_NEARLY_EQUAL("Point8_Apical", FVector(BudDirectionAttribute[8].Apical), FVector(ExpectedApical8), UE_KINDA_SMALL_NUMBER);
+	UTEST_NEARLY_EQUAL("Point9_Apical", FVector(BudDirectionAttribute[9].Apical), FVector(ExpectedApical9), UE_KINDA_SMALL_NUMBER);
+	UTEST_NEARLY_EQUAL("Point10_Apical", FVector(BudDirectionAttribute[10].Apical), FVector(ExpectedApical10), UE_KINDA_SMALL_NUMBER);
+	UTEST_NEARLY_EQUAL("Point11_Apical", FVector(BudDirectionAttribute[11].Apical), FVector(ExpectedApical11), UE_KINDA_SMALL_NUMBER);
 
 	// TODO: Test UpVector and Axillary
 
@@ -417,56 +755,56 @@ PV_IMPORTER_TEST(ComputeBudDirections)
 
 PV_IMPORTER_TEST(SplitFusedPoints)
 {
-	const FPVBranchHierarchyDescription BranchHierarchy = PVMockTree::CreateBranchHierarchyDescription();
-	const int32 NumBranchPoints = BranchHierarchy.GetNumBranchPoints();
-	const int32 NumFoliagePoints = 0;
-	const int32 NumBranches = BranchHierarchy.GetNumBranches();
-	const int32 NumFoliageMeshes = 0;
+	const TSharedRef<FManagedArrayCollection> Collection = PVMockTreeCollection::CreateCollection(false);
 
-	FManagedArrayCollection Collection;
-	PV::ImportHelper::CreateEmptyGrowthData(Collection, NumBranchPoints, NumFoliagePoints, NumBranches, NumFoliageMeshes);
-	PV::ImportHelper::FillGrowthDataFromBranchHierarchyDescription(Collection, { BranchHierarchy });
+	const bool bSplitFusedPointsResult = PV::ImportHelper::SplitFusedPoints(*Collection);
+	UTEST_TRUE("SplitFusedPointsSucceeded", bSplitFusedPointsResult);
 
-	PV::Facades::FPointFacade PointFacade(Collection);
-	const PV::Facades::FBranchFacade BranchFacade(Collection);
-	const PV::Facades::FPlantFacade PlantFacade(Collection);
-	PV::ImportHelper::ComputeBudDirections(PointFacade, BranchFacade, PlantFacade);
-	PV::ImportHelper::ComputeLengthFromRoot(PointFacade, BranchFacade, PlantFacade);
-	PV::ImportHelper::ComputeBudDevelopment(PointFacade, BranchFacade, PlantFacade);
-	PV::ImportHelper::ComputeBudDirections(PointFacade, BranchFacade, PlantFacade);
+	PV::FPointPositionAttributeConstView PointPositionAttribute = PV::FPointPositionAttribute::FindAttribute(*Collection);
+	UTEST_TRUE("HasAttribute", PointPositionAttribute.IsValid());
+	UTEST_TRUE("Has 15 elements", PointPositionAttribute.Num() == 15);
 
-	const bool bSplitFusedPointsResult = PV::ImportHelper::SplitFusedPoints(Collection);
-	UTEST_TRUE("SplitFusedPointsSucceeded", bSplitFusedPointsResult);
+	PV::FPointLengthFromRootAttributeConstView PointLengthFromRootAttribute = PV::FPointLengthFromRootAttribute::FindAttribute(*Collection);
+	UTEST_TRUE("HasAttribute", PointLengthFromRootAttribute.IsValid());
+	UTEST_TRUE("Has 15 elements", PointLengthFromRootAttribute.Num() == 15);
+
+	const PV::FPointLengthFromSeedAttributeView PointLengthFromSeedAttribute = PV::FPointLengthFromSeedAttribute::FindAttribute(*Collection);
+	UTEST_TRUE("HasAttribute", PointLengthFromSeedAttribute.IsValid());
+	UTEST_TRUE("Has 15 elements", PointLengthFromSeedAttribute.Num() == 15);
+
+	const PV::FBranchPointsAttributeView BranchPointsAttribute = PV::FBranchPointsAttribute::FindAttribute(*Collection);
+	UTEST_TRUE("HasAttribute", BranchPointsAttribute.IsValid());
+	UTEST_TRUE("Has 4 elements", BranchPointsAttribute.Num() == 4);
 
 	const int32 Branch0_OldRootPointIndex = 0;
 	const int32 Branch1_OldRootPointIndex = 2;
 	const int32 Branch2_OldRootPointIndex = 8;
 	const int32 Branch3_OldRootPointIndex = 3;
 
-	const int32 Branch0_NewRootPointIndex = BranchFacade.GetPoints(0)[0];
-	const int32 Branch1_NewRootPointIndex = BranchFacade.GetPoints(1)[0];
-	const int32 Branch2_NewRootPointIndex = BranchFacade.GetPoints(2)[0];
-	const int32 Branch3_NewRootPointIndex = BranchFacade.GetPoints(3)[0];
+	const int32 Branch0_NewRootPointIndex = BranchPointsAttribute[0][0];
+	const int32 Branch1_NewRootPointIndex = BranchPointsAttribute[1][0];
+	const int32 Branch2_NewRootPointIndex = BranchPointsAttribute[2][0];
+	const int32 Branch3_NewRootPointIndex = BranchPointsAttribute[3][0];
 
 	UTEST_TRUE("Branch0_HasSamePointIndex",  Branch0_NewRootPointIndex == Branch0_OldRootPointIndex);
 	UTEST_TRUE("Branch1_HasUniqePointIndex", Branch1_NewRootPointIndex != Branch1_OldRootPointIndex && Branch1_NewRootPointIndex > 11);
 	UTEST_TRUE("Branch2_HasUniqePointIndex", Branch2_NewRootPointIndex != Branch2_OldRootPointIndex && Branch2_NewRootPointIndex > 11);
 	UTEST_TRUE("Branch3_HasUniqePointIndex", Branch3_NewRootPointIndex != Branch3_OldRootPointIndex && Branch3_NewRootPointIndex > 11);
 
-	UTEST_EQUAL("Branch0_NewPointPositionIsSame", PointFacade.GetPosition(Branch0_NewRootPointIndex), PointFacade.GetPosition(Branch0_OldRootPointIndex));
-	UTEST_EQUAL("Branch1_NewPointPositionIsSame", PointFacade.GetPosition(Branch1_NewRootPointIndex), PointFacade.GetPosition(Branch1_OldRootPointIndex));
-	UTEST_EQUAL("Branch2_NewPointPositionIsSame", PointFacade.GetPosition(Branch2_NewRootPointIndex), PointFacade.GetPosition(Branch2_OldRootPointIndex));
-	UTEST_EQUAL("Branch3_NewPointPositionIsSame", PointFacade.GetPosition(Branch3_NewRootPointIndex), PointFacade.GetPosition(Branch3_OldRootPointIndex));
+	UTEST_EQUAL("Branch0_NewPointPositionIsSame", PointPositionAttribute[Branch0_NewRootPointIndex], PointPositionAttribute[Branch0_OldRootPointIndex]);
+	UTEST_EQUAL("Branch1_NewPointPositionIsSame", PointPositionAttribute[Branch1_NewRootPointIndex], PointPositionAttribute[Branch1_OldRootPointIndex]);
+	UTEST_EQUAL("Branch2_NewPointPositionIsSame", PointPositionAttribute[Branch2_NewRootPointIndex], PointPositionAttribute[Branch2_OldRootPointIndex]);
+	UTEST_EQUAL("Branch3_NewPointPositionIsSame", PointPositionAttribute[Branch3_NewRootPointIndex], PointPositionAttribute[Branch3_OldRootPointIndex]);
 
-	UTEST_EQUAL("Branch0_NewPointLengthFromRootIsSame", PointFacade.GetLengthFromRoot(Branch0_NewRootPointIndex), PointFacade.GetLengthFromRoot(Branch0_OldRootPointIndex));
-	UTEST_EQUAL("Branch1_NewPointLengthFromRootIsSame", PointFacade.GetLengthFromRoot(Branch1_NewRootPointIndex), PointFacade.GetLengthFromRoot(Branch1_OldRootPointIndex));
-	UTEST_EQUAL("Branch2_NewPointLengthFromRootIsSame", PointFacade.GetLengthFromRoot(Branch2_NewRootPointIndex), PointFacade.GetLengthFromRoot(Branch2_OldRootPointIndex));
-	UTEST_EQUAL("Branch3_NewPointLengthFromRootIsSame", PointFacade.GetLengthFromRoot(Branch3_NewRootPointIndex), PointFacade.GetLengthFromRoot(Branch3_OldRootPointIndex));
+	UTEST_EQUAL("Branch0_NewPointLengthFromRootIsSame", PointLengthFromRootAttribute[Branch0_NewRootPointIndex], PointLengthFromRootAttribute[Branch0_OldRootPointIndex]);
+	UTEST_EQUAL("Branch1_NewPointLengthFromRootIsSame", PointLengthFromRootAttribute[Branch1_NewRootPointIndex], PointLengthFromRootAttribute[Branch1_OldRootPointIndex]);
+	UTEST_EQUAL("Branch2_NewPointLengthFromRootIsSame", PointLengthFromRootAttribute[Branch2_NewRootPointIndex], PointLengthFromRootAttribute[Branch2_OldRootPointIndex]);
+	UTEST_EQUAL("Branch3_NewPointLengthFromRootIsSame", PointLengthFromRootAttribute[Branch3_NewRootPointIndex], PointLengthFromRootAttribute[Branch3_OldRootPointIndex]);
 
-	UTEST_EQUAL("Branch0_NewPointLengthFromSeedIsSame", PointFacade.GetLengthFromSeed(Branch0_NewRootPointIndex), PointFacade.GetLengthFromSeed(Branch0_OldRootPointIndex));
-	UTEST_EQUAL("Branch1_NewPointLengthFromSeedIsSame", PointFacade.GetLengthFromSeed(Branch1_NewRootPointIndex), PointFacade.GetLengthFromSeed(Branch1_OldRootPointIndex));
-	UTEST_EQUAL("Branch2_NewPointLengthFromSeedIsSame", PointFacade.GetLengthFromSeed(Branch2_NewRootPointIndex), PointFacade.GetLengthFromSeed(Branch2_OldRootPointIndex));
-	UTEST_EQUAL("Branch3_NewPointLengthFromSeedIsSame", PointFacade.GetLengthFromSeed(Branch3_NewRootPointIndex), PointFacade.GetLengthFromSeed(Branch3_OldRootPointIndex));
+	UTEST_EQUAL("Branch0_NewPointLengthFromSeedIsSame", PointLengthFromSeedAttribute[Branch0_NewRootPointIndex], PointLengthFromSeedAttribute[Branch0_OldRootPointIndex]);
+	UTEST_EQUAL("Branch1_NewPointLengthFromSeedIsSame", PointLengthFromSeedAttribute[Branch1_NewRootPointIndex], PointLengthFromSeedAttribute[Branch1_OldRootPointIndex]);
+	UTEST_EQUAL("Branch2_NewPointLengthFromSeedIsSame", PointLengthFromSeedAttribute[Branch2_NewRootPointIndex], PointLengthFromSeedAttribute[Branch2_OldRootPointIndex]);
+	UTEST_EQUAL("Branch3_NewPointLengthFromSeedIsSame", PointLengthFromSeedAttribute[Branch3_NewRootPointIndex], PointLengthFromSeedAttribute[Branch3_OldRootPointIndex]);
 
 	return true;
 }
