[PVE] Enable embedded graphs

#Synced-CL 50702163

diff --git a/Engine/Plugins/Experimental/ProceduralVegetationEditor/Source/ProceduralVegetation/Public/ProceduralVegetation.h b/Engine/Plugins/Experimental/ProceduralVegetationEditor/Source/ProceduralVegetation/Public/ProceduralVegetation.h
--- a/Engine/Plugins/Experimental/ProceduralVegetationEditor/Source/ProceduralVegetation/Public/ProceduralVegetation.h
+++ b/Engine/Plugins/Experimental/ProceduralVegetationEditor/Source/ProceduralVegetation/Public/ProceduralVegetation.h
@@ -35,9 +35,13 @@ public:
 	virtual bool IsExportToLibraryEnabled() const override { return false; }
 	virtual bool ShowGraphCustomization() const override { return false; }
 	virtual bool IsTemplatePropertyEnabled() const override { return false; }
-#endif
 
-	virtual bool SupportsEmbeddedSubgraphs() const override { return false; }
+	virtual TSubclassOf<UPCGGraph> GetEmbeddedSubgraphClass() const override { return UProceduralVegetationGraph::StaticClass(); }
+#endif
+	
+	virtual bool SupportsEmbeddedSubgraphs() const override { return !IsEmbeddedSubgraph(); }
+	UProceduralVegetationGraph* GetEmbeddedParentGraph() const { return GetTypedOuter<UProceduralVegetationGraph>(); }
+	bool IsEmbeddedSubgraph() const { return !!GetEmbeddedParentGraph(); }
 };
 
 UCLASS(Abstract)
diff --git a/Engine/Plugins/Experimental/ProceduralVegetationEditor/Source/ProceduralVegetationEditor/Private/PVEditor.cpp b/Engine/Plugins/Experimental/ProceduralVegetationEditor/Source/ProceduralVegetationEditor/Private/PVEditor.cpp
--- a/Engine/Plugins/Experimental/ProceduralVegetationEditor/Source/ProceduralVegetationEditor/Private/PVEditor.cpp
+++ b/Engine/Plugins/Experimental/ProceduralVegetationEditor/Source/ProceduralVegetationEditor/Private/PVEditor.cpp
@@ -259,7 +259,8 @@ bool FPVEditor::IsPanelAvailable(const FName PanelID) const
 		PCGEditorTabs::LogID,
 		PCGEditorTabs::PropertyDetailsID[0],
 		PCGEditorTabs::UserParamsID,
-		PCGEditorTabs::ViewportID[0]
+		PCGEditorTabs::ViewportID[0],
+		PCGEditorTabs::EmbeddedSubgraphsID
 	};
 	return SupportedIDs.Contains(PanelID);
 }
@@ -369,6 +370,18 @@ void FPVEditor::OnSelectedNodesChanged(const TSet<UObject*>& InNewSelection)
 			}
 		}	
 	}
+
+	// TODO: Hacky way to execute embedded subgraphs
+	// change execution source when focused graph is changed.
+	if (ExecutionSource && GetFocusedGraph())
+	{
+		ExecutionSource->SetGraphInterface(GetFocusedGraph());
+
+		FPCGStack Stack;
+		Stack.PushFrame(ExecutionSource.Get());
+		Stack.PushFrame(GetFocusedGraph());
+		SetStackBeingInspected(Stack);
+	}
 }
 
 void FPVEditor::OnGraphExecuted(IPCGBaseSubsystem* InSubsystem, IPCGGraphExecutionSource* InExecutionSource,
