Remove ensure in TAttributeDefinition::FindAttribute accidentally included in CL 50701977

#Synced-CL 50701977

diff --git a/Engine/Plugins/Experimental/ProceduralVegetationEditor/Source/ProceduralVegetation/Public/Facades/PVAttributes.h b/Engine/Plugins/Experimental/ProceduralVegetationEditor/Source/ProceduralVegetation/Public/Facades/PVAttributes.h
--- a/Engine/Plugins/Experimental/ProceduralVegetationEditor/Source/ProceduralVegetation/Public/Facades/PVAttributes.h
+++ b/Engine/Plugins/Experimental/ProceduralVegetationEditor/Source/ProceduralVegetation/Public/Facades/PVAttributes.h
@@ -153,13 +153,7 @@ public:
 	*/
 	static AttributeConstView FindAttribute(const FManagedArrayCollection& InCollection)
 	{
-		const TManagedArray<AttributeType>* Attribute = InCollection.FindAttributeTyped<AttributeType>(AttributeName, GroupName);
-		if (ensure(Attribute != nullptr))
-		{
-			return AttributeConstView(Attribute);
-		}
-
-		return AttributeConstView(nullptr);
+		return AttributeConstView(InCollection.FindAttributeTyped<AttributeType>(AttributeName, GroupName));
 	}
 	
 	/**
@@ -168,13 +162,7 @@ public:
 	*/
 	static AttributeView FindAttribute(FManagedArrayCollection& InCollection)
 	{
-		TManagedArray<AttributeType>* Attribute = InCollection.FindAttributeTyped<AttributeType>(AttributeName, GroupName);
-		if (ensure(Attribute != nullptr))
-		{
-			return AttributeView(Attribute);
-		}
-
-		return AttributeView(nullptr);
+		return AttributeView(InCollection.FindAttributeTyped<AttributeType>(AttributeName, GroupName));
 	}
 
 	/**
