[PVE] Introduced the Grower preset data asset for the grower node.
- Added option to save as preset and apply current settings to the selected preset.
- Preset option is also added in grower settings nodes.
- Moved FPVGrowerParams to DataTypes/PVGrowerParams.h

#Synced-CL 50702500

diff --git a/Engine/Plugins/Experimental/ProceduralVegetationEditor/Source/ProceduralVegetation/Private/DataAssets/ProceduralVegetationGrowerPreset.h b/Engine/Plugins/Experimental/ProceduralVegetationEditor/Source/ProceduralVegetation/Private/DataAssets/ProceduralVegetationGrowerPreset.h
new file mode 100644
--- /dev/null
+++ b/Engine/Plugins/Experimental/ProceduralVegetationEditor/Source/ProceduralVegetation/Private/DataAssets/ProceduralVegetationGrowerPreset.h
@@ -0,0 +1,16 @@
+// Copyright Epic Games, Inc. All Rights Reserved.
+
+#pragma once
+
+#include "Engine/DataAsset.h"
+#include "Implementations/PVGrower.h"
+#include "ProceduralVegetationGrowerPreset.generated.h"
+
+UCLASS(BlueprintType)
+class PROCEDURALVEGETATION_API UProceduralVegetationGrowerPreset : public UDataAsset
+{
+	GENERATED_BODY()
+public:
+	UPROPERTY(EditAnywhere, Category = Settings, meta=( ShowOnlyInnerProperties))
+	FPVGrowerParams GrowthParams;
+};
diff --git a/Engine/Plugins/Experimental/ProceduralVegetationEditor/Source/ProceduralVegetation/Private/DataTypes/PVGrowerParams.h b/Engine/Plugins/Experimental/ProceduralVegetationEditor/Source/ProceduralVegetation/Private/DataTypes/PVGrowerParams.h
--- a/Engine/Plugins/Experimental/ProceduralVegetationEditor/Source/ProceduralVegetation/Private/DataTypes/PVGrowerParams.h
+++ b/Engine/Plugins/Experimental/ProceduralVegetationEditor/Source/ProceduralVegetation/Private/DataTypes/PVGrowerParams.h
@@ -6,6 +6,8 @@
 #include "PVFloatRamp.h"
 #include "PVGrowerParams.generated.h"
 
+struct FPVColliderParams;
+
 USTRUCT()
 struct FPVBranchingParams
 {
@@ -505,4 +507,100 @@ struct FPVGrowerPhyllotaxyWithTargets
 
 	UPROPERTY(EditAnywhere , Category = "Phyllotaxy", meta = (ShowOnlyInnerProperties))
 	FPVGrowerPhyllotaxyParams Params;
+};
+
+UENUM()
+enum class EPVPresetTarget : uint8
+{
+	Trunk,
+	Branch
+};
+
+USTRUCT()
+struct FPVGrowerParams
+{
+	GENERATED_BODY()
+
+public:
+	UPROPERTY(EditAnywhere , Category = "Basic" , meta=(UIMin=1.0f, UIMax=45.0f, PCG_Overridable))
+	uint32 GrowthCycles = 1;
+	
+	UPROPERTY(EditAnywhere , Category = "Phyllotaxy", meta =(ShowOnlyInnerProperties, PCG_NotOverridable))
+	FPVGrowerPhyllotaxyParams Phyllotaxy;
+
+	UPROPERTY(EditAnywhere , Category = "Phyllotaxy", meta =(PCG_NotOverridable))
+	bool bBranchPhyllotaxySameAsTrunk = true;
+
+	UPROPERTY(EditAnywhere , Category = "Phyllotaxy" , meta = (EditCondition = "bBranchPhyllotaxySameAsTrunk == false" , EditConditionHides, PCG_NotOverridable))
+	FPVGrowerPhyllotaxyParams BranchPhyllotaxy;
+	
+	UPROPERTY(EditAnywhere , Category = "Phyllotaxy", meta =(PCG_NotOverridable))
+	bool bLeafPhyllotaxySameAsBranch = true;
+
+	UPROPERTY(EditAnywhere , Category = "Axial", meta =(ShowOnlyInnerProperties, PCG_NotOverridable))
+	FPVAxialParams AxialParams;
+
+	UPROPERTY(EditAnywhere , Category = "Lateral", meta =(ShowOnlyInnerProperties, PCG_NotOverridable))
+	FPVLateralParams LateralParams;
+	
+	UPROPERTY(EditAnywhere , Category = "Trunk Growth", meta =(ShowOnlyInnerProperties, PCG_NotOverridable))
+	FPVTrunkGrowthParams TrunkGrowth;
+
+	UPROPERTY(EditAnywhere , Category = "Branching", meta =(ShowOnlyInnerProperties, PCG_NotOverridable))
+	FPVBranchingParams Branching;
+
+	UPROPERTY(EditAnywhere , Category = "Branching", meta =(PCG_NotOverridable))
+	bool bBranchBranchingConditionSameAsTrunk = true;
+
+	UPROPERTY(EditAnywhere , Category = "Branching" , meta = (EditCondition = "bBranchBranchingConditionSameAsTrunk == false" , EditConditionHides, PCG_NotOverridable))
+	FPVBranchingParams BranchBranching;
+
+	UPROPERTY(EditAnywhere , Category = "Phototropism", meta =(ShowOnlyInnerProperties, PCG_NotOverridable))
+	FPVPhototropismParams Phototropism;
+
+	UPROPERTY(EditAnywhere , Category = "Phototropism", meta =(PCG_NotOverridable))
+	bool bBranchPhototropismSameAsTrunk = true;
+
+	UPROPERTY(EditAnywhere , Category = "Phototropism" , meta = (EditCondition = "bBranchPhototropismSameAsTrunk == false" , EditConditionHides, PCG_NotOverridable))
+	FPVPhototropismParams BranchPhototropism;
+
+	UPROPERTY(EditAnywhere , Category = "Gravity", meta =(ShowOnlyInnerProperties, PCG_NotOverridable))
+	FPVGrowerGravityParams GravityParams;
+
+	UPROPERTY(EditAnywhere, Category = "Leaf Growth", meta =(ShowOnlyInnerProperties, PCG_NotOverridable))
+	FPVLeafGrowthParams LeafGrowth;
+
+	UPROPERTY(EditAnywhere , Category = "RandomAngle", meta =(ShowOnlyInnerProperties, PCG_NotOverridable))
+	FPVRandomAngleParams RandomAngle;
+
+	UPROPERTY(EditAnywhere , Category = "RandomAngle", meta =(PCG_NotOverridable))
+	bool bBranchRandomAngleSameAsTrunk = true;
+
+	UPROPERTY(EditAnywhere , Category = "RandomAngle", meta = (EditCondition = "bBranchRandomAngleSameAsTrunk == false" , EditConditionHides, PCG_NotOverridable))
+	FPVRandomAngleParams BranchRandomAngle;
+
+	UPROPERTY(EditAnywhere , Category = "Age Senescence", meta =(ShowOnlyInnerProperties, PCG_NotOverridable))
+	FPVAgeSenescenceParams AgeSenescence;
+
+	UPROPERTY(EditAnywhere , Category = "Light Senescence", meta =(ShowOnlyInnerProperties, PCG_NotOverridable))
+	FPVLightSenescenceParams LightSenescence;
+
+	UPROPERTY()
+	FPVGuideParams GuideSettings;
+
+	UPROPERTY(EditAnywhere , Category = "Collider", meta =(PCG_NotOverridable))
+	TArray<FPVColliderParams> ColliderSettings; 
+
+	UPROPERTY(EditAnywhere , Category = "Basic", meta =(PCG_NotOverridable))
+	uint32 RandomSeed = 0;
+
+	void GetPhyllotaxy(FPVGrowerPhyllotaxyParams& OutPhyllotaxy, int Generation) const;
+
+	float GetPhyllotaxyAngle(const FPVGrowerPhyllotaxyParams& InPhyllotaxy) const;
+
+	void GetBranching(FPVBranchingParams& OutBranching, int Generation) const;
+
+	void GetPhototropism(FPVPhototropismParams& OutPhototropism, int Generation) const;
+
+	void GetRandomAngle(FPVRandomAngleParams& OutRandomAngle, int Generation) const;
 };
\ No newline at end of file
diff --git a/Engine/Plugins/Experimental/ProceduralVegetationEditor/Source/ProceduralVegetation/Private/Implementations/PVGrower.h b/Engine/Plugins/Experimental/ProceduralVegetationEditor/Source/ProceduralVegetation/Private/Implementations/PVGrower.h
--- a/Engine/Plugins/Experimental/ProceduralVegetationEditor/Source/ProceduralVegetation/Private/Implementations/PVGrower.h
+++ b/Engine/Plugins/Experimental/ProceduralVegetationEditor/Source/ProceduralVegetation/Private/Implementations/PVGrower.h
@@ -8,101 +8,11 @@
 #include "PVLightDetection.h"
 #include "GeometryCollection/ManagedArrayCollection.h"
 #include "Nodes/PVObjectInteractionSettings.h"
-#include "PVGrower.generated.h"
 
 struct FPVColliderMeshData;
 struct FPVSeedPoint;
 struct FPVPointGravityAttribute;
 
-USTRUCT()
-struct FPVGrowerParams
-{
-	GENERATED_BODY()
-
-public:
-	UPROPERTY(EditAnywhere , Category = "Basic" , meta=(UIMin=1.0f, UIMax=45.0f, PCG_Overridable))
-	uint32 GrowthCycles = 1;
-	
-	UPROPERTY(EditAnywhere , Category = "Phyllotaxy", meta =(ShowOnlyInnerProperties, PCG_NotOverridable))
-	FPVGrowerPhyllotaxyParams Phyllotaxy;
-
-	UPROPERTY(EditAnywhere , Category = "Phyllotaxy", meta =(PCG_NotOverridable))
-	bool bBranchPhyllotaxySameAsTrunk = true;
-
-	UPROPERTY(EditAnywhere , Category = "Phyllotaxy" , meta = (EditCondition = "bBranchPhyllotaxySameAsTrunk == false" , EditConditionHides, PCG_NotOverridable))
-	FPVGrowerPhyllotaxyParams BranchPhyllotaxy;
-	
-	UPROPERTY(EditAnywhere , Category = "Phyllotaxy", meta =(PCG_NotOverridable))
-	bool bLeafPhyllotaxySameAsBranch = true;
-
-	UPROPERTY(EditAnywhere , Category = "Axial", meta =(ShowOnlyInnerProperties, PCG_NotOverridable))
-	FPVAxialParams AxialParams;
-
-	UPROPERTY(EditAnywhere , Category = "Lateral", meta =(ShowOnlyInnerProperties, PCG_NotOverridable))
-	FPVLateralParams LateralParams;
-	
-	UPROPERTY(EditAnywhere , Category = "Trunk Growth", meta =(ShowOnlyInnerProperties, PCG_NotOverridable))
-	FPVTrunkGrowthParams TrunkGrowth;
-
-	UPROPERTY(EditAnywhere , Category = "Branching", meta =(ShowOnlyInnerProperties, PCG_NotOverridable))
-	FPVBranchingParams Branching;
-
-	UPROPERTY(EditAnywhere , Category = "Branching", meta =(PCG_NotOverridable))
-	bool bBranchBranchingConditionSameAsTrunk = true;
-
-	UPROPERTY(EditAnywhere , Category = "Branching" , meta = (EditCondition = "bBranchBranchingConditionSameAsTrunk == false" , EditConditionHides, PCG_NotOverridable))
-	FPVBranchingParams BranchBranching;
-
-	UPROPERTY(EditAnywhere , Category = "Phototropism", meta =(ShowOnlyInnerProperties, PCG_NotOverridable))
-	FPVPhototropismParams Phototropism;
-
-	UPROPERTY(EditAnywhere , Category = "Phototropism", meta =(PCG_NotOverridable))
-	bool bBranchPhototropismSameAsTrunk = true;
-
-	UPROPERTY(EditAnywhere , Category = "Phototropism" , meta = (EditCondition = "bBranchPhototropismSameAsTrunk == false" , EditConditionHides, PCG_NotOverridable))
-	FPVPhototropismParams BranchPhototropism;
-
-	UPROPERTY(EditAnywhere , Category = "Gravity", meta =(ShowOnlyInnerProperties, PCG_NotOverridable))
-	FPVGrowerGravityParams GravityParams;
-
-	UPROPERTY(EditAnywhere, Category = "Leaf Growth", meta =(ShowOnlyInnerProperties, PCG_NotOverridable))
-	FPVLeafGrowthParams LeafGrowth;
-
-	UPROPERTY(EditAnywhere , Category = "RandomAngle", meta =(ShowOnlyInnerProperties, PCG_NotOverridable))
-	FPVRandomAngleParams RandomAngle;
-
-	UPROPERTY(EditAnywhere , Category = "RandomAngle", meta =(PCG_NotOverridable))
-	bool bBranchRandomAngleSameAsTrunk = true;
-
-	UPROPERTY(EditAnywhere , Category = "RandomAngle", meta = (EditCondition = "bBranchRandomAngleSameAsTrunk == false" , EditConditionHides, PCG_NotOverridable))
-	FPVRandomAngleParams BranchRandomAngle;
-
-	UPROPERTY(EditAnywhere , Category = "Age Senescence", meta =(ShowOnlyInnerProperties, PCG_NotOverridable))
-	FPVAgeSenescenceParams AgeSenescence;
-
-	UPROPERTY(EditAnywhere , Category = "Light Senescence", meta =(ShowOnlyInnerProperties, PCG_NotOverridable))
-	FPVLightSenescenceParams LightSenescence;
-
-	UPROPERTY()
-	FPVGuideParams GuideSettings;
-
-	UPROPERTY(EditAnywhere , Category = "Collider", meta =(PCG_NotOverridable))
-	TArray<FPVColliderParams> ColliderSettings; 
-
-	UPROPERTY(EditAnywhere , Category = "Basic", meta =(PCG_NotOverridable))
-	uint32 RandomSeed = 0;
-
-	void GetPhyllotaxy(FPVGrowerPhyllotaxyParams& OutPhyllotaxy, int Generation) const;
-
-	float GetPhyllotaxyAngle(const FPVGrowerPhyllotaxyParams& InPhyllotaxy) const;
-
-	void GetBranching(FPVBranchingParams& OutBranching, int Generation) const;
-
-	void GetPhototropism(FPVPhototropismParams& OutPhototropism, int Generation) const;
-
-	void GetRandomAngle(FPVRandomAngleParams& OutRandomAngle, int Generation) const;
-};
-
 struct FPVGrower
 {
 	static void Grow(const FPVSeedPoint& SeedPoint, const FPVGrowerParams& InGrowerParams, FManagedArrayCollection& InCollection);
diff --git a/Engine/Plugins/Experimental/ProceduralVegetationEditor/Source/ProceduralVegetation/Private/Nodes/GrowerSettings/PVGrowerAgeSenescenceSettings.cpp b/Engine/Plugins/Experimental/ProceduralVegetationEditor/Source/ProceduralVegetation/Private/Nodes/GrowerSettings/PVGrowerAgeSenescenceSettings.cpp
--- a/Engine/Plugins/Experimental/ProceduralVegetationEditor/Source/ProceduralVegetation/Private/Nodes/GrowerSettings/PVGrowerAgeSenescenceSettings.cpp
+++ b/Engine/Plugins/Experimental/ProceduralVegetationEditor/Source/ProceduralVegetation/Private/Nodes/GrowerSettings/PVGrowerAgeSenescenceSettings.cpp
@@ -3,8 +3,10 @@
 #include "PVGrowerAgeSenescenceSettings.h"
 #include "DataTypes/PVGrowthData.h"
 #include "Data/PCGBasePointData.h"
+#include "DataAssets/ProceduralVegetationGrowerPreset.h"
 #include "DataTypes/PVGrowerSettingsData.h"
 #include "Helpers/PCGSettingsHelpers.h"
+#include "Implementations/PVGrower.h"
 
 #define LOCTEXT_NAMESPACE "PVGrowerAgeSenescenceSettings"
 
@@ -15,6 +17,14 @@ FText UPVGrowerAgeSenescenceSettings::GetDefaultNodeTitle() const
 }
 #endif
 
+void UPVGrowerAgeSenescenceSettings::OnPresetChanged()
+{
+	if (Preset)
+	{
+		Params = Preset->GrowthParams.AgeSenescence;
+	}
+}
+
 FPCGDataTypeIdentifier UPVGrowerAgeSenescenceSettings::GetOutputPinTypeIdentifier() const
 {
 	return FPCGDataTypeIdentifier{ FPVDataTypeInfoGrowerAgeSenescence::AsId() };
diff --git a/Engine/Plugins/Experimental/ProceduralVegetationEditor/Source/ProceduralVegetation/Private/Nodes/GrowerSettings/PVGrowerAgeSenescenceSettings.h b/Engine/Plugins/Experimental/ProceduralVegetationEditor/Source/ProceduralVegetation/Private/Nodes/GrowerSettings/PVGrowerAgeSenescenceSettings.h
--- a/Engine/Plugins/Experimental/ProceduralVegetationEditor/Source/ProceduralVegetation/Private/Nodes/GrowerSettings/PVGrowerAgeSenescenceSettings.h
+++ b/Engine/Plugins/Experimental/ProceduralVegetationEditor/Source/ProceduralVegetation/Private/Nodes/GrowerSettings/PVGrowerAgeSenescenceSettings.h
@@ -22,6 +22,8 @@ public:
 protected:
 	virtual FPCGDataTypeIdentifier GetOutputPinTypeIdentifier() const override;
 	virtual FPCGElementPtr CreateElement() const override;
+	virtual void OnPresetChanged() override;
+
 public:
 	UPROPERTY(EditAnywhere, Category="AgeSenescence", meta=(ShowOnlyInnerProperties, PCG_Overridable))
 	FPVAgeSenescenceParams Params;
diff --git a/Engine/Plugins/Experimental/ProceduralVegetationEditor/Source/ProceduralVegetation/Private/Nodes/GrowerSettings/PVGrowerAxialSettings.cpp b/Engine/Plugins/Experimental/ProceduralVegetationEditor/Source/ProceduralVegetation/Private/Nodes/GrowerSettings/PVGrowerAxialSettings.cpp
--- a/Engine/Plugins/Experimental/ProceduralVegetationEditor/Source/ProceduralVegetation/Private/Nodes/GrowerSettings/PVGrowerAxialSettings.cpp
+++ b/Engine/Plugins/Experimental/ProceduralVegetationEditor/Source/ProceduralVegetation/Private/Nodes/GrowerSettings/PVGrowerAxialSettings.cpp
@@ -1,10 +1,11 @@
 ﻿// Copyright Epic Games, Inc. All Rights Reserved.
 
 #include "PVGrowerAxialSettings.h"
-#include "DataTypes/PVGrowthData.h"
 #include "Data/PCGBasePointData.h"
+#include "DataAssets/ProceduralVegetationGrowerPreset.h"
 #include "DataTypes/PVGrowerSettingsData.h"
 #include "Helpers/PCGSettingsHelpers.h"
+#include "Implementations/PVGrower.h"
 
 #define LOCTEXT_NAMESPACE "PVGrowerAxialSettings"
 
@@ -15,6 +16,14 @@ FText UPVGrowerAxialSettings::GetDefaultNodeTitle() const
 }
 #endif
 
+void UPVGrowerAxialSettings::OnPresetChanged()
+{
+	if (Preset)
+	{
+		Params = Preset->GrowthParams.AxialParams;
+	}
+}
+
 FPCGDataTypeIdentifier UPVGrowerAxialSettings::GetOutputPinTypeIdentifier() const
 {
 	return FPCGDataTypeIdentifier{ FPVDataTypeInfoGrowerAxial::AsId() };
diff --git a/Engine/Plugins/Experimental/ProceduralVegetationEditor/Source/ProceduralVegetation/Private/Nodes/GrowerSettings/PVGrowerAxialSettings.h b/Engine/Plugins/Experimental/ProceduralVegetationEditor/Source/ProceduralVegetation/Private/Nodes/GrowerSettings/PVGrowerAxialSettings.h
--- a/Engine/Plugins/Experimental/ProceduralVegetationEditor/Source/ProceduralVegetation/Private/Nodes/GrowerSettings/PVGrowerAxialSettings.h
+++ b/Engine/Plugins/Experimental/ProceduralVegetationEditor/Source/ProceduralVegetation/Private/Nodes/GrowerSettings/PVGrowerAxialSettings.h
@@ -21,6 +21,7 @@ public:
 protected:
 	virtual FPCGDataTypeIdentifier GetOutputPinTypeIdentifier() const override;
 	virtual FPCGElementPtr CreateElement() const override;
+	virtual void OnPresetChanged() override;
 
 public:
 	UPROPERTY(EditAnywhere, Category="AxialSettings", meta=(ShowOnlyInnerProperties, PCG_Overridable))
diff --git a/Engine/Plugins/Experimental/ProceduralVegetationEditor/Source/ProceduralVegetation/Private/Nodes/GrowerSettings/PVGrowerBaseSettings.cpp b/Engine/Plugins/Experimental/ProceduralVegetationEditor/Source/ProceduralVegetation/Private/Nodes/GrowerSettings/PVGrowerBaseSettings.cpp
--- a/Engine/Plugins/Experimental/ProceduralVegetationEditor/Source/ProceduralVegetation/Private/Nodes/GrowerSettings/PVGrowerBaseSettings.cpp
+++ b/Engine/Plugins/Experimental/ProceduralVegetationEditor/Source/ProceduralVegetation/Private/Nodes/GrowerSettings/PVGrowerBaseSettings.cpp
@@ -12,6 +12,18 @@ FText UPVGrowerBaseSettings::GetNodeTooltipText() const
 		"\n\nPress Ctrl + L to lock/unlock node output"
 	); 
 }
+
+void UPVGrowerBaseSettings::PostEditChangeProperty(FPropertyChangedEvent& PropertyChangedEvent)
+{
+	if (PropertyChangedEvent.Property &&
+		(PropertyChangedEvent.Property->GetFName() == GET_MEMBER_NAME_CHECKED(UPVGrowerBaseSettings, Preset) ||
+		PropertyChangedEvent.Property->GetFName() == GET_MEMBER_NAME_CHECKED(UPVGrowerBaseSettings, PickFrom)))
+	{
+		OnPresetChanged();
+	}
+	
+	Super::PostEditChangeProperty(PropertyChangedEvent);
+}
 #endif
 
 TArray<FPCGPinProperties> UPVGrowerBaseSettings::InputPinProperties() const
diff --git a/Engine/Plugins/Experimental/ProceduralVegetationEditor/Source/ProceduralVegetation/Private/Nodes/GrowerSettings/PVGrowerBaseSettings.h b/Engine/Plugins/Experimental/ProceduralVegetationEditor/Source/ProceduralVegetation/Private/Nodes/GrowerSettings/PVGrowerBaseSettings.h
--- a/Engine/Plugins/Experimental/ProceduralVegetationEditor/Source/ProceduralVegetation/Private/Nodes/GrowerSettings/PVGrowerBaseSettings.h
+++ b/Engine/Plugins/Experimental/ProceduralVegetationEditor/Source/ProceduralVegetation/Private/Nodes/GrowerSettings/PVGrowerBaseSettings.h
@@ -5,6 +5,10 @@
 #include "Nodes/PVBaseSettings.h"
 #include "PVGrowerBaseSettings.generated.h"
 
+class UProceduralVegetationGrowerPreset;
+
+enum class EPVPresetTarget : uint8;
+
 UCLASS(BlueprintType, Abstract, ClassGroup = (Procedural))
 class UPVGrowerBaseSettings : public UPVBaseSettings
 {
@@ -20,4 +24,22 @@ public:
 
 protected:
 	virtual TArray<FPCGPinProperties> InputPinProperties() const override;
-};
+
+#if WITH_EDITOR
+	
+	UFUNCTION()
+	virtual bool ShowPresetPickFrom() { return false; }
+
+	virtual bool HasValidPreset() const { return Preset != nullptr; }
+	
+	virtual void PostEditChangeProperty(FPropertyChangedEvent& PropertyChangedEvent) override;
+#endif
+
+	virtual void OnPresetChanged() PURE_VIRTUAL(UPVGrowerBaseSettings::OnPresetChanged, );
+	
+	UPROPERTY(EditAnywhere, Category = Preset)
+	TObjectPtr<UProceduralVegetationGrowerPreset> Preset;
+
+	UPROPERTY(EditAnywhere, Category = Preset, meta=(EditCondition = "ShowPresetPickFrom()", EditConditionHides))
+	EPVPresetTarget PickFrom;
+};
\ No newline at end of file
diff --git a/Engine/Plugins/Experimental/ProceduralVegetationEditor/Source/ProceduralVegetation/Private/Nodes/GrowerSettings/PVGrowerBranchingSettings.cpp b/Engine/Plugins/Experimental/ProceduralVegetationEditor/Source/ProceduralVegetation/Private/Nodes/GrowerSettings/PVGrowerBranchingSettings.cpp
--- a/Engine/Plugins/Experimental/ProceduralVegetationEditor/Source/ProceduralVegetation/Private/Nodes/GrowerSettings/PVGrowerBranchingSettings.cpp
+++ b/Engine/Plugins/Experimental/ProceduralVegetationEditor/Source/ProceduralVegetation/Private/Nodes/GrowerSettings/PVGrowerBranchingSettings.cpp
@@ -3,8 +3,10 @@
 #include "PVGrowerBranchingSettings.h"
 #include "DataTypes/PVGrowthData.h"
 #include "Data/PCGBasePointData.h"
+#include "DataAssets/ProceduralVegetationGrowerPreset.h"
 #include "DataTypes/PVGrowerSettingsData.h"
 #include "Helpers/PCGSettingsHelpers.h"
+#include "Implementations/PVGrower.h"
 
 #define LOCTEXT_NAMESPACE "PVGrowerBranchingSettings"
 
@@ -30,6 +32,21 @@ FPCGElementPtr UPVGrowerBranchingSettings::CreateElement() const
 	return MakeShared<FPVGrowerBranchingElement>();
 }
 
+void UPVGrowerBranchingSettings::OnPresetChanged()
+{
+	if (Preset)
+	{
+		if (PickFrom == EPVPresetTarget::Trunk)
+		{
+			Preset->GrowthParams.GetBranching(ParamsWithTargets.Params, 1);
+		}
+		else if (PickFrom == EPVPresetTarget::Branch)
+		{
+			Preset->GrowthParams.GetBranching(ParamsWithTargets.Params, 2);
+		}
+	}
+}
+
 bool FPVGrowerBranchingElement::ExecuteInternal(FPCGContext* InContext) const
 {
 	TRACE_CPUPROFILER_EVENT_SCOPE(FPVGrowerBranchingElement::Execute);
diff --git a/Engine/Plugins/Experimental/ProceduralVegetationEditor/Source/ProceduralVegetation/Private/Nodes/GrowerSettings/PVGrowerBranchingSettings.h b/Engine/Plugins/Experimental/ProceduralVegetationEditor/Source/ProceduralVegetation/Private/Nodes/GrowerSettings/PVGrowerBranchingSettings.h
--- a/Engine/Plugins/Experimental/ProceduralVegetationEditor/Source/ProceduralVegetation/Private/Nodes/GrowerSettings/PVGrowerBranchingSettings.h
+++ b/Engine/Plugins/Experimental/ProceduralVegetationEditor/Source/ProceduralVegetation/Private/Nodes/GrowerSettings/PVGrowerBranchingSettings.h
@@ -22,6 +22,11 @@ public:
 protected:
 	virtual FPCGDataTypeIdentifier GetOutputPinTypeIdentifier() const override;
 	virtual FPCGElementPtr CreateElement() const override;
+	virtual void OnPresetChanged() override;
+#if WITH_EDITOR
+	virtual bool ShowPresetPickFrom() override { return HasValidPreset(); }
+#endif
+
 public:
 	UPROPERTY(EditAnywhere, Category="BranchingSettings", meta=(ShowOnlyInnerProperties, PCG_Overridable))
 	FPVGrowerBranchingWithTargets ParamsWithTargets;
diff --git a/Engine/Plugins/Experimental/ProceduralVegetationEditor/Source/ProceduralVegetation/Private/Nodes/GrowerSettings/PVGrowerGravitySettings.cpp b/Engine/Plugins/Experimental/ProceduralVegetationEditor/Source/ProceduralVegetation/Private/Nodes/GrowerSettings/PVGrowerGravitySettings.cpp
--- a/Engine/Plugins/Experimental/ProceduralVegetationEditor/Source/ProceduralVegetation/Private/Nodes/GrowerSettings/PVGrowerGravitySettings.cpp
+++ b/Engine/Plugins/Experimental/ProceduralVegetationEditor/Source/ProceduralVegetation/Private/Nodes/GrowerSettings/PVGrowerGravitySettings.cpp
@@ -3,8 +3,10 @@
 #include "PVGrowerGravitySettings.h"
 #include "DataTypes/PVGrowthData.h"
 #include "Data/PCGBasePointData.h"
+#include "DataAssets/ProceduralVegetationGrowerPreset.h"
 #include "DataTypes/PVGrowerSettingsData.h"
 #include "Helpers/PCGSettingsHelpers.h"
+#include "Implementations/PVGrower.h"
 
 #define LOCTEXT_NAMESPACE "PVGrowerGravitySettings"
 
@@ -15,6 +17,14 @@ FText UPVGrowerGravitySettings::GetDefaultNodeTitle() const
 }
 #endif
 
+void UPVGrowerGravitySettings::OnPresetChanged()
+{
+	if (Preset)
+	{
+		Params = Preset->GrowthParams.GravityParams;
+	}
+}
+
 FPCGDataTypeIdentifier UPVGrowerGravitySettings::GetOutputPinTypeIdentifier() const
 {
 	return FPCGDataTypeIdentifier{ FPVDataTypeInfoGrowerGravity::AsId() };
diff --git a/Engine/Plugins/Experimental/ProceduralVegetationEditor/Source/ProceduralVegetation/Private/Nodes/GrowerSettings/PVGrowerGravitySettings.h b/Engine/Plugins/Experimental/ProceduralVegetationEditor/Source/ProceduralVegetation/Private/Nodes/GrowerSettings/PVGrowerGravitySettings.h
--- a/Engine/Plugins/Experimental/ProceduralVegetationEditor/Source/ProceduralVegetation/Private/Nodes/GrowerSettings/PVGrowerGravitySettings.h
+++ b/Engine/Plugins/Experimental/ProceduralVegetationEditor/Source/ProceduralVegetation/Private/Nodes/GrowerSettings/PVGrowerGravitySettings.h
@@ -21,6 +21,7 @@ public:
 protected:
 	virtual FPCGDataTypeIdentifier GetOutputPinTypeIdentifier() const override;
 	virtual FPCGElementPtr CreateElement() const override;
+	virtual void OnPresetChanged() override;
 public:
 
 	UPROPERTY(EditAnywhere, Category="GravitySetttings", meta=(ShowOnlyInnerProperties, PCG_Overridable))
diff --git a/Engine/Plugins/Experimental/ProceduralVegetationEditor/Source/ProceduralVegetation/Private/Nodes/GrowerSettings/PVGrowerLateralSettings.cpp b/Engine/Plugins/Experimental/ProceduralVegetationEditor/Source/ProceduralVegetation/Private/Nodes/GrowerSettings/PVGrowerLateralSettings.cpp
--- a/Engine/Plugins/Experimental/ProceduralVegetationEditor/Source/ProceduralVegetation/Private/Nodes/GrowerSettings/PVGrowerLateralSettings.cpp
+++ b/Engine/Plugins/Experimental/ProceduralVegetationEditor/Source/ProceduralVegetation/Private/Nodes/GrowerSettings/PVGrowerLateralSettings.cpp
@@ -3,8 +3,10 @@
 #include "PVGrowerLateralSettings.h"
 #include "DataTypes/PVGrowthData.h"
 #include "Data/PCGBasePointData.h"
+#include "DataAssets/ProceduralVegetationGrowerPreset.h"
 #include "DataTypes/PVGrowerSettingsData.h"
 #include "Helpers/PCGSettingsHelpers.h"
+#include "Implementations/PVGrower.h"
 
 #define LOCTEXT_NAMESPACE "PVGrowerLateralSettings"
 
@@ -15,6 +17,13 @@ FText UPVGrowerLateralSettings::GetDefaultNodeTitle() const
 }
 #endif
 
+void UPVGrowerLateralSettings::OnPresetChanged()
+{
+	if (Preset)
+	{
+		Params = Preset->GrowthParams.LateralParams;
+	}
+}
 FPCGDataTypeIdentifier UPVGrowerLateralSettings::GetOutputPinTypeIdentifier() const
 {
 	return FPCGDataTypeIdentifier{ FPVDataTypeInfoGrowerLateral::AsId() };
diff --git a/Engine/Plugins/Experimental/ProceduralVegetationEditor/Source/ProceduralVegetation/Private/Nodes/GrowerSettings/PVGrowerLateralSettings.h b/Engine/Plugins/Experimental/ProceduralVegetationEditor/Source/ProceduralVegetation/Private/Nodes/GrowerSettings/PVGrowerLateralSettings.h
--- a/Engine/Plugins/Experimental/ProceduralVegetationEditor/Source/ProceduralVegetation/Private/Nodes/GrowerSettings/PVGrowerLateralSettings.h
+++ b/Engine/Plugins/Experimental/ProceduralVegetationEditor/Source/ProceduralVegetation/Private/Nodes/GrowerSettings/PVGrowerLateralSettings.h
@@ -20,6 +20,8 @@ public:
 protected:
 	virtual FPCGDataTypeIdentifier GetOutputPinTypeIdentifier() const override;
 	virtual FPCGElementPtr CreateElement() const override;
+	virtual void OnPresetChanged() override;
+	
 public:
 
 	UPROPERTY(EditAnywhere, Category="LateralSettings", meta=(ShowOnlyInnerProperties, PCG_Overridable))
diff --git a/Engine/Plugins/Experimental/ProceduralVegetationEditor/Source/ProceduralVegetation/Private/Nodes/GrowerSettings/PVGrowerLeafGrowthSettings.cpp b/Engine/Plugins/Experimental/ProceduralVegetationEditor/Source/ProceduralVegetation/Private/Nodes/GrowerSettings/PVGrowerLeafGrowthSettings.cpp
--- a/Engine/Plugins/Experimental/ProceduralVegetationEditor/Source/ProceduralVegetation/Private/Nodes/GrowerSettings/PVGrowerLeafGrowthSettings.cpp
+++ b/Engine/Plugins/Experimental/ProceduralVegetationEditor/Source/ProceduralVegetation/Private/Nodes/GrowerSettings/PVGrowerLeafGrowthSettings.cpp
@@ -4,9 +4,11 @@
 #include "DataTypes/PVData.h"
 #include "DataTypes/PVGrowthData.h"
 #include "Data/PCGBasePointData.h"
+#include "DataAssets/ProceduralVegetationGrowerPreset.h"
 #include "DataTypes/PVGrowerSettingsData.h"
 #include "GeometryCollection/ManagedArrayCollection.h"
 #include "Helpers/PCGSettingsHelpers.h"
+#include "Implementations/PVGrower.h"
 
 #define LOCTEXT_NAMESPACE "PVGrowerLeafGrowthSettings"
 
@@ -39,6 +41,16 @@ FPCGElementPtr UPVGrowerLeafGrowthSettings::CreateElement() const
 	return MakeShared<FPVGrowerLeafGrowthElement>();
 }
 
+void UPVGrowerLeafGrowthSettings::OnPresetChanged()
+{
+	if (Preset)
+	{
+		Params = Preset->GrowthParams.LeafGrowth;
+	}
+}
+
+
+
 bool FPVGrowerLeafGrowthElement::ExecuteInternal(FPCGContext* InContext) const
 {
 	TRACE_CPUPROFILER_EVENT_SCOPE(FPVGrowerLeafGrowthElement::Execute);
diff --git a/Engine/Plugins/Experimental/ProceduralVegetationEditor/Source/ProceduralVegetation/Private/Nodes/GrowerSettings/PVGrowerLeafGrowthSettings.h b/Engine/Plugins/Experimental/ProceduralVegetationEditor/Source/ProceduralVegetation/Private/Nodes/GrowerSettings/PVGrowerLeafGrowthSettings.h
--- a/Engine/Plugins/Experimental/ProceduralVegetationEditor/Source/ProceduralVegetation/Private/Nodes/GrowerSettings/PVGrowerLeafGrowthSettings.h
+++ b/Engine/Plugins/Experimental/ProceduralVegetationEditor/Source/ProceduralVegetation/Private/Nodes/GrowerSettings/PVGrowerLeafGrowthSettings.h
@@ -2,14 +2,14 @@
 #pragma once
 
 #include "CoreMinimal.h"
+#include "PVGrowerBaseSettings.h"
 #include "DataTypes/PVGrowerParams.h"
-#include "Nodes/PVBaseSettings.h"
 #include "Implementations/PVGravity.h"
 #include "PVGrowerLeafGrowthSettings.generated.h"
 
 
 UCLASS(BlueprintType, ClassGroup = (Procedural))
-class UPVGrowerLeafGrowthSettings : public UPVBaseSettings
+class UPVGrowerLeafGrowthSettings : public UPVGrowerBaseSettings
 {
 	GENERATED_BODY()
 
@@ -26,8 +26,8 @@ public:
 protected:
 	virtual TArray<FPCGPinProperties> InputPinProperties() const override;
 	virtual FPCGDataTypeIdentifier GetOutputPinTypeIdentifier() const override;
-
 	virtual FPCGElementPtr CreateElement() const override;
+	virtual void OnPresetChanged() override;
 
 public:
 
diff --git a/Engine/Plugins/Experimental/ProceduralVegetationEditor/Source/ProceduralVegetation/Private/Nodes/GrowerSettings/PVGrowerPhototropismSettings.cpp b/Engine/Plugins/Experimental/ProceduralVegetationEditor/Source/ProceduralVegetation/Private/Nodes/GrowerSettings/PVGrowerPhototropismSettings.cpp
--- a/Engine/Plugins/Experimental/ProceduralVegetationEditor/Source/ProceduralVegetation/Private/Nodes/GrowerSettings/PVGrowerPhototropismSettings.cpp
+++ b/Engine/Plugins/Experimental/ProceduralVegetationEditor/Source/ProceduralVegetation/Private/Nodes/GrowerSettings/PVGrowerPhototropismSettings.cpp
@@ -3,8 +3,10 @@
 #include "PVGrowerPhototropismSettings.h"
 #include "DataTypes/PVGrowthData.h"
 #include "Data/PCGBasePointData.h"
+#include "DataAssets/ProceduralVegetationGrowerPreset.h"
 #include "DataTypes/PVGrowerSettingsData.h"
 #include "Helpers/PCGSettingsHelpers.h"
+#include "Implementations/PVGrower.h"
 
 #define LOCTEXT_NAMESPACE "PVGrowerPhototropismSettings"
 
@@ -30,6 +32,21 @@ FPCGElementPtr UPVGrowerPhototropismSettings::CreateElement() const
 	return MakeShared<FPVGrowerPhototropismElement>();
 }
 
+void UPVGrowerPhototropismSettings::OnPresetChanged()
+{
+	if (Preset)
+	{
+		if (PickFrom == EPVPresetTarget::Trunk)
+		{
+			Preset->GrowthParams.GetPhototropism(ParamsWithTargets.Params, 1);
+		}
+		else if (PickFrom == EPVPresetTarget::Branch)
+		{
+			Preset->GrowthParams.GetPhototropism(ParamsWithTargets.Params, 2);
+		}
+	}
+}
+
 bool FPVGrowerPhototropismElement::ExecuteInternal(FPCGContext* InContext) const
 {
 	TRACE_CPUPROFILER_EVENT_SCOPE(FPVGrowerPhototropismElement::Execute);
diff --git a/Engine/Plugins/Experimental/ProceduralVegetationEditor/Source/ProceduralVegetation/Private/Nodes/GrowerSettings/PVGrowerPhototropismSettings.h b/Engine/Plugins/Experimental/ProceduralVegetationEditor/Source/ProceduralVegetation/Private/Nodes/GrowerSettings/PVGrowerPhototropismSettings.h
--- a/Engine/Plugins/Experimental/ProceduralVegetationEditor/Source/ProceduralVegetation/Private/Nodes/GrowerSettings/PVGrowerPhototropismSettings.h
+++ b/Engine/Plugins/Experimental/ProceduralVegetationEditor/Source/ProceduralVegetation/Private/Nodes/GrowerSettings/PVGrowerPhototropismSettings.h
@@ -22,7 +22,11 @@ public:
 protected:
 	virtual FPCGDataTypeIdentifier GetOutputPinTypeIdentifier() const override;
 	virtual FPCGElementPtr CreateElement() const override;
-
+	virtual void OnPresetChanged() override;
+#if WITH_EDITOR
+	virtual bool ShowPresetPickFrom() override { return Preset != nullptr; }
+#endif
+	
 public:
 
 	UPROPERTY(EditAnywhere, Category="PhototropismSettings", meta=(ShowOnlyInnerProperties, PCG_Overridable))
diff --git a/Engine/Plugins/Experimental/ProceduralVegetationEditor/Source/ProceduralVegetation/Private/Nodes/GrowerSettings/PVGrowerPhyllotaxySettings.cpp b/Engine/Plugins/Experimental/ProceduralVegetationEditor/Source/ProceduralVegetation/Private/Nodes/GrowerSettings/PVGrowerPhyllotaxySettings.cpp
--- a/Engine/Plugins/Experimental/ProceduralVegetationEditor/Source/ProceduralVegetation/Private/Nodes/GrowerSettings/PVGrowerPhyllotaxySettings.cpp
+++ b/Engine/Plugins/Experimental/ProceduralVegetationEditor/Source/ProceduralVegetation/Private/Nodes/GrowerSettings/PVGrowerPhyllotaxySettings.cpp
@@ -3,8 +3,10 @@
 #include "PVGrowerPhyllotaxySettings.h"
 #include "DataTypes/PVGrowthData.h"
 #include "Data/PCGBasePointData.h"
+#include "DataAssets/ProceduralVegetationGrowerPreset.h"
 #include "DataTypes/PVGrowerSettingsData.h"
 #include "Helpers/PCGSettingsHelpers.h"
+#include "Implementations/PVGrower.h"
 
 #define LOCTEXT_NAMESPACE "PVPhyllotaxySettings"
 
@@ -30,6 +32,21 @@ FPCGElementPtr UPVGrowerPhyllotaxySettings::CreateElement() const
 	return MakeShared<FPVGrowerPhyllotaxyElement>();
 }
 
+void UPVGrowerPhyllotaxySettings::OnPresetChanged()
+{
+	if (Preset)
+	{
+		if (PickFrom == EPVPresetTarget::Trunk)
+		{
+			Preset->GrowthParams.GetPhyllotaxy(ParamsWithTargets.Params, 1);
+		}
+		else if (PickFrom == EPVPresetTarget::Branch)
+		{
+			Preset->GrowthParams.GetPhyllotaxy(ParamsWithTargets.Params, 2);
+		}
+	}
+}
+
 bool FPVGrowerPhyllotaxyElement::ExecuteInternal(FPCGContext* InContext) const
 {
 	TRACE_CPUPROFILER_EVENT_SCOPE(FPVPhylotaxyElement::Execute);
diff --git a/Engine/Plugins/Experimental/ProceduralVegetationEditor/Source/ProceduralVegetation/Private/Nodes/GrowerSettings/PVGrowerPhyllotaxySettings.h b/Engine/Plugins/Experimental/ProceduralVegetationEditor/Source/ProceduralVegetation/Private/Nodes/GrowerSettings/PVGrowerPhyllotaxySettings.h
--- a/Engine/Plugins/Experimental/ProceduralVegetationEditor/Source/ProceduralVegetation/Private/Nodes/GrowerSettings/PVGrowerPhyllotaxySettings.h
+++ b/Engine/Plugins/Experimental/ProceduralVegetationEditor/Source/ProceduralVegetation/Private/Nodes/GrowerSettings/PVGrowerPhyllotaxySettings.h
@@ -6,7 +6,6 @@
 #include "DataTypes/PVGrowerParams.h"
 #include "PVGrowerPhyllotaxySettings.generated.h"
 
-
 UCLASS(BlueprintType, ClassGroup = (Procedural))
 class UPVGrowerPhyllotaxySettings : public UPVGrowerBaseSettings
 {
@@ -23,6 +22,11 @@ public:
 protected:
 	virtual FPCGDataTypeIdentifier GetOutputPinTypeIdentifier() const override;
 	virtual FPCGElementPtr CreateElement() const override;
+	virtual void OnPresetChanged() override;
+#if WITH_EDITOR
+	virtual bool ShowPresetPickFrom() override { return HasValidPreset(); }
+#endif
+
 public:
 	UPROPERTY(EditAnywhere, Category="PhyllotaxySettings", meta=(ShowOnlyInnerProperties, PCG_Overridable))
 	FPVGrowerPhyllotaxyWithTargets ParamsWithTargets;
diff --git a/Engine/Plugins/Experimental/ProceduralVegetationEditor/Source/ProceduralVegetation/Private/Nodes/GrowerSettings/PVGrowerRandomAngleSettings.cpp b/Engine/Plugins/Experimental/ProceduralVegetationEditor/Source/ProceduralVegetation/Private/Nodes/GrowerSettings/PVGrowerRandomAngleSettings.cpp
--- a/Engine/Plugins/Experimental/ProceduralVegetationEditor/Source/ProceduralVegetation/Private/Nodes/GrowerSettings/PVGrowerRandomAngleSettings.cpp
+++ b/Engine/Plugins/Experimental/ProceduralVegetationEditor/Source/ProceduralVegetation/Private/Nodes/GrowerSettings/PVGrowerRandomAngleSettings.cpp
@@ -3,8 +3,10 @@
 #include "PVGrowerRandomAngleSettings.h"
 #include "DataTypes/PVGrowthData.h"
 #include "Data/PCGBasePointData.h"
+#include "DataAssets/ProceduralVegetationGrowerPreset.h"
 #include "DataTypes/PVGrowerSettingsData.h"
 #include "Helpers/PCGSettingsHelpers.h"
+#include "Implementations/PVGrower.h"
 
 #define LOCTEXT_NAMESPACE "PVGrowerRandomAngleSettings"
 
@@ -30,6 +32,21 @@ FPCGElementPtr UPVGrowerRandomAngleSettings::CreateElement() const
 	return MakeShared<FPVGrowerAngleElement>();
 }
 
+void UPVGrowerRandomAngleSettings::OnPresetChanged()
+{
+	if (Preset)
+	{
+		if (PickFrom == EPVPresetTarget::Trunk)
+		{
+			Preset->GrowthParams.GetRandomAngle(ParamsWithTargets.Params, 1);
+		}
+		else if (PickFrom == EPVPresetTarget::Branch)
+		{
+			Preset->GrowthParams.GetRandomAngle(ParamsWithTargets.Params, 2);
+		}
+	}
+}
+
 bool FPVGrowerAngleElement::ExecuteInternal(FPCGContext* InContext) const
 {
 	TRACE_CPUPROFILER_EVENT_SCOPE(FPVGrowerAngleElement::Execute);
diff --git a/Engine/Plugins/Experimental/ProceduralVegetationEditor/Source/ProceduralVegetation/Private/Nodes/GrowerSettings/PVGrowerRandomAngleSettings.h b/Engine/Plugins/Experimental/ProceduralVegetationEditor/Source/ProceduralVegetation/Private/Nodes/GrowerSettings/PVGrowerRandomAngleSettings.h
--- a/Engine/Plugins/Experimental/ProceduralVegetationEditor/Source/ProceduralVegetation/Private/Nodes/GrowerSettings/PVGrowerRandomAngleSettings.h
+++ b/Engine/Plugins/Experimental/ProceduralVegetationEditor/Source/ProceduralVegetation/Private/Nodes/GrowerSettings/PVGrowerRandomAngleSettings.h
@@ -6,7 +6,6 @@
 #include "DataTypes/PVGrowerParams.h"
 #include "PVGrowerRandomAngleSettings.generated.h"
 
-
 UCLASS(BlueprintType, ClassGroup = (Procedural))
 class UPVGrowerRandomAngleSettings : public UPVGrowerBaseSettings
 {
@@ -23,6 +22,10 @@ public:
 protected:
 	virtual FPCGDataTypeIdentifier GetOutputPinTypeIdentifier() const override;
 	virtual FPCGElementPtr CreateElement() const override;
+	virtual void OnPresetChanged() override;
+#if WITH_EDITOR
+	virtual bool ShowPresetPickFrom() override { return HasValidPreset(); }
+#endif
 
 public:
 
diff --git a/Engine/Plugins/Experimental/ProceduralVegetationEditor/Source/ProceduralVegetation/Private/Nodes/GrowerSettings/PVGrowerTrunkGrowthSettings.cpp b/Engine/Plugins/Experimental/ProceduralVegetationEditor/Source/ProceduralVegetation/Private/Nodes/GrowerSettings/PVGrowerTrunkGrowthSettings.cpp
--- a/Engine/Plugins/Experimental/ProceduralVegetationEditor/Source/ProceduralVegetation/Private/Nodes/GrowerSettings/PVGrowerTrunkGrowthSettings.cpp
+++ b/Engine/Plugins/Experimental/ProceduralVegetationEditor/Source/ProceduralVegetation/Private/Nodes/GrowerSettings/PVGrowerTrunkGrowthSettings.cpp
@@ -1,10 +1,11 @@
 ﻿// Copyright Epic Games, Inc. All Rights Reserved.
 
 #include "PVGrowerTrunkGrowthSettings.h"
-#include "DataTypes/PVGrowthData.h"
 #include "Data/PCGBasePointData.h"
+#include "DataAssets/ProceduralVegetationGrowerPreset.h"
 #include "DataTypes/PVGrowerSettingsData.h"
 #include "Helpers/PCGSettingsHelpers.h"
+#include "Implementations/PVGrower.h"
 
 #define LOCTEXT_NAMESPACE "PVGrowerTrunkGrowthSettings"
 
@@ -15,6 +16,14 @@ FText UPVGrowerTrunkGrowthSettings::GetDefaultNodeTitle() const
 }
 #endif
 
+void UPVGrowerTrunkGrowthSettings::OnPresetChanged()
+{
+	if (Preset)
+	{
+		Params = Preset->GrowthParams.TrunkGrowth;
+	}
+}
+
 FPCGDataTypeIdentifier UPVGrowerTrunkGrowthSettings::GetOutputPinTypeIdentifier() const
 {
 	return FPCGDataTypeIdentifier{ FPVDataTypeInfoGrowerTrunkGrowth::AsId() };
@@ -25,6 +34,8 @@ FPCGElementPtr UPVGrowerTrunkGrowthSettings::CreateElement() const
 	return MakeShared<FPVGrowerTrunkGrowthElement>();
 }
 
+
+
 bool FPVGrowerTrunkGrowthElement::ExecuteInternal(FPCGContext* InContext) const
 {
 	TRACE_CPUPROFILER_EVENT_SCOPE(FPVGrowerTrunkGrowthElement::Execute);
diff --git a/Engine/Plugins/Experimental/ProceduralVegetationEditor/Source/ProceduralVegetation/Private/Nodes/GrowerSettings/PVGrowerTrunkGrowthSettings.h b/Engine/Plugins/Experimental/ProceduralVegetationEditor/Source/ProceduralVegetation/Private/Nodes/GrowerSettings/PVGrowerTrunkGrowthSettings.h
--- a/Engine/Plugins/Experimental/ProceduralVegetationEditor/Source/ProceduralVegetation/Private/Nodes/GrowerSettings/PVGrowerTrunkGrowthSettings.h
+++ b/Engine/Plugins/Experimental/ProceduralVegetationEditor/Source/ProceduralVegetation/Private/Nodes/GrowerSettings/PVGrowerTrunkGrowthSettings.h
@@ -20,6 +20,7 @@ public:
 protected:
 	virtual FPCGDataTypeIdentifier GetOutputPinTypeIdentifier() const override;
 	virtual FPCGElementPtr CreateElement() const override;
+	virtual void OnPresetChanged() override;
 
 public:
 
diff --git a/Engine/Plugins/Experimental/ProceduralVegetationEditor/Source/ProceduralVegetation/Private/Nodes/PVGrowerSettings.cpp b/Engine/Plugins/Experimental/ProceduralVegetationEditor/Source/ProceduralVegetation/Private/Nodes/PVGrowerSettings.cpp
--- a/Engine/Plugins/Experimental/ProceduralVegetationEditor/Source/ProceduralVegetation/Private/Nodes/PVGrowerSettings.cpp
+++ b/Engine/Plugins/Experimental/ProceduralVegetationEditor/Source/ProceduralVegetation/Private/Nodes/PVGrowerSettings.cpp
@@ -2,11 +2,13 @@
 
 #include "PVGrowerSettings.h"
 
+#include "PCGAssetExporterUtils.h"
 #include "ProceduralVegetationModule.h"
 #include "GrowerSettings/PVGrowerPhyllotaxySettings.h"
 #include "DataTypes/PVData.h"
 #include "DataTypes/PVGrowthData.h"
 #include "Data/PCGBasePointData.h"
+#include "DataAssets/ProceduralVegetationGrowerPreset.h"
 #include "DataTypes/PVGrowerSettingsData.h"
 #include "GeometryCollection/ManagedArrayCollection.h"
 #include "Implementations/PVSeedGenerator.h"
@@ -55,6 +57,21 @@ TMap<UObject*, FTransform> UPVGrowerSettings::GetViewportObjects() const
 
 	return ViewportObjects;
 }
+
+void UPVGrowerSettings::PostEditChangeProperty(FPropertyChangedEvent& PropertyChangedEvent)
+{
+	if (PropertyChangedEvent.Property
+		&& PropertyChangedEvent.Property->GetFName() == GET_MEMBER_NAME_CHECKED(UPVGrowerSettings, Preset))
+	{
+		if (Preset)
+		{
+			GrowerParams = Preset->GrowthParams;
+		}
+	}
+	
+	Super::PostEditChangeProperty(PropertyChangedEvent);
+}
+
 #endif
 
 TArray<FPCGPinProperties> UPVGrowerSettings::InputPinProperties() const
@@ -125,6 +142,31 @@ FPCGElementPtr UPVGrowerSettings::CreateElement() const
 	return MakeShared<FPVGrowerElement>();
 }
 
+void UPVGrowerSettings::SaveToPreset() const
+{
+	FPCGAssetExporterParameters ExportParams;
+	ExportParams.bOpenSaveDialog = true;
+	
+	UPCGAssetExporterUtils::CreateAsset<UProceduralVegetationGrowerPreset>(ExportParams, [this](const FString&, UObject* Asset)
+	{
+		if (UProceduralVegetationGrowerPreset* GrowerPreset = Cast<UProceduralVegetationGrowerPreset>(Asset))
+		{
+			GrowerPreset->GrowthParams = GrowerParams;	
+		}
+
+		return true;
+		
+	}, nullptr);
+}
+
+void UPVGrowerSettings::ApplyToSelectedPreset() const
+{
+	if (Preset)
+	{
+		Preset->GrowthParams = GrowerParams;
+	}
+}
+
 bool FPVGrowerElement::ExecuteInternal(FPCGContext* InContext) const
 {
 	TRACE_CPUPROFILER_EVENT_SCOPE(FPVGrowerElement::Execute);
diff --git a/Engine/Plugins/Experimental/ProceduralVegetationEditor/Source/ProceduralVegetation/Private/Nodes/PVGrowerSettings.h b/Engine/Plugins/Experimental/ProceduralVegetationEditor/Source/ProceduralVegetation/Private/Nodes/PVGrowerSettings.h
--- a/Engine/Plugins/Experimental/ProceduralVegetationEditor/Source/ProceduralVegetation/Private/Nodes/PVGrowerSettings.h
+++ b/Engine/Plugins/Experimental/ProceduralVegetationEditor/Source/ProceduralVegetation/Private/Nodes/PVGrowerSettings.h
@@ -8,6 +8,8 @@
 #include "Implementations/PVGrower.h"
 #include "PVGrowerSettings.generated.h"
 
+class UProceduralVegetationGrowerPreset;
+
 UCLASS(BlueprintType, ClassGroup = (Procedural))
 class UPVGrowerSettings : public UPVBaseSettings
 {
@@ -31,7 +33,21 @@ protected:
 	virtual FPCGDataTypeIdentifier GetOutputPinTypeIdentifier() const override;
 	virtual FPCGElementPtr CreateElement() const override;
 
+#if WITH_EDITOR
+	virtual void PostEditChangeProperty(FPropertyChangedEvent& PropertyChangedEvent) override;
+#endif
+
 public:
+
+	UPROPERTY(EditAnywhere, Category = Preset)
+	TObjectPtr<UProceduralVegetationGrowerPreset> Preset;
+
+	UFUNCTION(CallInEditor, Category = Preset, meta = (DisplayName = "Save to preset"))
+	void SaveToPreset() const;
+
+	UFUNCTION(CallInEditor, Category = Preset, meta = (DisplayName = "Apply to selected preset"))
+	void ApplyToSelectedPreset() const;
+	
 	UPROPERTY(EditAnywhere, Category = Settings, meta=( ShowOnlyInnerProperties, PCG_Overridable))
 	FPVGrowerParams GrowerParams;
 };
