Move attribute computation functions to shared library

- Move ComputeBudNumbers, ComputeNjordPixelIndex, ComputeLengthFromRoot, ComputeBudDevelopment, ComputeBudDirections, ComputeBudHormoneLevels, ComputeGradients, ComputeBudStatus, ComputeBudLateralMeristem, ComputeBudLightDetected, InterpolatePointsScales, SplitFusedPoints
- Move related tests to new PVAttributesTests.cpp file
- Fix issue where ComputeBudLateralMeristem inverted the ParentDot

#Synced-CL 50701977

diff --git a/Engine/Plugins/Experimental/ProceduralVegetationEditor/Source/ProceduralVegetation/Private/Helpers/PVAttributesHelper.cpp b/Engine/Plugins/Experimental/ProceduralVegetationEditor/Source/ProceduralVegetation/Private/Helpers/PVAttributesHelper.cpp
new file mode 100644
--- /dev/null
+++ b/Engine/Plugins/Experimental/ProceduralVegetationEditor/Source/ProceduralVegetation/Private/Helpers/PVAttributesHelper.cpp
@@ -0,0 +1,615 @@
+// Copyright Epic Games, Inc. All Rights Reserved.
+
+#include "Helpers/PVAttributesHelper.h"
+
+void PV::AttributesHelper::ComputeBudNumbers(
+	FPointBudNumberAttributeView PointBudNumberAttribute,
+	FBranchSourceBudNumberAttributeView BranchSourceBudNumberAttribute,
+	FBranchPointsAttributeConstView BranchPointsAttribute
+)
+{
+	for (int32 PointIndex = 0; PointIndex < PointBudNumberAttribute.Num(); ++PointIndex)
+	{
+		PointBudNumberAttribute[PointIndex] = PointIndex + 1;
+	}
+
+	for (int32 BranchIndex = 0; BranchIndex < BranchSourceBudNumberAttribute.Num(); ++BranchIndex)
+	{
+		check(BranchPointsAttribute[BranchIndex].Num() > 1);
+		const int32 BranchRootPointIndex = BranchPointsAttribute[BranchIndex][0];
+		BranchSourceBudNumberAttribute[BranchIndex] = PointBudNumberAttribute[BranchRootPointIndex];
+	}
+}
+
+void PV::AttributesHelper::ComputeNjordPixelIndex(
+	PV::FPointNjordPixelIndexAttributeView PointNjordPixelIndexAttribute,
+	PV::FBranchPointsAttributeConstView BranchPointsAttribute,
+	PV::FBranchParentNumberAttributeConstView BranchParentNumberAttribute,
+	PV::FBranchNumberAttributeConstView BranchNumberAttribute,
+	PV::FBranchChildrenAttributeConstView BranchChildrenAttribute
+)
+{
+	using namespace PV::AttributesHelper;
+
+	RecursiveWalkBranches<int32>(
+		BranchParentNumberAttribute,
+		BranchChildrenAttribute,
+		BranchNumberAttribute,
+		1, 
+		[&](int32 BranchIndex, int32& NjordPixelIndex)->EForEachResult
+		{
+			const bool bIsTrunk = GetBranchParentIndex(BranchParentNumberAttribute, BranchNumberAttribute, BranchIndex) == INDEX_NONE;
+			const TArray<int32>& BranchPoints = BranchPointsAttribute[BranchIndex];
+			for (int32 i = 0; i < BranchPoints.Num(); ++i)
+			{
+				if (bIsTrunk || i != 0)
+				{
+					PointNjordPixelIndexAttribute[BranchPoints[i]] = NjordPixelIndex;
+					NjordPixelIndex += 1;
+				}
+			}
+
+			return EForEachResult::Continue;
+		}
+	);
+}
+
+void PV::AttributesHelper::ComputeLengthFromRoot(
+	PV::FPointLengthFromRootAttributeView PointLengthFromRootAttribute,
+	PV::FPointLengthFromSeedAttributeView PointLengthFromSeedAttribute,
+	PV::FPointPositionAttributeConstView PointPositionAttribute,
+	PV::FBranchPointsAttributeConstView BranchPointsAttribute,
+	PV::FBranchParentNumberAttributeConstView BranchParentNumberAttribute,
+	PV::FBranchNumberAttributeConstView BranchNumberAttribute,
+	PV::FBranchChildrenAttributeConstView BranchChildrenAttribute
+)
+{
+	using namespace PV::AttributesHelper;
+
+	RecursiveWalkBranches(BranchParentNumberAttribute, BranchChildrenAttribute, BranchNumberAttribute, [&](int32 BranchIndex)->EForEachResult
+	{
+		const TArray<int32>& CurrentBranchPoints = BranchPointsAttribute[BranchIndex];
+		check(CurrentBranchPoints.Num() > 0);
+
+		const int32 FirstPointIndex = CurrentBranchPoints[0];
+		const int32 ParentBranchIndex = GetBranchParentIndex(BranchParentNumberAttribute, BranchNumberAttribute, BranchIndex);
+
+		if (ParentBranchIndex == INDEX_NONE)
+		{
+			PointLengthFromRootAttribute[FirstPointIndex] = 0.f; // No parent, length to root is 0
+			PointLengthFromSeedAttribute[FirstPointIndex] = 0.f;
+		}
+		else
+		{
+			const TArray<int32>& ParentBranchPoints = BranchPointsAttribute[ParentBranchIndex];
+			check(ParentBranchPoints.Num() > 0);
+
+			const FVector3f& FirstPoint = PointPositionAttribute[FirstPointIndex];
+
+			int32 ClosestParentPointIndex = ParentBranchPoints[0];
+			float ClosestParentPointDistSquared = (PointPositionAttribute[ClosestParentPointIndex] - FirstPoint).SizeSquared();
+
+			for (int32 i = 1; i < ParentBranchPoints.Num(); ++i)
+			{
+				const int32 PointIndex = ParentBranchPoints[i];
+				const FVector3f ParentPoint = PointPositionAttribute[PointIndex];
+
+				const float DistSquared = (ParentPoint - FirstPoint).SizeSquared();
+				if (DistSquared > ClosestParentPointDistSquared)
+				{
+					continue;
+				}
+
+				const FVector3f ParentBranchDirection = [&]()
+				{
+					if (i < ParentBranchPoints.Num() - 1)
+					{
+						const int32 NextPointIndex = ParentBranchPoints[i + 1];
+						return (PointPositionAttribute[NextPointIndex] - PointPositionAttribute[PointIndex]).GetSafeNormal();
+					}
+
+					if (i > 0)
+					{
+						const int32 PrevPointIndex = ParentBranchPoints[i - 1];
+						return (PointPositionAttribute[PointIndex] - PointPositionAttribute[PrevPointIndex]).GetSafeNormal();
+					}
+
+					return FVector3f::UpVector;
+				}();
+
+				const FVector3f ParentPointToFirstPointDirection = (FirstPoint - ParentPoint).GetSafeNormal();
+				if (ParentPointToFirstPointDirection.Dot(ParentBranchDirection) > 0)
+				{
+					ClosestParentPointIndex = PointIndex;
+					ClosestParentPointDistSquared = DistSquared;
+				}
+			}
+
+			PointLengthFromRootAttribute[FirstPointIndex] = PointLengthFromRootAttribute[ClosestParentPointIndex] + (ClosestParentPointDistSquared > 0 ? FMath::Sqrt(ClosestParentPointDistSquared) : 0) * 0.01 /*cm->m*/;
+			PointLengthFromSeedAttribute[FirstPointIndex] = PointLengthFromRootAttribute[FirstPointIndex];
+		}
+
+		for (int32 i = 1; i < CurrentBranchPoints.Num(); ++i)
+		{
+			const int32 PointIndex = CurrentBranchPoints[i];
+			const int32 PrevPointIndex = CurrentBranchPoints[i - 1];
+
+			const FVector3f& Position = PointPositionAttribute[PointIndex];
+			const FVector3f& PrevPosition = PointPositionAttribute[PrevPointIndex];
+
+			PointLengthFromRootAttribute[PointIndex] = PointLengthFromRootAttribute[PrevPointIndex] + (Position - PrevPosition).Size() * 0.01 /*cm->m*/;
+			PointLengthFromSeedAttribute[PointIndex] = PointLengthFromRootAttribute[PointIndex];
+		}
+
+		return EForEachResult::Continue;
+	});
+}
+
+void PV::AttributesHelper::ComputeBudDevelopment(
+	PV::FBudDevelopmentAttributeView BudDevelopmentAttribute,
+	PV::FPointPlantGradientAttributeConstView PointPlantGradientAttribute,
+	PV::FBranchPointsAttributeConstView BranchPointsAttribute,
+	PV::FBranchParentNumberAttributeConstView BranchParentNumberAttribute,
+	PV::FBranchParentsAttributeConstView BranchParentsAttribute,
+	PV::FBranchNumberAttributeConstView BranchNumberAttribute
+)
+{
+	using namespace PV::AttributesHelper;
+
+	int32 MaxBudAge = 0;
+	for (int32 BranchIndex = 0; BranchIndex < BranchPointsAttribute.Num(); ++BranchIndex)
+	{
+		MaxBudAge = FMath::Max(MaxBudAge, BranchPointsAttribute[BranchIndex].Num());
+	}
+
+	ForEachUniquePointOnBranches_FusedPoints(
+		BranchPointsAttribute,
+		BranchParentNumberAttribute,
+		BranchNumberAttribute,
+		[&](int32 BranchIndex, int32 PointIndex)
+		{
+			const float Gradient = PointPlantGradientAttribute[PointIndex];
+			const int32 BudAge = FMath::RoundToInt(FMath::GetMappedRangeValueClamped(FVector2f(0, 1), FVector2f(1, float(MaxBudAge)), Gradient));
+
+			BudDevelopmentAttribute[PointIndex].BudAge = BudAge;
+			BudDevelopmentAttribute[PointIndex].RelativeBudAge = BudAge;
+			BudDevelopmentAttribute[PointIndex].Generation = BranchParentsAttribute[BranchIndex].Num() + 1;
+
+			return EForEachResult::Continue;
+		}
+	);
+
+	ForEachUniquePointOnBranches_FusedPoints(
+		BranchPointsAttribute,
+		BranchParentNumberAttribute,
+		BranchNumberAttribute,
+		[&](int32 BranchIndex, int32 PointIndex)
+		{
+			const bool bIsRootBranch = GetBranchParentIndex(BranchParentNumberAttribute, BranchNumberAttribute, BranchIndex) == INDEX_NONE;
+			const int32 RootPointIndex = bIsRootBranch ? BranchPointsAttribute[BranchIndex][0] : BranchPointsAttribute[BranchIndex][1];
+			const int32 BranchAge = BudDevelopmentAttribute[RootPointIndex].BudAge;
+			BudDevelopmentAttribute[PointIndex].BranchAge = BranchAge;
+
+			return EForEachResult::Continue;
+		}
+	);
+}
+
+void PV::AttributesHelper::ComputeBudDirections(
+	PV::FBudDirectionAttributeView BudDirectionAttribute,
+	PV::FPointPositionAttributeConstView PointPositionAttribute,
+	PV::FBranchChildrenAttributeConstView BranchChildrenAttribute,
+	PV::FBranchParentNumberAttributeConstView BranchParentNumberAttribute,
+	PV::FBranchNumberAttributeConstView BranchNumberAttribute,
+	PV::FBranchPointsAttributeConstView BranchPointsAttribute
+)
+{
+	using namespace PV::AttributesHelper;
+
+	RecursiveWalkPlantPoints_FusedPoints<int32>(
+		BranchChildrenAttribute,
+		BranchParentNumberAttribute,
+		BranchNumberAttribute,
+		BranchPointsAttribute,
+		INDEX_NONE,
+		[&](int32 BranchIndex, int32 BranchPointIndex, int32 PointIndex, const TArray<int32>& ChildPointIndices, int32& PrevPointIndex)->EForEachResult
+		{
+			const TArray<int32>& BranchPoints = BranchPointsAttribute[BranchIndex];
+
+			const FVector3f Apical = [&]()->FVector3f
+			{
+				if (BranchPointIndex < BranchPoints.Num() - 1)
+				{
+					const int32 NextPointIndex = BranchPoints[BranchPointIndex + 1];
+					return (PointPositionAttribute[NextPointIndex] - PointPositionAttribute[PointIndex]).GetSafeNormal();
+				}
+				else if (BranchPointIndex > 0)
+				{
+					const int32 PrevPointIndex = BranchPoints[BranchPointIndex - 1];
+					return (PointPositionAttribute[PointIndex] - PointPositionAttribute[PrevPointIndex]).GetSafeNormal();
+				}
+
+				ensure(false); // We should never be able to get here unless we have a branch with only 1 point
+				return FVector3f::UpVector;
+			}();
+
+			const FVector3f UpVector = FRotationMatrix44f::MakeFromZX(Apical, FVector3f::ForwardVector).ToQuat().GetRightVector() * -1;
+			const FVector3f PrevAxillary = PrevPointIndex == INDEX_NONE ? FVector3f::ForwardVector : BudDirectionAttribute[PrevPointIndex].Axillary;
+			const FVector3f Axillary = FQuat4f(Apical, FMath::DegreesToRadians(135.f)).RotateVector(PrevAxillary);
+
+			BudDirectionAttribute[PointIndex].Apical = Apical;
+			BudDirectionAttribute[PointIndex].Axillary = Axillary;
+			BudDirectionAttribute[PointIndex].LightOptimal = Apical;
+			BudDirectionAttribute[PointIndex].LightSubOptimal = Apical;
+			BudDirectionAttribute[PointIndex].GuideCurve = Apical;
+			BudDirectionAttribute[PointIndex].UpVector = UpVector;
+			
+			PrevPointIndex = PointIndex;
+
+			return EForEachResult::Continue;
+		}
+	);
+}
+
+void PV::AttributesHelper::ComputeBudHormoneLevels(
+	PV::FBudHormoneLevelsAttributeView BudHormoneLevelsAttribute,
+	PV::FPointPlantGradientAttributeConstView PointPlantGradientAttribute
+)
+{
+	using namespace PV::AttributesHelper;
+		
+	for (int32 PointIndex = 0; PointIndex < BudHormoneLevelsAttribute.Num(); PointIndex++)
+	{
+		const float Gradient = PointPlantGradientAttribute[PointIndex];
+		const float InverseGradient = 1.f - Gradient;
+
+		BudHormoneLevelsAttribute[PointIndex].Apical = FMath::GetMappedRangeValueClamped(FVector2f(0.75, 1.0), FVector2f(0, 1), InverseGradient);
+		BudHormoneLevelsAttribute[PointIndex].Axillary = 0;
+		BudHormoneLevelsAttribute[PointIndex].AxillaryInhibition = 0;
+		BudHormoneLevelsAttribute[PointIndex].Radical = 0;
+		BudHormoneLevelsAttribute[PointIndex].Ethylene = Gradient;
+		BudHormoneLevelsAttribute[PointIndex].Cykininin = 0;
+	}
+}
+
+void PV::AttributesHelper::ComputeGradients(
+	PV::FPointGroundGradientAttributeView PointGroundGradientAttribute,
+	PV::FPointScaleGradientAttributeView PointScaleGradientAttribute,
+	PV::FPointPlantGradientAttributeView PointPlantGradientAttribute,
+	PV::FPointBranchGradientAttributeView PointBranchGradientAttribute,
+	PV::FPointLengthFromRootAttributeConstView PointLengthFromRootAttribute,
+	PV::FPointScaleAttributeConstView PointScaleAttribute,
+	PV::FBranchPointsAttributeConstView BranchPointsAttribute,
+	PV::FBranchParentNumberAttributeConstView BranchParentNumberAttribute,
+	PV::FBranchNumberAttributeConstView BranchNumberAttribute,
+	PV::FBranchChildrenAttributeConstView BranchChildrenAttribute
+)
+{
+	using namespace PV::AttributesHelper;
+
+	const float* MaxLengthFromRoot = Algo::MaxElement(PointLengthFromRootAttribute);
+	const float* MaxPointScale = Algo::MaxElement(PointScaleAttribute);
+	if (MaxLengthFromRoot == nullptr || MaxPointScale == nullptr || *MaxLengthFromRoot == 0 || *MaxPointScale == 0)
+	{
+		return;
+	}
+
+	for (int32 PointIndex = 0; PointIndex < PointGroundGradientAttribute.Num(); PointIndex++)
+	{
+		const float LengthFromRoot = PointLengthFromRootAttribute[PointIndex];
+		const float LengthFromRootAlpha = LengthFromRoot / *MaxLengthFromRoot;
+		PointGroundGradientAttribute[PointIndex] = 1.0 - LengthFromRootAlpha;
+
+		const float PointScale = PointScaleAttribute[PointIndex];
+		const float PointScaleAlpha = PointScale / *MaxPointScale;
+		PointScaleGradientAttribute[PointIndex] = PointScaleAlpha;
+	}
+
+	// PlantGradient - A gradient from 1-0 where the tip of each branch is 0, and the base of each branch is whatever the gradient of the parent branch is at that point.
+	// BranchGradient - A gradient from 1-0 where 1 = base of the branch and 0 = tip of the branch
+	RecursiveWalkBranches(
+		BranchParentNumberAttribute,
+		BranchChildrenAttribute,
+		BranchNumberAttribute,
+		[&](int32 BranchIndex)->EForEachResult
+		{
+			const TArray<int32>& BranchPoints = BranchPointsAttribute[BranchIndex];
+			check(BranchPoints.Num() > 1);
+
+			const float BasePointPlantGradient = GetBranchParentIndex(BranchParentNumberAttribute, BranchNumberAttribute, BranchIndex) == INDEX_NONE
+				? 1.0
+				: PointPlantGradientAttribute[BranchPoints[0]];
+
+			for (int32 i = 0; i < BranchPoints.Num(); ++i)
+			{
+				const float BranchAlpha = i / float(BranchPoints.Num() - 1);
+				const float PointPlantGradient = FMath::Lerp(BasePointPlantGradient, 0, BranchAlpha);
+				PointPlantGradientAttribute[BranchPoints[i]] = PointPlantGradient;
+				PointBranchGradientAttribute[BranchPoints[i]] = 1.0 - BranchAlpha;
+			}
+
+			return EForEachResult::Continue;
+		}
+	);
+}
+
+void PV::AttributesHelper::ComputeBudStatus(
+	PV::FBudStatusAttributeView BudStatusAttribute,
+	PV::FBranchPointsAttributeConstView BranchPointsAttribute,
+	PV::FBranchParentNumberAttributeConstView BranchParentNumberAttribute,
+	PV::FBranchNumberAttributeConstView BranchNumberAttribute,
+	PV::FBranchChildrenAttributeConstView BranchChildrenAttribute
+)
+{
+	using namespace PV::AttributesHelper;
+
+	ForEachUniquePointOnBranches_FusedPoints(
+		BranchPointsAttribute,
+		BranchParentNumberAttribute,
+		BranchNumberAttribute,
+		[&](int32 BranchIndex, int32 PointIndex)->EForEachResult
+		{
+			PV::FBudStatusView BudStatusView = BudStatusAttribute[PointIndex];
+			BudStatusView.ApicalMeristem = 0;
+			BudStatusView.Codominant = 0;
+			BudStatusView.Axillary = 0;
+			BudStatusView.Seed = 0;
+			BudStatusView.Dormant = 0;
+			BudStatusView.Triggered = 1;
+			BudStatusView.NumTriggered = 1;
+			BudStatusView.Inactive = 1;
+			BudStatusView.BrokenTip = 0;
+			BudStatusView.Broken = 0;
+
+			if (BranchPointsAttribute[BranchIndex].Last() == PointIndex) // IsTip
+			{
+				BudStatusView.ApicalMeristem = 1;
+				BudStatusView.Dormant = 1;
+				BudStatusView.Triggered = 0;
+				BudStatusView.NumTriggered = 0;
+				BudStatusView.Inactive = 0;
+			}
+			else if (BranchPointsAttribute[BranchIndex][0] == PointIndex) // Is branch root point.
+				// Due to ForEachUniquePointOnBranches_FusedPoints only iterating shared points on the parent branch, we don't have to check whether this is the root branch
+				// (as the root branch is the only branch which will have the root point iterated on by ForEachUniquePointOnBranches_FusedPoints).
+			{
+				BudStatusView.Seed = 1;
+			}
+			else
+			{
+				const int32 NumChildren = GetBranchPointChildBranchIndices_FusedPoints(
+					BranchChildrenAttribute,
+					BranchParentNumberAttribute,
+					BranchNumberAttribute,
+					BranchPointsAttribute,
+					BranchIndex, 
+					PointIndex
+				).Num();
+
+				if (NumChildren > 0) 
+				{
+					BudStatusView.NumTriggered = NumChildren;
+					BudStatusView.Axillary = 1;
+				}
+			}
+
+			return EForEachResult::Continue;
+		}
+	);
+}
+
+void PV::AttributesHelper::ComputeBudLateralMeristem(
+	PV::FBudLateralMeristemAttributeView BudLateralMeristemAttribute,
+	PV::FPointPositionAttributeConstView PointPositionAttribute,
+	PV::FPointScaleAttributeConstView PointScaleAttribute,
+	PV::FPointLengthFromRootAttributeConstView PointLengthFromRootAttribute,
+	PV::FBudDirectionAttributeConstView BudDirectionAttribute,
+	PV::FBranchPointsAttributeConstView BranchPointsAttribute,
+	PV::FBranchParentNumberAttributeConstView BranchParentNumberAttribute,
+	PV::FBranchChildrenAttributeConstView BranchChildrenAttribute,
+	PV::FBranchNumberAttributeConstView BranchNumberAttribute
+)
+{
+	using namespace PV::AttributesHelper;
+
+	RecursiveWalkBranches_Reversed(
+		BranchParentNumberAttribute,
+		BranchChildrenAttribute,
+		BranchNumberAttribute,
+		[&](int32 BranchIndex)->EForEachResult
+		{
+			const TArray<int32>& BranchPointIndices = BranchPointsAttribute[BranchIndex];
+			check(BranchPointIndices.Num() > 1);
+
+			const FVector3f& FirstBranchPointPos = PointPositionAttribute[BranchPointIndices[0]];
+			const FVector3f& SecondBranchPointPos = PointPositionAttribute[BranchPointIndices[1]];
+
+			const FVector3f BranchDir = (SecondBranchPointPos - FirstBranchPointPos).GetSafeNormal();
+			const FVector3f FirstPointNormal = BudDirectionAttribute[BranchPointIndices[0]].Apical;;
+			
+			const float ParentDot = FVector3f::DotProduct(BranchDir, FirstPointNormal);
+
+			float PreviousPointDavinciValue = 0;
+			for (int32 i = BranchPointIndices.Num() - 1; i>= 0; --i)
+			{
+				const int32 PointIndex = BranchPointIndices[i];
+				const float ExisitingDavinciValue = BudLateralMeristemAttribute[PointIndex].Davinci;
+				const float PointScale = PointScaleAttribute[PointIndex];
+				const float AccumulatedDavinciValue = PreviousPointDavinciValue + ExisitingDavinciValue + PointScale;
+				PreviousPointDavinciValue = AccumulatedDavinciValue;
+
+				BudLateralMeristemAttribute[PointIndex].LateralMeristem = 0.0;
+				BudLateralMeristemAttribute[PointIndex].Multiplier = 1;
+				BudLateralMeristemAttribute[PointIndex].Inactive = 0.0;
+				BudLateralMeristemAttribute[PointIndex].Davinci = AccumulatedDavinciValue;
+				BudLateralMeristemAttribute[PointIndex].ParentDot = ParentDot;
+				BudLateralMeristemAttribute[PointIndex].RootDistance = PointLengthFromRootAttribute[PointIndex];
+				BudLateralMeristemAttribute[PointIndex].Degredation = 0.0;
+			}
+
+			return EForEachResult::Continue;
+		}
+	);
+}
+
+void PV::AttributesHelper::ComputeBudLightDetected(PV::FBudLightDetectedAttributeView BudLightDetectedAttribute)
+{
+	for (int32 i = 0; i < BudLightDetectedAttribute.Num(); ++i)
+	{
+		BudLightDetectedAttribute[i].Availible = 1.0;
+		BudLightDetectedAttribute[i].Resource = 1.0;
+		BudLightDetectedAttribute[i].Branch = 1.0;
+		BudLightDetectedAttribute[i].Collision = 0.0;
+	}
+}
+
+void PV::AttributesHelper::InterpolatePointsScales(
+	PV::FPointScaleAttributeView PointScaleAttribute,
+	PV::FPointLengthFromRootAttributeConstView PointLengthFromRootAttribute,
+	PV::FBranchPointsAttributeConstView BranchPointsAttribute,
+	PV::FBranchParentNumberAttributeConstView BranchParentNumberAttribute,
+	PV::FBranchChildrenAttributeConstView BranchChildrenAttribute,
+	PV::FBranchNumberAttributeConstView BranchNumberAttribute
+)
+{
+	using namespace PV::AttributesHelper;
+
+	RecursiveWalkBranches(
+		BranchParentNumberAttribute,
+		BranchChildrenAttribute,
+		BranchNumberAttribute,
+		[&](int32 BranchIndex)->EForEachResult
+		{
+			const TArray<int32>& BranchPoints = BranchPointsAttribute[BranchIndex];
+
+			float CurrentPointScale = PointScaleAttribute[BranchPoints.Last()];
+			float CurrentLengthFromRoot = PointLengthFromRootAttribute[BranchPoints.Last()];
+
+			for (int32 i = BranchPoints.Num() - 1; i >= 0; --i)
+			{
+				float NextValidScale = -1;
+				float NextValidScale_LengthFromRoot = -1;
+				int32 NextValidScale_Index = INDEX_NONE;
+
+				for (int32 j = i - 1; j >= 0; --j)
+				{
+					const int32 PointIndex = BranchPoints[j];
+					const float PointScale = PointScaleAttribute[PointIndex];
+					if (PointScale > 0)
+					{
+						NextValidScale = PointScale;
+						NextValidScale_LengthFromRoot = PointLengthFromRootAttribute[PointIndex];
+						NextValidScale_Index = j;
+						break;
+					}
+				}
+
+				if (NextValidScale_Index == INDEX_NONE)
+				{
+					for (int32 j = i - 1; j >= 0; --j)
+					{
+						PointScaleAttribute[BranchPoints[j]] = CurrentPointScale;
+					}
+					break;
+				}
+
+				if (NextValidScale_Index == i - 1)
+				{
+					CurrentPointScale = PointScaleAttribute[BranchPoints[NextValidScale_Index]];
+					continue;
+				}
+
+				for (int32 j = i - 1; j >= NextValidScale_Index; --j)
+				{
+					const int32 PointIndex = BranchPoints[j];
+					const float PointLengthFromRoot = PointLengthFromRootAttribute[PointIndex];
+					const float Alpha = FMath::GetMappedRangeValueClamped(FVector2D(CurrentLengthFromRoot, NextValidScale_LengthFromRoot), FVector2D(0, 1), PointLengthFromRoot);
+					const float PointScale = FMath::Lerp(CurrentPointScale, NextValidScale, Alpha);
+					PointScaleAttribute[PointIndex] = PointScale;
+				}
+
+				CurrentPointScale = NextValidScale;
+				CurrentLengthFromRoot = NextValidScale_LengthFromRoot;
+				i = NextValidScale_Index + 1;
+			}
+
+			return EForEachResult::Continue;
+		}
+	);
+}
+
+bool PV::AttributesHelper::SplitFusedPoints(FManagedArrayCollection& Collection)
+{
+	if (!Collection.HasGroup(PV::GroupNames::PointGroup) || !Collection.HasGroup(PV::GroupNames::BranchGroup))
+	{
+		return false;
+	}
+
+	auto BranchPointsAttribute = FBranchPointsAttribute::GetAttribute(Collection);
+	auto BranchParentNumberAttribute = FBranchParentNumberAttribute::GetAttribute(Collection);
+
+	auto PointPositionAttribute = FPointPositionAttribute::GetAttribute(Collection);
+	auto LengthFromRootAttribute = FPointLengthFromRootAttribute::GetAttribute(Collection);
+	auto LengthFromSeedAttribute = FPointLengthFromSeedAttribute::GetAttribute(Collection);
+	auto PointScaleGradientAttribute = FPointScaleGradientAttribute::GetAttribute(Collection);
+	auto PointHullGradientAttribute = FPointHullGradientAttribute::GetAttribute(Collection);
+	auto PointMainTrunkGradientAttribute = FPointMainTrunkGradientAttribute::GetAttribute(Collection);
+	auto GroundGradientAttribute = FPointGroundGradientAttribute::GetAttribute(Collection);
+	auto PointScaleAttribute = FPointScaleAttribute::GetAttribute(Collection);
+	auto BudDirectionAttribute = FBudDirectionAttribute::GetAttribute(Collection);
+	auto BudStatusAttribute = FBudStatusAttribute::GetAttribute(Collection);
+	auto BudLateralMeristemAttribute = FBudLateralMeristemAttribute::GetAttribute(Collection);
+	auto PointBranchGradientAttribute = FPointBranchGradientAttribute::GetAttribute(Collection);
+	auto BudHormoneLevelsAttribute = FBudHormoneLevelsAttribute::GetAttribute(Collection);
+	auto PlantGradientAttribute = FPointPlantGradientAttribute::GetAttribute(Collection);
+	auto BudNumberAttribute = FPointBudNumberAttribute::GetAttribute(Collection);
+	auto NjordPixelIndexAttribute = FPointNjordPixelIndexAttribute::GetAttribute(Collection);
+	auto BudLightDetectedAttribute = FBudLightDetectedAttribute::GetAttribute(Collection);
+	auto BudDevelopmentAttribute = FBudDevelopmentAttribute::GetAttribute(Collection);
+
+	const int32* MaxBudNumber = Algo::MaxElement(BudNumberAttribute);
+	check(MaxBudNumber != nullptr);
+
+	int32 NewBudNumber = *MaxBudNumber + 1;
+
+	for (int32 i = 0; i < BranchPointsAttribute.Num(); ++i)
+	{
+		const int32 BranchIndex = i;
+		if (BranchParentNumberAttribute[BranchIndex] == 0)
+		{
+			continue;
+		}
+
+		const int32 NumPoints = PointPositionAttribute.Num();
+		Collection.Resize(NumPoints + 1, PV::GroupNames::PointGroup);
+
+		TArray<int32>& BranchPoints = BranchPointsAttribute[BranchIndex];
+		const int32 PrevPointIndex = BranchPoints[0];
+		const int32 NewPointIndex = NumPoints;
+
+		BranchPoints[0] = NewPointIndex;
+
+		PointPositionAttribute[NewPointIndex] = PointPositionAttribute[PrevPointIndex];
+		LengthFromRootAttribute[NewPointIndex] = LengthFromRootAttribute[PrevPointIndex];
+		LengthFromSeedAttribute[NewPointIndex] = LengthFromSeedAttribute[PrevPointIndex];
+		PointScaleGradientAttribute[NewPointIndex] = PointScaleGradientAttribute[PrevPointIndex];
+		PointHullGradientAttribute[NewPointIndex] = PointHullGradientAttribute[PrevPointIndex];
+		PointMainTrunkGradientAttribute[NewPointIndex] = PointMainTrunkGradientAttribute[PrevPointIndex];
+		GroundGradientAttribute[NewPointIndex] = GroundGradientAttribute[PrevPointIndex];
+		PointScaleAttribute[NewPointIndex] = PointScaleAttribute[PrevPointIndex];
+		BudDirectionAttribute[NewPointIndex].Array = BudDirectionAttribute[PrevPointIndex].Array;
+		BudStatusAttribute[NewPointIndex].Array = BudStatusAttribute[PrevPointIndex].Array;
+		BudLateralMeristemAttribute[NewPointIndex].Array = BudLateralMeristemAttribute[PrevPointIndex].Array;
+		PointBranchGradientAttribute[NewPointIndex] = PointBranchGradientAttribute[PrevPointIndex];
+		BudHormoneLevelsAttribute[NewPointIndex].Array = BudHormoneLevelsAttribute[PrevPointIndex].Array;
+		PlantGradientAttribute[NewPointIndex] = PlantGradientAttribute[PrevPointIndex];
+		BudNumberAttribute[NewPointIndex] = NewBudNumber;
+		NjordPixelIndexAttribute[NewPointIndex] = NjordPixelIndexAttribute[PrevPointIndex];
+		BudLightDetectedAttribute[NewPointIndex].Array = BudLightDetectedAttribute[PrevPointIndex].Array;
+		BudDevelopmentAttribute[NewPointIndex].Array = BudDevelopmentAttribute[PrevPointIndex].Array;
+
+		NewBudNumber++;
+	}
+
+	return true;
+}
\ No newline at end of file
diff --git a/Engine/Plugins/Experimental/ProceduralVegetationEditor/Source/ProceduralVegetation/Private/Tests/PVAttributesTests.cpp b/Engine/Plugins/Experimental/ProceduralVegetationEditor/Source/ProceduralVegetation/Private/Tests/PVAttributesTests.cpp
new file mode 100644
--- /dev/null
+++ b/Engine/Plugins/Experimental/ProceduralVegetationEditor/Source/ProceduralVegetation/Private/Tests/PVAttributesTests.cpp
@@ -0,0 +1,686 @@
+// Copyright Epic Games, Inc. All Rights Reserved.
+
+#if WITH_DEV_AUTOMATION_TESTS
+
+#include "Tests/PVTestsCommon.h"
+#include "Helpers/PVAttributesHelper.h"
+
+#define PV_ATTRIBUTES_TEST(TestName) PV_SIMPLE_AUTOMATION_TEST(Attributes, TestName)
+
+PV_ATTRIBUTES_TEST(ComputeLengthFromRoot)
+{
+	const TSharedRef<FManagedArrayCollection> Collection = PVMockTreeCollection::CreateCollection();
+
+	PV::FPointLengthFromRootAttributeConstView LengthFromRootAttribute = PV::FPointLengthFromRootAttribute::FindAttribute(*Collection);
+	UTEST_TRUE("HasAttribute", LengthFromRootAttribute.IsValid());
+	UTEST_TRUE("Has 12 elements", LengthFromRootAttribute.Num() >= 12);
+
+	UTEST_NEARLY_EQUAL("Point0_LengthFromRoot", LengthFromRootAttribute[0], PVMockTree::Point0_LengthFromRoot, 0.01f);
+	UTEST_NEARLY_EQUAL("Point1_LengthFromRoot", LengthFromRootAttribute[1], PVMockTree::Point1_LengthFromRoot, 0.01f);
+	UTEST_NEARLY_EQUAL("Point2_LengthFromRoot", LengthFromRootAttribute[2], PVMockTree::Point2_LengthFromRoot, 0.01f);
+	UTEST_NEARLY_EQUAL("Point3_LengthFromRoot", LengthFromRootAttribute[3], PVMockTree::Point3_LengthFromRoot, 0.01f);
+	UTEST_NEARLY_EQUAL("Point4_LengthFromRoot", LengthFromRootAttribute[4], PVMockTree::Point4_LengthFromRoot, 0.01f);
+	UTEST_NEARLY_EQUAL("Point5_LengthFromRoot", LengthFromRootAttribute[5], PVMockTree::Point5_LengthFromRoot, 0.01f);
+	UTEST_NEARLY_EQUAL("Point6_LengthFromRoot", LengthFromRootAttribute[6], PVMockTree::Point6_LengthFromRoot, 0.01f);
+	UTEST_NEARLY_EQUAL("Point7_LengthFromRoot", LengthFromRootAttribute[7], PVMockTree::Point7_LengthFromRoot, 0.01f);
+	UTEST_NEARLY_EQUAL("Point8_LengthFromRoot", LengthFromRootAttribute[8], PVMockTree::Point8_LengthFromRoot, 0.01f);
+	UTEST_NEARLY_EQUAL("Point9_LengthFromRoot", LengthFromRootAttribute[9], PVMockTree::Point9_LengthFromRoot, 0.01f);
+	UTEST_NEARLY_EQUAL("Point10_LengthFromRoot", LengthFromRootAttribute[10], PVMockTree::Point10_LengthFromRoot, 0.01f);
+	UTEST_NEARLY_EQUAL("Point11_LengthFromRoot", LengthFromRootAttribute[11], PVMockTree::Point11_LengthFromRoot, 0.01f);
+
+	const PV::FPointLengthFromSeedAttributeView LengthFromSeedAttribute = PV::FPointLengthFromSeedAttribute::FindAttribute(*Collection);
+	UTEST_TRUE("HasAttribute", LengthFromSeedAttribute.IsValid());
+	UTEST_TRUE("Has 12 elements", LengthFromSeedAttribute.Num() >= 12);
+
+	UTEST_NEARLY_EQUAL("Point0_LengthFromRoot", LengthFromSeedAttribute[0], PVMockTree::Point0_LengthFromRoot, 0.01f);
+	UTEST_NEARLY_EQUAL("Point1_LengthFromRoot", LengthFromSeedAttribute[1], PVMockTree::Point1_LengthFromRoot, 0.01f);
+	UTEST_NEARLY_EQUAL("Point2_LengthFromRoot", LengthFromSeedAttribute[2], PVMockTree::Point2_LengthFromRoot, 0.01f);
+	UTEST_NEARLY_EQUAL("Point3_LengthFromRoot", LengthFromSeedAttribute[3], PVMockTree::Point3_LengthFromRoot, 0.01f);
+	UTEST_NEARLY_EQUAL("Point4_LengthFromRoot", LengthFromSeedAttribute[4], PVMockTree::Point4_LengthFromRoot, 0.01f);
+	UTEST_NEARLY_EQUAL("Point5_LengthFromRoot", LengthFromSeedAttribute[5], PVMockTree::Point5_LengthFromRoot, 0.01f);
+	UTEST_NEARLY_EQUAL("Point6_LengthFromRoot", LengthFromSeedAttribute[6], PVMockTree::Point6_LengthFromRoot, 0.01f);
+	UTEST_NEARLY_EQUAL("Point7_LengthFromRoot", LengthFromSeedAttribute[7], PVMockTree::Point7_LengthFromRoot, 0.01f);
+	UTEST_NEARLY_EQUAL("Point8_LengthFromRoot", LengthFromSeedAttribute[8], PVMockTree::Point8_LengthFromRoot, 0.01f);
+	UTEST_NEARLY_EQUAL("Point9_LengthFromRoot", LengthFromSeedAttribute[9], PVMockTree::Point9_LengthFromRoot, 0.01f);
+	UTEST_NEARLY_EQUAL("Point10_LengthFromRoot", LengthFromSeedAttribute[10], PVMockTree::Point10_LengthFromRoot, 0.01f);
+	UTEST_NEARLY_EQUAL("Point11_LengthFromRoot", LengthFromSeedAttribute[11], PVMockTree::Point11_LengthFromRoot, 0.01f);
+
+	return true;
+}
+
+PV_ATTRIBUTES_TEST(ComputeBudDevelopment)
+{
+	struct FBudDevelopmentDataLayout
+	{
+		int32 Generation;
+		int32 Age;
+		int32 Branch_Age;
+		int32 AgeSenescence;
+		int32 Light_Senescence;
+		int32 Relative_Age;
+
+		bool operator==(const TArray<int32>& Other) const
+		{
+			return Other.Num() == 6
+				&& Other[0] == Generation
+				&& Other[1] == Age
+				&& Other[2] == Branch_Age
+				&& Other[3] == AgeSenescence
+				&& Other[4] == Light_Senescence
+				&& Other[5] == Relative_Age;
+		}
+	};
+
+	const FBudDevelopmentDataLayout Point0_ExpectedBudDevelopment = {
+		.Generation = 1,
+		.Age = 6,
+		.Branch_Age = 6,
+		.AgeSenescence = 0,
+		.Light_Senescence = 0,
+		.Relative_Age = 6
+	};
+	const FBudDevelopmentDataLayout Point1_ExpectedBudDevelopment = {
+		.Generation = 1,
+		.Age = 5,
+		.Branch_Age = 6,
+		.AgeSenescence = 0,
+		.Light_Senescence = 0,
+		.Relative_Age = 5
+	};
+	const FBudDevelopmentDataLayout Point2_ExpectedBudDevelopment = {
+		.Generation = 1,
+		.Age = 4,
+		.Branch_Age = 6,
+		.AgeSenescence = 0,
+		.Light_Senescence = 0,
+		.Relative_Age = 4
+	};
+	const FBudDevelopmentDataLayout Point3_ExpectedBudDevelopment = {
+		.Generation = 1,
+		.Age = 3,
+		.Branch_Age = 6,
+		.AgeSenescence = 0,
+		.Light_Senescence = 0,
+		.Relative_Age = 3
+	};
+	const FBudDevelopmentDataLayout Point4_ExpectedBudDevelopment = {
+		.Generation = 1,
+		.Age = 2,
+		.Branch_Age = 6,
+		.AgeSenescence = 0,
+		.Light_Senescence = 0,
+		.Relative_Age = 2
+	};
+	const FBudDevelopmentDataLayout Point5_ExpectedBudDevelopment = {
+		.Generation = 1,
+		.Age = 1,
+		.Branch_Age = 6,
+		.AgeSenescence = 0,
+		.Light_Senescence = 0,
+		.Relative_Age = 1
+	};
+	const FBudDevelopmentDataLayout Point6_ExpectedBudDevelopment = {
+		.Generation = 2,
+		.Age = 1,
+		.Branch_Age = 1,
+		.AgeSenescence = 0,
+		.Light_Senescence = 0,
+		.Relative_Age = 1
+	};
+	const FBudDevelopmentDataLayout Point7_ExpectedBudDevelopment = {
+		.Generation = 2,
+		.Age = 3,
+		.Branch_Age = 3,
+		.AgeSenescence = 0,
+		.Light_Senescence = 0,
+		.Relative_Age = 3
+	};
+	const FBudDevelopmentDataLayout Point8_ExpectedBudDevelopment = {
+		.Generation = 2,
+		.Age = 2,
+		.Branch_Age = 3,
+		.AgeSenescence = 0,
+		.Light_Senescence = 0,
+		.Relative_Age = 2
+	};
+	const FBudDevelopmentDataLayout Point9_ExpectedBudDevelopment = {
+		.Generation = 2,
+		.Age = 1,
+		.Branch_Age = 3,
+		.AgeSenescence = 0,
+		.Light_Senescence = 0,
+		.Relative_Age = 1
+	};
+	const FBudDevelopmentDataLayout Point10_ExpectedBudDevelopment = {
+		.Generation = 3,
+		.Age = 2,
+		.Branch_Age = 2,
+		.AgeSenescence = 0,
+		.Light_Senescence = 0,
+		.Relative_Age = 2
+	};
+	const FBudDevelopmentDataLayout Point11_ExpectedBudDevelopment = {
+		.Generation = 3,
+		.Age = 1,
+		.Branch_Age = 2,
+		.AgeSenescence = 0,
+		.Light_Senescence = 0,
+		.Relative_Age = 1
+	};
+
+	const TSharedRef<FManagedArrayCollection> Collection = PVMockTreeCollection::CreateCollection();
+	const PV::FBudDevelopmentAttributeConstView BudDevelopmentAttribute = PV::FBudDevelopmentAttribute::FindAttribute(*Collection);
+	UTEST_TRUE("HasAttribute", BudDevelopmentAttribute);
+	UTEST_TRUE("Has 12 elements", BudDevelopmentAttribute.Num() >= 12);
+
+	UTEST_TRUE("Point0_BudDevelopment",  BudDevelopmentAttribute[0].Array  == Point0_ExpectedBudDevelopment);
+	UTEST_TRUE("Point1_BudDevelopment",  BudDevelopmentAttribute[1].Array  == Point1_ExpectedBudDevelopment);
+	UTEST_TRUE("Point2_BudDevelopment",  BudDevelopmentAttribute[2].Array  == Point2_ExpectedBudDevelopment);
+	UTEST_TRUE("Point3_BudDevelopment",  BudDevelopmentAttribute[3].Array  == Point3_ExpectedBudDevelopment);
+	UTEST_TRUE("Point4_BudDevelopment",  BudDevelopmentAttribute[4].Array  == Point4_ExpectedBudDevelopment);
+	UTEST_TRUE("Point5_BudDevelopment",  BudDevelopmentAttribute[5].Array  == Point5_ExpectedBudDevelopment);
+	UTEST_TRUE("Point6_BudDevelopment",  BudDevelopmentAttribute[6].Array  == Point6_ExpectedBudDevelopment);
+	UTEST_TRUE("Point7_BudDevelopment",  BudDevelopmentAttribute[7].Array  == Point7_ExpectedBudDevelopment);
+	UTEST_TRUE("Point8_BudDevelopment",  BudDevelopmentAttribute[8].Array  == Point8_ExpectedBudDevelopment);
+	UTEST_TRUE("Point9_BudDevelopment",  BudDevelopmentAttribute[9].Array  == Point9_ExpectedBudDevelopment);
+	UTEST_TRUE("Point10_BudDevelopment", BudDevelopmentAttribute[10].Array == Point10_ExpectedBudDevelopment);
+	UTEST_TRUE("Point11_BudDevelopment", BudDevelopmentAttribute[11].Array == Point11_ExpectedBudDevelopment);
+
+	return true;
+}
+
+PV_ATTRIBUTES_TEST(ComputeGradients)
+{
+	const float Point0_ExpectedPlantGradient  = 1;
+	const float Point1_ExpectedPlantGradient  = 0.8;
+	const float Point2_ExpectedPlantGradient  = 0.6;
+	const float Point3_ExpectedPlantGradient  = 0.4;
+	const float Point4_ExpectedPlantGradient  = 0.2;
+	const float Point5_ExpectedPlantGradient  = 0;
+	const float Point6_ExpectedPlantGradient  = 0;
+	const float Point7_ExpectedPlantGradient  = FMath::Lerp(Point2_ExpectedPlantGradient, 0.f, 1.f / 3);
+	const float Point8_ExpectedPlantGradient  = FMath::Lerp(Point2_ExpectedPlantGradient, 0.f, 2.f / 3);
+	const float Point9_ExpectedPlantGradient  = 0;
+	const float Point10_ExpectedPlantGradient = FMath::Lerp(Point8_ExpectedPlantGradient, 0.f, 1.f / 2);
+	const float Point11_ExpectedPlantGradient = 0;
+
+	const TSharedRef<FManagedArrayCollection> Collection = PVMockTreeCollection::CreateCollection();
+	
+	const PV::FPointPlantGradientAttributeView PointPlantGradientAttribute = PV::FPointPlantGradientAttribute::FindAttribute(*Collection);
+	UTEST_TRUE("HasAttribute", PointPlantGradientAttribute.IsValid());
+	UTEST_TRUE("Has 12 elements", PointPlantGradientAttribute.Num() >= 12);
+
+	UTEST_TRUE("Point0_ExpectedPlantGradient",  FMath::IsNearlyEqual(PointPlantGradientAttribute[0],  Point0_ExpectedPlantGradient,  0.01f));
+	UTEST_TRUE("Point1_ExpectedPlantGradient",  FMath::IsNearlyEqual(PointPlantGradientAttribute[1],  Point1_ExpectedPlantGradient,  0.01f));
+	UTEST_TRUE("Point2_ExpectedPlantGradient",  FMath::IsNearlyEqual(PointPlantGradientAttribute[2],  Point2_ExpectedPlantGradient,  0.01f));
+	UTEST_TRUE("Point3_ExpectedPlantGradient",  FMath::IsNearlyEqual(PointPlantGradientAttribute[3],  Point3_ExpectedPlantGradient,  0.01f));
+	UTEST_TRUE("Point4_ExpectedPlantGradient",  FMath::IsNearlyEqual(PointPlantGradientAttribute[4],  Point4_ExpectedPlantGradient,  0.01f));
+	UTEST_TRUE("Point5_ExpectedPlantGradient",  FMath::IsNearlyEqual(PointPlantGradientAttribute[5],  Point5_ExpectedPlantGradient,  0.01f));
+	UTEST_TRUE("Point6_ExpectedPlantGradient",  FMath::IsNearlyEqual(PointPlantGradientAttribute[6],  Point6_ExpectedPlantGradient,  0.01f));
+	UTEST_TRUE("Point7_ExpectedPlantGradient",  FMath::IsNearlyEqual(PointPlantGradientAttribute[7],  Point7_ExpectedPlantGradient,  0.01f));
+	UTEST_TRUE("Point8_ExpectedPlantGradient",  FMath::IsNearlyEqual(PointPlantGradientAttribute[8],  Point8_ExpectedPlantGradient,  0.01f));
+	UTEST_TRUE("Point9_ExpectedPlantGradient",  FMath::IsNearlyEqual(PointPlantGradientAttribute[9],  Point9_ExpectedPlantGradient,  0.01f));
+	UTEST_TRUE("Point10_ExpectedPlantGradient", FMath::IsNearlyEqual(PointPlantGradientAttribute[10], Point10_ExpectedPlantGradient, 0.01f));
+	UTEST_TRUE("Point11_ExpectedPlantGradient", FMath::IsNearlyEqual(PointPlantGradientAttribute[11], Point11_ExpectedPlantGradient, 0.01f));
+
+	return true;
+}
+
+PV_ATTRIBUTES_TEST(ComputeBudStatus)
+{
+	struct FBudStatusDataLayout {
+		int32 ApicalMeristem;
+		int32 Codominant;
+		int32 Axillary;
+		int32 Seed;
+		int32 Dormant;
+		int32 Triggered;
+		int32 NumTriggered;
+		int32 Inactive;
+		int32 BrokenTip;
+		int32 Broken;
+
+		bool operator==(const TArray<int32>& Other) const
+		{
+			return Other.Num() == 10
+				&& ApicalMeristem == Other[0]
+				&& Codominant == Other[1]
+				&& Axillary == Other[2]
+				&& Seed == Other[3]
+				&& Dormant == Other[4]
+				&& Triggered == Other[5]
+				&& NumTriggered == Other[6]
+				&& Inactive == Other[7]
+				&& BrokenTip == Other[8]
+				&& Broken == Other[9];
+		}
+	};
+
+	const FBudStatusDataLayout Point0_ExpectedBudStatus = {
+		.ApicalMeristem = 0,
+		.Codominant = 0,
+		.Axillary = 0,
+		.Seed = 1,
+		.Dormant = 0,
+		.Triggered = 1,
+		.NumTriggered = 1,
+		.Inactive = 1,
+		.BrokenTip = 0,
+		.Broken = 0,
+	};
+
+	const FBudStatusDataLayout Point1_ExpectedBudStatus = {
+		.ApicalMeristem = 0,
+		.Codominant = 0,
+		.Axillary = 0,
+		.Seed = 0,
+		.Dormant = 0,
+		.Triggered = 1,
+		.NumTriggered = 1,
+		.Inactive = 1,
+		.BrokenTip = 0,
+		.Broken = 0,
+	};
+
+	const FBudStatusDataLayout Point2_ExpectedBudStatus = {
+		.ApicalMeristem = 0,
+		.Codominant = 0,
+		.Axillary = 1,
+		.Seed = 0,
+		.Dormant = 0,
+		.Triggered = 1,
+		.NumTriggered = 1,
+		.Inactive = 1,
+		.BrokenTip = 0,
+		.Broken = 0,
+	};
+
+	const FBudStatusDataLayout Point3_ExpectedBudStatus = {
+		.ApicalMeristem = 0,
+		.Codominant = 0,
+		.Axillary = 1,
+		.Seed = 0,
+		.Dormant = 0,
+		.Triggered = 1,
+		.NumTriggered = 1,
+		.Inactive = 1,
+		.BrokenTip = 0,
+		.Broken = 0,
+	};
+
+	const FBudStatusDataLayout Point4_ExpectedBudStatus = {
+		.ApicalMeristem = 0,
+		.Codominant = 0,
+		.Axillary = 0,
+		.Seed = 0,
+		.Dormant = 0,
+		.Triggered = 1,
+		.NumTriggered = 1,
+		.Inactive = 1,
+		.BrokenTip = 0,
+		.Broken = 0,
+	};
+
+	const FBudStatusDataLayout Point5_ExpectedBudStatus = {
+		.ApicalMeristem = 1,
+		.Codominant = 0,
+		.Axillary = 0,
+		.Seed = 0,
+		.Dormant = 1,
+		.Triggered = 0,
+		.NumTriggered = 0,
+		.Inactive = 0,
+		.BrokenTip = 0,
+		.Broken = 0,
+	};
+
+	const FBudStatusDataLayout Point6_ExpectedBudStatus = {
+		.ApicalMeristem = 1,
+		.Codominant = 0,
+		.Axillary = 0,
+		.Seed = 0,
+		.Dormant = 1,
+		.Triggered = 0,
+		.NumTriggered = 0,
+		.Inactive = 0,
+		.BrokenTip = 0,
+		.Broken = 0,
+	};
+
+	const FBudStatusDataLayout Point7_ExpectedBudStatus = {
+		.ApicalMeristem = 0,
+		.Codominant = 0,
+		.Axillary = 0,
+		.Seed = 0,
+		.Dormant = 0,
+		.Triggered = 1,
+		.NumTriggered = 1,
+		.Inactive = 1,
+		.BrokenTip = 0,
+		.Broken = 0,
+	};
+
+	const FBudStatusDataLayout Point8_ExpectedBudStatus = {
+		.ApicalMeristem = 0,
+		.Codominant = 0,
+		.Axillary = 1,
+		.Seed = 0,
+		.Dormant = 0,
+		.Triggered = 1,
+		.NumTriggered = 1,
+		.Inactive = 1,
+		.BrokenTip = 0,
+		.Broken = 0,
+	};
+
+	const FBudStatusDataLayout Point9_ExpectedBudStatus = {
+		.ApicalMeristem = 1,
+		.Codominant = 0,
+		.Axillary = 0,
+		.Seed = 0,
+		.Dormant = 1,
+		.Triggered = 0,
+		.NumTriggered = 0,
+		.Inactive = 0,
+		.BrokenTip = 0,
+		.Broken = 0,
+	};
+
+	const FBudStatusDataLayout Point10_ExpectedBudStatus = {
+		.ApicalMeristem = 0,
+		.Codominant = 0,
+		.Axillary = 0,
+		.Seed = 0,
+		.Dormant = 0,
+		.Triggered = 1,
+		.NumTriggered = 1,
+		.Inactive = 1,
+		.BrokenTip = 0,
+		.Broken = 0,
+	};
+
+	const FBudStatusDataLayout Point11_ExpectedBudStatus = {
+		.ApicalMeristem = 1,
+		.Codominant = 0,
+		.Axillary = 0,
+		.Seed = 0,
+		.Dormant = 1,
+		.Triggered = 0,
+		.NumTriggered = 0,
+		.Inactive = 0,
+		.BrokenTip = 0,
+		.Broken = 0,
+	};
+
+	const TSharedRef<FManagedArrayCollection> Collection = PVMockTreeCollection::CreateCollection();
+
+	const PV::FBudStatusAttributeConstView BudStatusAttribute = PV::FBudStatusAttribute::FindAttribute(*Collection);
+	UTEST_TRUE("HasAttribute", BudStatusAttribute);
+	UTEST_TRUE("Has 12 elements", BudStatusAttribute.Num() >= 12);
+
+	UTEST_TRUE("Point0_BudStatus",  BudStatusAttribute[0].Array  == Point0_ExpectedBudStatus);
+	UTEST_TRUE("Point1_BudStatus",  BudStatusAttribute[1].Array  == Point1_ExpectedBudStatus);
+	UTEST_TRUE("Point2_BudStatus",  BudStatusAttribute[2].Array  == Point2_ExpectedBudStatus);
+	UTEST_TRUE("Point3_BudStatus",  BudStatusAttribute[3].Array  == Point3_ExpectedBudStatus);
+	UTEST_TRUE("Point4_BudStatus",  BudStatusAttribute[4].Array  == Point4_ExpectedBudStatus);
+	UTEST_TRUE("Point5_BudStatus",  BudStatusAttribute[5].Array  == Point5_ExpectedBudStatus);
+	UTEST_TRUE("Point6_BudStatus",  BudStatusAttribute[6].Array  == Point6_ExpectedBudStatus);
+	UTEST_TRUE("Point7_BudStatus",  BudStatusAttribute[7].Array  == Point7_ExpectedBudStatus);
+	UTEST_TRUE("Point8_BudStatus",  BudStatusAttribute[8].Array  == Point8_ExpectedBudStatus);
+	UTEST_TRUE("Point9_BudStatus",  BudStatusAttribute[9].Array  == Point9_ExpectedBudStatus);
+	UTEST_TRUE("Point10_BudStatus", BudStatusAttribute[10].Array == Point10_ExpectedBudStatus);
+	UTEST_TRUE("Point11_BudStatus", BudStatusAttribute[11].Array == Point11_ExpectedBudStatus);
+
+	return true;
+}
+
+PV_ATTRIBUTES_TEST(ComputeBudHormoneLevels)
+{
+	struct FBudHormoneLevelsDataLayout
+	{
+		float Apical;
+		float Axillary;
+		float AxillaryInhibition;
+		float Radical;
+		float Ethylene;
+		float Cykininin;
+
+		bool operator==(const TArray<float>& Other) const
+		{
+			return Other.Num() == 6
+				&& FMath::IsNearlyEqual(Other[0], Apical, 0.01f)
+				&& FMath::IsNearlyEqual(Other[1], Axillary, 0.01f)
+				&& FMath::IsNearlyEqual(Other[2], AxillaryInhibition, 0.01f)
+				&& FMath::IsNearlyEqual(Other[3], Radical, 0.01f)
+				&& FMath::IsNearlyEqual(Other[4], Ethylene, 0.01f)
+				&& FMath::IsNearlyEqual(Other[5], Cykininin, 0.01f);
+		}
+	};
+
+	const FBudHormoneLevelsDataLayout Point0_ExpectedBudHormoneLevels = {
+		.Apical = 0,
+		.Axillary = 0,
+		.AxillaryInhibition = 0,
+		.Radical = 0,
+		.Ethylene = 1,
+		.Cykininin = 0,
+	};
+	const FBudHormoneLevelsDataLayout Point5_ExpectedBudHormoneLevels = {
+		.Apical = 1,
+		.Axillary = 0,
+		.AxillaryInhibition = 0,
+		.Radical = 0,
+		.Ethylene = 0,
+		.Cykininin = 0,
+	};
+	const FBudHormoneLevelsDataLayout Point6_ExpectedBudHormoneLevels = {
+		.Apical = 1,
+		.Axillary = 0,
+		.AxillaryInhibition = 0,
+		.Radical = 0,
+		.Ethylene = 0,
+		.Cykininin = 0,
+	};
+	const FBudHormoneLevelsDataLayout Point9_ExpectedBudHormoneLevels = {
+		.Apical = 1,
+		.Axillary = 0,
+		.AxillaryInhibition = 0,
+		.Radical = 0,
+		.Ethylene = 0,
+		.Cykininin = 0,
+	};
+	const FBudHormoneLevelsDataLayout Point11_ExpectedBudHormoneLevels = {
+		.Apical = 1,
+		.Axillary = 0,
+		.AxillaryInhibition = 0,
+		.Radical = 0,
+		.Ethylene = 0,
+		.Cykininin = 0,
+	};
+
+	const TSharedRef<FManagedArrayCollection> Collection = PVMockTreeCollection::CreateCollection();
+
+	const PV::FBudHormoneLevelsAttributeConstView BudHormoneLevelsAttribute = PV::FBudHormoneLevelsAttribute::FindAttribute(*Collection);
+	UTEST_TRUE("HasAttribute", BudHormoneLevelsAttribute.IsValid());
+	UTEST_TRUE("Has 12 elements", BudHormoneLevelsAttribute.Num() >= 12);
+
+	UTEST_TRUE("Point0_ExpectedBudHormoneLevels",  BudHormoneLevelsAttribute[0].Array  == Point0_ExpectedBudHormoneLevels);
+	UTEST_TRUE("Point5_ExpectedBudHormoneLevels",  BudHormoneLevelsAttribute[5].Array  == Point5_ExpectedBudHormoneLevels);
+	UTEST_TRUE("Point6_ExpectedBudHormoneLevels",  BudHormoneLevelsAttribute[6].Array  == Point6_ExpectedBudHormoneLevels);
+	UTEST_TRUE("Point9_ExpectedBudHormoneLevels",  BudHormoneLevelsAttribute[9].Array  == Point9_ExpectedBudHormoneLevels);
+	UTEST_TRUE("Point11_ExpectedBudHormoneLevels", BudHormoneLevelsAttribute[11].Array == Point11_ExpectedBudHormoneLevels);
+
+	return true;
+}
+
+PV_ATTRIBUTES_TEST(ComputeBudLateralMeristem)
+{
+	const TSharedRef<FManagedArrayCollection> Collection = PVMockTreeCollection::CreateCollection();
+	
+	const PV::FBudLateralMeristemAttributeConstView BudLateralMeristemAttribute = PV::FBudLateralMeristemAttribute::FindAttribute(*Collection);
+	UTEST_TRUE("HasAttribute", BudLateralMeristemAttribute.IsValid());
+	UTEST_TRUE("Has 12 elements", BudLateralMeristemAttribute.Num() >= 12);
+
+	UTEST_TRUE("Point0_DavinciIsNotZero",  !FMath::IsNearlyZero(BudLateralMeristemAttribute[0].Davinci));
+	UTEST_TRUE("Point1_DavinciIsNotZero",  !FMath::IsNearlyZero(BudLateralMeristemAttribute[1].Davinci));
+	UTEST_TRUE("Point2_DavinciIsNotZero",  !FMath::IsNearlyZero(BudLateralMeristemAttribute[2].Davinci));
+	UTEST_TRUE("Point3_DavinciIsNotZero",  !FMath::IsNearlyZero(BudLateralMeristemAttribute[3].Davinci));
+	UTEST_TRUE("Point4_DavinciIsNotZero",  !FMath::IsNearlyZero(BudLateralMeristemAttribute[4].Davinci));
+	UTEST_TRUE("Point7_DavinciIsNotZero",  !FMath::IsNearlyZero(BudLateralMeristemAttribute[7].Davinci));
+	UTEST_TRUE("Point8_DavinciIsNotZero",  !FMath::IsNearlyZero(BudLateralMeristemAttribute[8].Davinci));
+	UTEST_TRUE("Point10_DavinciIsNotZero", !FMath::IsNearlyZero(BudLateralMeristemAttribute[10].Davinci));
+
+	UTEST_TRUE("Point5_DavinciIsZero",  FMath::IsNearlyZero(BudLateralMeristemAttribute[5].Davinci));
+	UTEST_TRUE("Point6_DavinciIsZero",  FMath::IsNearlyZero(BudLateralMeristemAttribute[6].Davinci));
+	UTEST_TRUE("Point9_DavinciIsZero",  FMath::IsNearlyZero(BudLateralMeristemAttribute[9].Davinci));
+	UTEST_TRUE("Point11_DavinciIsZero", FMath::IsNearlyZero(BudLateralMeristemAttribute[11].Davinci));
+
+	return true;
+}
+
+PV_ATTRIBUTES_TEST(ComputeBudLightDetected)
+{
+	struct FBudLightDetectedDataLayout
+	{
+		float Availible;
+		float Resource;
+		float Branch;
+		float Collision;
+
+		bool operator==(const TArray<float>& Other) const
+		{
+			return Other.Num() == 4
+				&& Other[0] == Availible
+				&& Other[1] == Resource
+				&& Other[2] == Branch
+				&& Other[3] == Collision;
+		}
+	};
+
+	FBudLightDetectedDataLayout ExpectedBudLightDetected = {
+		.Availible = 1.0,
+		.Resource = 1.0,
+		.Branch = 1.0,
+		.Collision = 0.0,
+	};
+
+	const TSharedRef<FManagedArrayCollection> Collection = PVMockTreeCollection::CreateCollection();
+
+	const PV::FBudLightDetectedAttributeConstView BudLightDetectedAttribute = PV::FBudLightDetectedAttribute::FindAttribute(*Collection);
+	UTEST_TRUE("HasAttribute", BudLightDetectedAttribute.IsValid());
+	UTEST_TRUE("Has 12 elements", BudLightDetectedAttribute.Num() >= 12);
+
+	UTEST_TRUE("Point0_BudLightDetected",  BudLightDetectedAttribute[0].Array  == ExpectedBudLightDetected);
+	UTEST_TRUE("Point1_BudLightDetected",  BudLightDetectedAttribute[1].Array  == ExpectedBudLightDetected);
+	UTEST_TRUE("Point2_BudLightDetected",  BudLightDetectedAttribute[2].Array  == ExpectedBudLightDetected);
+	UTEST_TRUE("Point3_BudLightDetected",  BudLightDetectedAttribute[3].Array  == ExpectedBudLightDetected);
+	UTEST_TRUE("Point4_BudLightDetected",  BudLightDetectedAttribute[4].Array  == ExpectedBudLightDetected);
+	UTEST_TRUE("Point5_BudLightDetected",  BudLightDetectedAttribute[5].Array  == ExpectedBudLightDetected);
+	UTEST_TRUE("Point6_BudLightDetected",  BudLightDetectedAttribute[6].Array  == ExpectedBudLightDetected);
+	UTEST_TRUE("Point7_BudLightDetected",  BudLightDetectedAttribute[7].Array  == ExpectedBudLightDetected);
+	UTEST_TRUE("Point8_BudLightDetected",  BudLightDetectedAttribute[8].Array  == ExpectedBudLightDetected);
+	UTEST_TRUE("Point9_BudLightDetected",  BudLightDetectedAttribute[9].Array  == ExpectedBudLightDetected);
+	UTEST_TRUE("Point10_BudLightDetected", BudLightDetectedAttribute[10].Array == ExpectedBudLightDetected);
+	UTEST_TRUE("Point11_BudLightDetected", BudLightDetectedAttribute[11].Array == ExpectedBudLightDetected);
+
+	return true;
+}
+
+PV_ATTRIBUTES_TEST(ComputeBudDirections)
+{
+	const FVector3f ExpectedApical0 = (PVMockTree::Point1 - PVMockTree::Point0).GetSafeNormal();
+	const FVector3f ExpectedApical1 = (PVMockTree::Point2 - PVMockTree::Point1).GetSafeNormal();
+	const FVector3f ExpectedApical2 = (PVMockTree::Point3 - PVMockTree::Point2).GetSafeNormal();
+	const FVector3f ExpectedApical3 = (PVMockTree::Point4 - PVMockTree::Point3).GetSafeNormal();
+	const FVector3f ExpectedApical4 = (PVMockTree::Point5 - PVMockTree::Point4).GetSafeNormal();
+	const FVector3f ExpectedApical5 = (PVMockTree::Point5 - PVMockTree::Point4).GetSafeNormal();
+	const FVector3f ExpectedApical6 = (PVMockTree::Point6 - PVMockTree::Point3).GetSafeNormal();
+	const FVector3f ExpectedApical7 = (PVMockTree::Point8 - PVMockTree::Point7).GetSafeNormal();
+	const FVector3f ExpectedApical8 = (PVMockTree::Point9 - PVMockTree::Point8).GetSafeNormal();
+	const FVector3f ExpectedApical9 = (PVMockTree::Point9 - PVMockTree::Point8).GetSafeNormal();
+	const FVector3f ExpectedApical10 = (PVMockTree::Point11 - PVMockTree::Point10).GetSafeNormal();
+	const FVector3f ExpectedApical11 = (PVMockTree::Point11 - PVMockTree::Point10).GetSafeNormal();
+
+	const int32 ExpectedNumBudDirectionValues = 6;
+
+	const TSharedRef<FManagedArrayCollection> Collection = PVMockTreeCollection::CreateCollection();
+
+	const PV::FBudDirectionAttributeConstView BudDirectionAttribute = PV::FBudDirectionAttribute::FindAttribute(*Collection);
+	UTEST_TRUE("HasAttribute", BudDirectionAttribute.IsValid());
+	UTEST_TRUE("Has 12 elements", BudDirectionAttribute.Num() >= 12);
+
+	UTEST_NEARLY_EQUAL("Point0_Apical", FVector(BudDirectionAttribute[0].Apical), FVector(ExpectedApical0), UE_KINDA_SMALL_NUMBER);
+	UTEST_NEARLY_EQUAL("Point1_Apical", FVector(BudDirectionAttribute[1].Apical), FVector(ExpectedApical1), UE_KINDA_SMALL_NUMBER);
+	UTEST_NEARLY_EQUAL("Point2_Apical", FVector(BudDirectionAttribute[2].Apical), FVector(ExpectedApical2), UE_KINDA_SMALL_NUMBER);
+	UTEST_NEARLY_EQUAL("Point3_Apical", FVector(BudDirectionAttribute[3].Apical), FVector(ExpectedApical3), UE_KINDA_SMALL_NUMBER);
+	UTEST_NEARLY_EQUAL("Point4_Apical", FVector(BudDirectionAttribute[4].Apical), FVector(ExpectedApical4), UE_KINDA_SMALL_NUMBER);
+	UTEST_NEARLY_EQUAL("Point5_Apical", FVector(BudDirectionAttribute[5].Apical), FVector(ExpectedApical5), UE_KINDA_SMALL_NUMBER);
+	UTEST_NEARLY_EQUAL("Point6_Apical", FVector(BudDirectionAttribute[6].Apical), FVector(ExpectedApical6), UE_KINDA_SMALL_NUMBER);
+	UTEST_NEARLY_EQUAL("Point7_Apical", FVector(BudDirectionAttribute[7].Apical), FVector(ExpectedApical7), UE_KINDA_SMALL_NUMBER);
+	UTEST_NEARLY_EQUAL("Point8_Apical", FVector(BudDirectionAttribute[8].Apical), FVector(ExpectedApical8), UE_KINDA_SMALL_NUMBER);
+	UTEST_NEARLY_EQUAL("Point9_Apical", FVector(BudDirectionAttribute[9].Apical), FVector(ExpectedApical9), UE_KINDA_SMALL_NUMBER);
+	UTEST_NEARLY_EQUAL("Point10_Apical", FVector(BudDirectionAttribute[10].Apical), FVector(ExpectedApical10), UE_KINDA_SMALL_NUMBER);
+	UTEST_NEARLY_EQUAL("Point11_Apical", FVector(BudDirectionAttribute[11].Apical), FVector(ExpectedApical11), UE_KINDA_SMALL_NUMBER);
+
+	// TODO: Test UpVector and Axillary
+
+	return true;
+}
+
+PV_ATTRIBUTES_TEST(SplitFusedPoints)
+{
+	const TSharedRef<FManagedArrayCollection> Collection = PVMockTreeCollection::CreateCollection(false);
+
+	const bool bSplitFusedPointsResult = PV::AttributesHelper::SplitFusedPoints(*Collection);
+	UTEST_TRUE("SplitFusedPointsSucceeded", bSplitFusedPointsResult);
+
+	PV::FPointPositionAttributeConstView PointPositionAttribute = PV::FPointPositionAttribute::FindAttribute(*Collection);
+	UTEST_TRUE("HasAttribute", PointPositionAttribute.IsValid());
+	UTEST_TRUE("Has 15 elements", PointPositionAttribute.Num() == 15);
+
+	PV::FPointLengthFromRootAttributeConstView PointLengthFromRootAttribute = PV::FPointLengthFromRootAttribute::FindAttribute(*Collection);
+	UTEST_TRUE("HasAttribute", PointLengthFromRootAttribute.IsValid());
+	UTEST_TRUE("Has 15 elements", PointLengthFromRootAttribute.Num() == 15);
+
+	const PV::FPointLengthFromSeedAttributeView PointLengthFromSeedAttribute = PV::FPointLengthFromSeedAttribute::FindAttribute(*Collection);
+	UTEST_TRUE("HasAttribute", PointLengthFromSeedAttribute.IsValid());
+	UTEST_TRUE("Has 15 elements", PointLengthFromSeedAttribute.Num() == 15);
+
+	const PV::FBranchPointsAttributeView BranchPointsAttribute = PV::FBranchPointsAttribute::FindAttribute(*Collection);
+	UTEST_TRUE("HasAttribute", BranchPointsAttribute.IsValid());
+	UTEST_TRUE("Has 4 elements", BranchPointsAttribute.Num() == 4);
+
+	const int32 Branch0_OldRootPointIndex = 0;
+	const int32 Branch1_OldRootPointIndex = 2;
+	const int32 Branch2_OldRootPointIndex = 8;
+	const int32 Branch3_OldRootPointIndex = 3;
+
+	const int32 Branch0_NewRootPointIndex = BranchPointsAttribute[0][0];
+	const int32 Branch1_NewRootPointIndex = BranchPointsAttribute[1][0];
+	const int32 Branch2_NewRootPointIndex = BranchPointsAttribute[2][0];
+	const int32 Branch3_NewRootPointIndex = BranchPointsAttribute[3][0];
+
+	UTEST_TRUE("Branch0_HasSamePointIndex",  Branch0_NewRootPointIndex == Branch0_OldRootPointIndex);
+	UTEST_TRUE("Branch1_HasUniqePointIndex", Branch1_NewRootPointIndex != Branch1_OldRootPointIndex && Branch1_NewRootPointIndex > 11);
+	UTEST_TRUE("Branch2_HasUniqePointIndex", Branch2_NewRootPointIndex != Branch2_OldRootPointIndex && Branch2_NewRootPointIndex > 11);
+	UTEST_TRUE("Branch3_HasUniqePointIndex", Branch3_NewRootPointIndex != Branch3_OldRootPointIndex && Branch3_NewRootPointIndex > 11);
+
+	UTEST_EQUAL("Branch0_NewPointPositionIsSame", PointPositionAttribute[Branch0_NewRootPointIndex], PointPositionAttribute[Branch0_OldRootPointIndex]);
+	UTEST_EQUAL("Branch1_NewPointPositionIsSame", PointPositionAttribute[Branch1_NewRootPointIndex], PointPositionAttribute[Branch1_OldRootPointIndex]);
+	UTEST_EQUAL("Branch2_NewPointPositionIsSame", PointPositionAttribute[Branch2_NewRootPointIndex], PointPositionAttribute[Branch2_OldRootPointIndex]);
+	UTEST_EQUAL("Branch3_NewPointPositionIsSame", PointPositionAttribute[Branch3_NewRootPointIndex], PointPositionAttribute[Branch3_OldRootPointIndex]);
+
+	UTEST_EQUAL("Branch0_NewPointLengthFromRootIsSame", PointLengthFromRootAttribute[Branch0_NewRootPointIndex], PointLengthFromRootAttribute[Branch0_OldRootPointIndex]);
+	UTEST_EQUAL("Branch1_NewPointLengthFromRootIsSame", PointLengthFromRootAttribute[Branch1_NewRootPointIndex], PointLengthFromRootAttribute[Branch1_OldRootPointIndex]);
+	UTEST_EQUAL("Branch2_NewPointLengthFromRootIsSame", PointLengthFromRootAttribute[Branch2_NewRootPointIndex], PointLengthFromRootAttribute[Branch2_OldRootPointIndex]);
+	UTEST_EQUAL("Branch3_NewPointLengthFromRootIsSame", PointLengthFromRootAttribute[Branch3_NewRootPointIndex], PointLengthFromRootAttribute[Branch3_OldRootPointIndex]);
+
+	UTEST_EQUAL("Branch0_NewPointLengthFromSeedIsSame", PointLengthFromSeedAttribute[Branch0_NewRootPointIndex], PointLengthFromSeedAttribute[Branch0_OldRootPointIndex]);
+	UTEST_EQUAL("Branch1_NewPointLengthFromSeedIsSame", PointLengthFromSeedAttribute[Branch1_NewRootPointIndex], PointLengthFromSeedAttribute[Branch1_OldRootPointIndex]);
+	UTEST_EQUAL("Branch2_NewPointLengthFromSeedIsSame", PointLengthFromSeedAttribute[Branch2_NewRootPointIndex], PointLengthFromSeedAttribute[Branch2_OldRootPointIndex]);
+	UTEST_EQUAL("Branch3_NewPointLengthFromSeedIsSame", PointLengthFromSeedAttribute[Branch3_NewRootPointIndex], PointLengthFromSeedAttribute[Branch3_OldRootPointIndex]);
+
+	return true;
+}
+
+#endif //WITH_DEV_AUTOMATION_TESTS
\ No newline at end of file
diff --git a/Engine/Plugins/Experimental/ProceduralVegetationEditor/Source/ProceduralVegetation/Public/Helpers/PVAttributesHelper.h b/Engine/Plugins/Experimental/ProceduralVegetationEditor/Source/ProceduralVegetation/Public/Helpers/PVAttributesHelper.h
--- a/Engine/Plugins/Experimental/ProceduralVegetationEditor/Source/ProceduralVegetation/Public/Helpers/PVAttributesHelper.h
+++ b/Engine/Plugins/Experimental/ProceduralVegetationEditor/Source/ProceduralVegetation/Public/Helpers/PVAttributesHelper.h
@@ -601,4 +601,117 @@ namespace PV::AttributesHelper
 			ForEachPlantPoint(BranchPlantNumberAttribute, BranchPointsAttribute, PlantNumber, Body);
 		}
 	}
+
+	// Re-computes the point bud numbers and writes them to the BudNumber and BranchSourceBudNumber attributes.
+	PROCEDURALVEGETATION_API extern void ComputeBudNumbers(
+		FPointBudNumberAttributeView PointBudNumberAttribute,
+		FBranchSourceBudNumberAttributeView BranchSourceBudNumberAttribute,
+		FBranchPointsAttributeConstView BranchPointsAttribute
+	);
+
+	// Re-computes all njord pixel indices and writes them to the NjordPixelIndexAttribute.
+	PROCEDURALVEGETATION_API extern void ComputeNjordPixelIndex(
+		PV::FPointNjordPixelIndexAttributeView PointNjordPixelIndexAttribute,
+		PV::FBranchPointsAttributeConstView BranchPointsAttribute,
+		PV::FBranchParentNumberAttributeConstView BranchParentNumberAttribute,
+		PV::FBranchNumberAttributeConstView BranchNumberAttribute,
+		PV::FBranchChildrenAttributeConstView BranchChildrenAttribute
+	);
+
+	// Re-computes the length from root and length from seed for all points.
+	// NOTE: Assumes fused points
+	PROCEDURALVEGETATION_API extern void ComputeLengthFromRoot(
+		PV::FPointLengthFromRootAttributeView PointLengthFromRootAttribute,
+		PV::FPointLengthFromSeedAttributeView PointLengthFromSeedAttribute,
+		PV::FPointPositionAttributeConstView PointPositionAttribute,
+		PV::FBranchPointsAttributeConstView BranchPointsAttribute,
+		PV::FBranchParentNumberAttributeConstView BranchParentNumberAttribute,
+		PV::FBranchNumberAttributeConstView BranchNumberAttribute,
+		PV::FBranchChildrenAttributeConstView BranchChildrenAttribute
+	);
+
+	// Re-computes bud development for all points using the existing branch hierarchy and plant gradient. 
+	// NOTE: Assumes fused points
+	PROCEDURALVEGETATION_API extern void ComputeBudDevelopment(
+		PV::FBudDevelopmentAttributeView BudDevelopmentAttribute,
+		PV::FPointPlantGradientAttributeConstView PointPlantGradientAttribute,
+		PV::FBranchPointsAttributeConstView BranchPointsAttribute,
+		PV::FBranchParentNumberAttributeConstView BranchParentNumberAttribute,
+		PV::FBranchParentsAttributeConstView BranchParentsAttribute,
+		PV::FBranchNumberAttributeConstView BranchNumberAttribute
+	);
+
+	// Re-computes the bud directions for all points.
+	// NOTE: Assumes fused points
+	PROCEDURALVEGETATION_API extern void ComputeBudDirections(
+		PV::FBudDirectionAttributeView BudDirectionAttribute,
+		PV::FPointPositionAttributeConstView PointPositionAttribute,
+		PV::FBranchChildrenAttributeConstView BranchChildrenAttribute,
+		PV::FBranchParentNumberAttributeConstView BranchParentNumberAttribute,
+		PV::FBranchNumberAttributeConstView BranchNumberAttribute,
+		PV::FBranchPointsAttributeConstView BranchPointsAttribute
+	);
+
+	// Re-computes the bud hormone levels using the plant gradient for all points.
+	PROCEDURALVEGETATION_API extern void ComputeBudHormoneLevels(
+		PV::FBudHormoneLevelsAttributeView BudHormoneLevelsAttribute,
+		PV::FPointPlantGradientAttributeConstView PointPlantGradientAttribute
+	);
+
+	// Re-computes the point- ground gradient, scale gradient, plant gradient, and branch gradient.
+	PROCEDURALVEGETATION_API extern void ComputeGradients(
+		PV::FPointGroundGradientAttributeView PointGroundGradientAttribute,
+		PV::FPointScaleGradientAttributeView PointScaleGradientAttribute,
+		PV::FPointPlantGradientAttributeView PointPlantGradientAttribute,
+		PV::FPointBranchGradientAttributeView PointBranchGradientAttribute,
+		PV::FPointLengthFromRootAttributeConstView PointLengthFromRootAttribute,
+		PV::FPointScaleAttributeConstView PointScaleAttribute,
+		PV::FBranchPointsAttributeConstView BranchPointsAttribute,
+		PV::FBranchParentNumberAttributeConstView BranchParentNumberAttribute,
+		PV::FBranchNumberAttributeConstView BranchNumberAttribute,
+		PV::FBranchChildrenAttributeConstView BranchChildrenAttribute
+	);
+
+	// Re-computes the bud status for all points.
+	// NOTE: Assumes fused points
+	PROCEDURALVEGETATION_API extern void ComputeBudStatus(
+		PV::FBudStatusAttributeView BudStatusAttribute,
+		PV::FBranchPointsAttributeConstView BranchPointsAttribute,
+		PV::FBranchParentNumberAttributeConstView BranchParentNumberAttribute,
+		PV::FBranchNumberAttributeConstView BranchNumberAttribute,
+		PV::FBranchChildrenAttributeConstView BranchChildrenAttribute
+	);
+
+	// Re-computes the bud lateral meristem for all points.
+	PROCEDURALVEGETATION_API extern void ComputeBudLateralMeristem(
+		PV::FBudLateralMeristemAttributeView BudLateralMeristemAttribute,
+		PV::FPointPositionAttributeConstView PointPositionAttribute,
+		PV::FPointScaleAttributeConstView PointScaleAttribute,
+		PV::FPointLengthFromRootAttributeConstView PointLengthFromRootAttribute,
+		PV::FBudDirectionAttributeConstView BudDirectionAttribute,
+		PV::FBranchPointsAttributeConstView BranchPointsAttribute,
+		PV::FBranchParentNumberAttributeConstView BranchParentNumberAttribute,
+		PV::FBranchChildrenAttributeConstView BranchChildrenAttribute,
+		PV::FBranchNumberAttributeConstView BranchNumberAttribute
+	);
+
+	// Re-computes the bud light detected for all points.
+	PROCEDURALVEGETATION_API extern void ComputeBudLightDetected(PV::FBudLightDetectedAttributeView BudLightDetectedAttribute);
+
+	// This function ensures all points have a valid point scale by walking the point hierarchy
+	// and finding points with missing point scales. It then interpolates the missing point scales
+	// using the point scale of the parent/child points.
+	PROCEDURALVEGETATION_API extern void InterpolatePointsScales(
+		PV::FPointScaleAttributeView PointScaleAttribute,
+		PV::FPointLengthFromRootAttributeConstView PointLengthFromRootAttribute,
+		PV::FBranchPointsAttributeConstView BranchPointsAttribute,
+		PV::FBranchParentNumberAttributeConstView BranchParentNumberAttribute,
+		PV::FBranchChildrenAttributeConstView BranchChildrenAttribute,
+		PV::FBranchNumberAttributeConstView BranchNumberAttribute
+	);
+
+	// Converts a growth data collection with fused points to a split points collection.
+	// Will not validate that the input collection has fused points.
+	// NOTE: Assumes input collection is valid growth data.
+	PROCEDURALVEGETATION_API extern bool SplitFusedPoints(FManagedArrayCollection& Collection);
 }
diff --git a/Engine/Plugins/Experimental/ProceduralVegetationEditor/Source/ProceduralVegetationEditor/Private/Importer/PVImportHelpers.cpp b/Engine/Plugins/Experimental/ProceduralVegetationEditor/Source/ProceduralVegetationEditor/Private/Importer/PVImportHelpers.cpp
--- a/Engine/Plugins/Experimental/ProceduralVegetationEditor/Source/ProceduralVegetationEditor/Private/Importer/PVImportHelpers.cpp
+++ b/Engine/Plugins/Experimental/ProceduralVegetationEditor/Source/ProceduralVegetationEditor/Private/Importer/PVImportHelpers.cpp
@@ -887,618 +887,6 @@ void PV::ImportHelper::FillAttributesFromBranchHierarchy(
 	}
 }
 
-void PV::ImportHelper::ComputeBudNumbers(
-	FPointBudNumberAttributeView PointBudNumberAttribute,
-	FBranchSourceBudNumberAttributeView BranchSourceBudNumberAttribute,
-	FBranchPointsAttributeConstView BranchPointsAttribute
-)
-{
-	for (int32 PointIndex = 0; PointIndex < PointBudNumberAttribute.Num(); ++PointIndex)
-	{
-		PointBudNumberAttribute[PointIndex] = PointIndex + 1;
-	}
-
-	for (int32 BranchIndex = 0; BranchIndex < BranchSourceBudNumberAttribute.Num(); ++BranchIndex)
-	{
-		check(BranchPointsAttribute[BranchIndex].Num() > 1);
-		const int32 BranchRootPointIndex = BranchPointsAttribute[BranchIndex][0];
-		BranchSourceBudNumberAttribute[BranchIndex] = PointBudNumberAttribute[BranchRootPointIndex];
-	}
-}
-
-void PV::ImportHelper::ComputeNjordPixelIndex(
-	PV::FPointNjordPixelIndexAttributeView PointNjordPixelIndexAttribute,
-	PV::FBranchPointsAttributeConstView BranchPointsAttribute,
-	PV::FBranchParentNumberAttributeConstView BranchParentNumberAttribute,
-	PV::FBranchNumberAttributeConstView BranchNumberAttribute,
-	PV::FBranchChildrenAttributeConstView BranchChildrenAttribute
-)
-{
-	using namespace PV::AttributesHelper;
-
-	RecursiveWalkBranches<int32>(
-		BranchParentNumberAttribute,
-		BranchChildrenAttribute,
-		BranchNumberAttribute,
-		1, 
-		[&](int32 BranchIndex, int32& NjordPixelIndex)->EForEachResult
-		{
-			const bool bIsTrunk = GetBranchParentIndex(BranchParentNumberAttribute, BranchNumberAttribute, BranchIndex) == INDEX_NONE;
-			const TArray<int32>& BranchPoints = BranchPointsAttribute[BranchIndex];
-			for (int32 i = 0; i < BranchPoints.Num(); ++i)
-			{
-				if (bIsTrunk || i != 0)
-				{
-					PointNjordPixelIndexAttribute[BranchPoints[i]] = NjordPixelIndex;
-					NjordPixelIndex += 1;
-				}
-			}
-
-			return EForEachResult::Continue;
-		}
-	);
-}
-
-void PV::ImportHelper::ComputeLengthFromRoot(
-	PV::FPointLengthFromRootAttributeView PointLengthFromRootAttribute,
-	PV::FPointLengthFromSeedAttributeView PointLengthFromSeedAttribute,
-	PV::FPointPositionAttributeConstView PointPositionAttribute,
-	PV::FBranchPointsAttributeConstView BranchPointsAttribute,
-	PV::FBranchParentNumberAttributeConstView BranchParentNumberAttribute,
-	PV::FBranchNumberAttributeConstView BranchNumberAttribute,
-	PV::FBranchChildrenAttributeConstView BranchChildrenAttribute
-)
-{
-	using namespace PV::AttributesHelper;
-
-	RecursiveWalkBranches(BranchParentNumberAttribute, BranchChildrenAttribute, BranchNumberAttribute, [&](int32 BranchIndex)->EForEachResult
-	{
-		const TArray<int32>& CurrentBranchPoints = BranchPointsAttribute[BranchIndex];
-		check(CurrentBranchPoints.Num() > 0);
-
-		const int32 FirstPointIndex = CurrentBranchPoints[0];
-		const int32 ParentBranchIndex = GetBranchParentIndex(BranchParentNumberAttribute, BranchNumberAttribute, BranchIndex);
-
-		if (ParentBranchIndex == INDEX_NONE)
-		{
-			PointLengthFromRootAttribute[FirstPointIndex] = 0.f; // No parent, length to root is 0
-			PointLengthFromSeedAttribute[FirstPointIndex] = 0.f;
-		}
-		else
-		{
-			const TArray<int32>& ParentBranchPoints = BranchPointsAttribute[ParentBranchIndex];
-			check(ParentBranchPoints.Num() > 0);
-
-			const FVector3f& FirstPoint = PointPositionAttribute[FirstPointIndex];
-
-			int32 ClosestParentPointIndex = ParentBranchPoints[0];
-			float ClosestParentPointDistSquared = (PointPositionAttribute[ClosestParentPointIndex] - FirstPoint).SizeSquared();
-
-			for (int32 i = 1; i < ParentBranchPoints.Num(); ++i)
-			{
-				const int32 PointIndex = ParentBranchPoints[i];
-				const FVector3f ParentPoint = PointPositionAttribute[PointIndex];
-
-				const float DistSquared = (ParentPoint - FirstPoint).SizeSquared();
-				if (DistSquared > ClosestParentPointDistSquared)
-				{
-					continue;
-				}
-
-				const FVector3f ParentBranchDirection = [&]()
-				{
-					if (i < ParentBranchPoints.Num() - 1)
-					{
-						const int32 NextPointIndex = ParentBranchPoints[i + 1];
-						return (PointPositionAttribute[NextPointIndex] - PointPositionAttribute[PointIndex]).GetSafeNormal();
-					}
-
-					if (i > 0)
-					{
-						const int32 PrevPointIndex = ParentBranchPoints[i - 1];
-						return (PointPositionAttribute[PointIndex] - PointPositionAttribute[PrevPointIndex]).GetSafeNormal();
-					}
-
-					return FVector3f::UpVector;
-				}();
-
-				const FVector3f ParentPointToFirstPointDirection = (FirstPoint - ParentPoint).GetSafeNormal();
-				if (ParentPointToFirstPointDirection.Dot(ParentBranchDirection) > 0)
-				{
-					ClosestParentPointIndex = PointIndex;
-					ClosestParentPointDistSquared = DistSquared;
-				}
-			}
-
-			PointLengthFromRootAttribute[FirstPointIndex] = PointLengthFromRootAttribute[ClosestParentPointIndex] + (ClosestParentPointDistSquared > 0 ? FMath::Sqrt(ClosestParentPointDistSquared) : 0) * 0.01 /*cm->m*/;
-			PointLengthFromSeedAttribute[FirstPointIndex] = PointLengthFromRootAttribute[FirstPointIndex];
-		}
-
-		for (int32 i = 1; i < CurrentBranchPoints.Num(); ++i)
-		{
-			const int32 PointIndex = CurrentBranchPoints[i];
-			const int32 PrevPointIndex = CurrentBranchPoints[i - 1];
-
-			const FVector3f& Position = PointPositionAttribute[PointIndex];
-			const FVector3f& PrevPosition = PointPositionAttribute[PrevPointIndex];
-
-			PointLengthFromRootAttribute[PointIndex] = PointLengthFromRootAttribute[PrevPointIndex] + (Position - PrevPosition).Size() * 0.01 /*cm->m*/;
-			PointLengthFromSeedAttribute[PointIndex] = PointLengthFromRootAttribute[PointIndex];
-		}
-
-		return EForEachResult::Continue;
-	});
-}
-
-void PV::ImportHelper::ComputeBudDevelopment(
-	PV::FBudDevelopmentAttributeView BudDevelopmentAttribute,
-	PV::FPointPlantGradientAttributeConstView PointPlantGradientAttribute,
-	PV::FBranchPointsAttributeConstView BranchPointsAttribute,
-	PV::FBranchParentNumberAttributeConstView BranchParentNumberAttribute,
-	PV::FBranchParentsAttributeConstView BranchParentsAttribute,
-	PV::FBranchNumberAttributeConstView BranchNumberAttribute
-)
-{
-	using namespace PV::AttributesHelper;
-
-	int32 MaxBudAge = 0;
-	for (int32 BranchIndex = 0; BranchIndex < BranchPointsAttribute.Num(); ++BranchIndex)
-	{
-		MaxBudAge = FMath::Max(MaxBudAge, BranchPointsAttribute[BranchIndex].Num());
-	}
-
-	ForEachUniquePointOnBranches_FusedPoints(
-		BranchPointsAttribute,
-		BranchParentNumberAttribute,
-		BranchNumberAttribute,
-		[&](int32 BranchIndex, int32 PointIndex)
-		{
-			const float Gradient = PointPlantGradientAttribute[PointIndex];
-			const int32 BudAge = FMath::RoundToInt(FMath::GetMappedRangeValueClamped(FVector2f(0, 1), FVector2f(1, float(MaxBudAge)), Gradient));
-
-			BudDevelopmentAttribute[PointIndex].BudAge = BudAge;
-			BudDevelopmentAttribute[PointIndex].RelativeBudAge = BudAge;
-			BudDevelopmentAttribute[PointIndex].Generation = BranchParentsAttribute[BranchIndex].Num() + 1;
-
-			return EForEachResult::Continue;
-		}
-	);
-
-	ForEachUniquePointOnBranches_FusedPoints(
-		BranchPointsAttribute,
-		BranchParentNumberAttribute,
-		BranchNumberAttribute,
-		[&](int32 BranchIndex, int32 PointIndex)
-		{
-			const bool bIsRootBranch = GetBranchParentIndex(BranchParentNumberAttribute, BranchNumberAttribute, BranchIndex) == INDEX_NONE;
-			const int32 RootPointIndex = bIsRootBranch ? BranchPointsAttribute[BranchIndex][0] : BranchPointsAttribute[BranchIndex][1];
-			const int32 BranchAge = BudDevelopmentAttribute[RootPointIndex].BudAge;
-			BudDevelopmentAttribute[PointIndex].BranchAge = BranchAge;
-
-			return EForEachResult::Continue;
-		}
-	);
-}
-
-void PV::ImportHelper::ComputeBudDirections(
-	PV::FBudDirectionAttributeView BudDirectionAttribute,
-	PV::FPointPositionAttributeConstView PointPositionAttribute,
-	PV::FBranchChildrenAttributeConstView BranchChildrenAttribute,
-	PV::FBranchParentNumberAttributeConstView BranchParentNumberAttribute,
-	PV::FBranchNumberAttributeConstView BranchNumberAttribute,
-	PV::FBranchPointsAttributeConstView BranchPointsAttribute
-)
-{
-	using namespace PV::AttributesHelper;
-
-	RecursiveWalkPlantPoints_FusedPoints<int32>(
-		BranchChildrenAttribute,
-		BranchParentNumberAttribute,
-		BranchNumberAttribute,
-		BranchPointsAttribute,
-		INDEX_NONE,
-		[&](int32 BranchIndex, int32 BranchPointIndex, int32 PointIndex, const TArray<int32>& ChildPointIndices, int32& PrevPointIndex)->EForEachResult
-		{
-			const TArray<int32>& BranchPoints = BranchPointsAttribute[BranchIndex];
-
-			const FVector3f Apical = [&]()->FVector3f
-			{
-				if (BranchPointIndex < BranchPoints.Num() - 1)
-				{
-					const int32 NextPointIndex = BranchPoints[BranchPointIndex + 1];
-					return (PointPositionAttribute[NextPointIndex] - PointPositionAttribute[PointIndex]).GetSafeNormal();
-				}
-				else if (BranchPointIndex > 0)
-				{
-					const int32 PrevPointIndex = BranchPoints[BranchPointIndex - 1];
-					return (PointPositionAttribute[PointIndex] - PointPositionAttribute[PrevPointIndex]).GetSafeNormal();
-				}
-
-				ensure(false); // We should never be able to get here unless we have a branch with only 1 point
-				return FVector3f::UpVector;
-			}();
-
-			const FVector3f UpVector = FRotationMatrix44f::MakeFromZX(Apical, FVector3f::ForwardVector).ToQuat().GetRightVector() * -1;
-			const FVector3f PrevAxillary = PrevPointIndex == INDEX_NONE ? FVector3f::ForwardVector : BudDirectionAttribute[PrevPointIndex].Axillary;
-			const FVector3f Axillary = FQuat4f(Apical, FMath::DegreesToRadians(135.f)).RotateVector(PrevAxillary);
-
-			BudDirectionAttribute[PointIndex].Apical = Apical;
-			BudDirectionAttribute[PointIndex].Axillary = Axillary;
-			BudDirectionAttribute[PointIndex].LightOptimal = Apical;
-			BudDirectionAttribute[PointIndex].LightSubOptimal = Apical;
-			BudDirectionAttribute[PointIndex].GuideCurve = Apical;
-			BudDirectionAttribute[PointIndex].UpVector = UpVector;
-			
-			PrevPointIndex = PointIndex;
-
-			return EForEachResult::Continue;
-		}
-	);
-}
-
-void PV::ImportHelper::ComputeBudHormoneLevels(
-	PV::FBudHormoneLevelsAttributeView BudHormoneLevelsAttribute,
-	PV::FPointPlantGradientAttributeConstView PointPlantGradientAttribute
-)
-{
-	using namespace PV::AttributesHelper;
-		
-	for (int32 PointIndex = 0; PointIndex < BudHormoneLevelsAttribute.Num(); PointIndex++)
-	{
-		const float Gradient = PointPlantGradientAttribute[PointIndex];
-		const float InverseGradient = 1.f - Gradient;
-
-		BudHormoneLevelsAttribute[PointIndex].Apical = FMath::GetMappedRangeValueClamped(FVector2f(0.75, 1.0), FVector2f(0, 1), InverseGradient);
-		BudHormoneLevelsAttribute[PointIndex].Axillary = 0;
-		BudHormoneLevelsAttribute[PointIndex].AxillaryInhibition = 0;
-		BudHormoneLevelsAttribute[PointIndex].Radical = 0;
-		BudHormoneLevelsAttribute[PointIndex].Ethylene = Gradient;
-		BudHormoneLevelsAttribute[PointIndex].Cykininin = 0;
-	}
-}
-
-void PV::ImportHelper::ComputeGradients(
-	PV::FPointGroundGradientAttributeView PointGroundGradientAttribute,
-	PV::FPointScaleGradientAttributeView PointScaleGradientAttribute,
-	PV::FPointPlantGradientAttributeView PointPlantGradientAttribute,
-	PV::FPointBranchGradientAttributeView PointBranchGradientAttribute,
-	PV::FPointLengthFromRootAttributeConstView PointLengthFromRootAttribute,
-	PV::FPointScaleAttributeConstView PointScaleAttribute,
-	PV::FBranchPointsAttributeConstView BranchPointsAttribute,
-	PV::FBranchParentNumberAttributeConstView BranchParentNumberAttribute,
-	PV::FBranchNumberAttributeConstView BranchNumberAttribute,
-	PV::FBranchChildrenAttributeConstView BranchChildrenAttribute
-)
-{
-	using namespace PV::AttributesHelper;
-
-	const float* MaxLengthFromRoot = Algo::MaxElement(PointLengthFromRootAttribute);
-	const float* MaxPointScale = Algo::MaxElement(PointScaleAttribute);
-	if (MaxLengthFromRoot == nullptr || MaxPointScale == nullptr || *MaxLengthFromRoot == 0 || *MaxPointScale == 0)
-	{
-		return;
-	}
-
-	for (int32 PointIndex = 0; PointIndex < PointGroundGradientAttribute.Num(); PointIndex++)
-	{
-		const float LengthFromRoot = PointLengthFromRootAttribute[PointIndex];
-		const float LengthFromRootAlpha = LengthFromRoot / *MaxLengthFromRoot;
-		PointGroundGradientAttribute[PointIndex] = 1.0 - LengthFromRootAlpha;
-
-		const float PointScale = PointScaleAttribute[PointIndex];
-		const float PointScaleAlpha = PointScale / *MaxPointScale;
-		PointScaleGradientAttribute[PointIndex] = PointScaleAlpha;
-	}
-
-	// PlantGradient - A gradient from 1-0 where the tip of each branch is 0, and the base of each branch is whatever the gradient of the parent branch is at that point.
-	// BranchGradient - A gradient from 1-0 where 1 = base of the branch and 0 = tip of the branch
-	RecursiveWalkBranches(
-		BranchParentNumberAttribute,
-		BranchChildrenAttribute,
-		BranchNumberAttribute,
-		[&](int32 BranchIndex)->EForEachResult
-		{
-			const TArray<int32>& BranchPoints = BranchPointsAttribute[BranchIndex];
-			check(BranchPoints.Num() > 1);
-
-			const float BasePointPlantGradient = GetBranchParentIndex(BranchParentNumberAttribute, BranchNumberAttribute, BranchIndex) == INDEX_NONE
-				? 1.0
-				: PointPlantGradientAttribute[BranchPoints[0]];
-
-			for (int32 i = 0; i < BranchPoints.Num(); ++i)
-			{
-				const float BranchAlpha = i / float(BranchPoints.Num() - 1);
-				const float PointPlantGradient = FMath::Lerp(BasePointPlantGradient, 0, BranchAlpha);
-				PointPlantGradientAttribute[BranchPoints[i]] = PointPlantGradient;
-				PointBranchGradientAttribute[BranchPoints[i]] = 1.0 - BranchAlpha;
-			}
-
-			return EForEachResult::Continue;
-		}
-	);
-}
-
-void PV::ImportHelper::ComputeBudStatus(
-	PV::FBudStatusAttributeView BudStatusAttribute,
-	PV::FBranchPointsAttributeConstView BranchPointsAttribute,
-	PV::FBranchParentNumberAttributeConstView BranchParentNumberAttribute,
-	PV::FBranchNumberAttributeConstView BranchNumberAttribute,
-	PV::FBranchChildrenAttributeConstView BranchChildrenAttribute
-)
-{
-	using namespace PV::AttributesHelper;
-
-	ForEachUniquePointOnBranches_FusedPoints(
-		BranchPointsAttribute,
-		BranchParentNumberAttribute,
-		BranchNumberAttribute,
-		[&](int32 BranchIndex, int32 PointIndex)->EForEachResult
-		{
-			PV::FBudStatusView BudStatusView = BudStatusAttribute[PointIndex];
-			BudStatusView.ApicalMeristem = 0;
-			BudStatusView.Codominant = 0;
-			BudStatusView.Axillary = 0;
-			BudStatusView.Seed = 0;
-			BudStatusView.Dormant = 0;
-			BudStatusView.Triggered = 1;
-			BudStatusView.NumTriggered = 1;
-			BudStatusView.Inactive = 1;
-			BudStatusView.BrokenTip = 0;
-			BudStatusView.Broken = 0;
-
-			if (BranchPointsAttribute[BranchIndex].Last() == PointIndex) // IsTip
-			{
-				BudStatusView.ApicalMeristem = 1;
-				BudStatusView.Dormant = 1;
-				BudStatusView.Triggered = 0;
-				BudStatusView.NumTriggered = 0;
-				BudStatusView.Inactive = 0;
-			}
-			else if (BranchPointsAttribute[BranchIndex][0] == PointIndex) // Is branch root point.
-				// Due to ForEachUniquePointOnBranches_FusedPoints only iterating shared points on the parent branch, we don't have to check whether this is the root branch
-				// (as the root branch is the only branch which will have the root point iterated on by ForEachUniquePointOnBranches_FusedPoints).
-			{
-				BudStatusView.Seed = 1;
-			}
-			else
-			{
-				const int32 NumChildren = GetBranchPointChildBranchIndices_FusedPoints(
-					BranchChildrenAttribute,
-					BranchParentNumberAttribute,
-					BranchNumberAttribute,
-					BranchPointsAttribute,
-					BranchIndex, 
-					PointIndex
-				).Num();
-
-				if (NumChildren > 0) 
-				{
-					BudStatusView.NumTriggered = NumChildren;
-					BudStatusView.Axillary = 1;
-				}
-			}
-
-			return EForEachResult::Continue;
-		}
-	);
-}
-
-void PV::ImportHelper::ComputeBudLateralMeristem(
-	PV::FBudLateralMeristemAttributeView BudLateralMeristemAttribute,
-	PV::FPointPositionAttributeConstView PointPositionAttribute,
-	PV::FPointScaleAttributeConstView PointScaleAttribute,
-	PV::FPointLengthFromRootAttributeConstView PointLengthFromRootAttribute,
-	PV::FBudDirectionAttributeConstView BudDirectionAttribute,
-	PV::FBranchPointsAttributeConstView BranchPointsAttribute,
-	PV::FBranchParentNumberAttributeConstView BranchParentNumberAttribute,
-	PV::FBranchChildrenAttributeConstView BranchChildrenAttribute,
-	PV::FBranchNumberAttributeConstView BranchNumberAttribute
-)
-{
-	using namespace PV::AttributesHelper;
-
-	RecursiveWalkBranches_Reversed(
-		BranchParentNumberAttribute,
-		BranchChildrenAttribute,
-		BranchNumberAttribute,
-		[&](int32 BranchIndex)->EForEachResult
-		{
-			const TArray<int32>& BranchPointIndices = BranchPointsAttribute[BranchIndex];
-			check(BranchPointIndices.Num() > 1);
-
-			const FVector3f& FirstBranchPointPos = PointPositionAttribute[BranchPointIndices[0]];
-			const FVector3f& SecondBranchPointPos = PointPositionAttribute[BranchPointIndices[1]];
-
-			const FVector3f BranchDir = (FirstBranchPointPos - SecondBranchPointPos).GetSafeNormal();
-			const FVector3f FirstPointNormal = BudDirectionAttribute[BranchPointIndices[0]].Apical;;
-			
-			const float ParentDot = FVector3f::DotProduct(BranchDir, FirstPointNormal);
-
-			float PreviousPointDavinciValue = 0;
-			for (int32 i = BranchPointIndices.Num() - 1; i>= 0; --i)
-			{
-				const int32 PointIndex = BranchPointIndices[i];
-				const float ExisitingDavinciValue = BudLateralMeristemAttribute[PointIndex].Davinci;
-				const float PointScale = PointScaleAttribute[PointIndex];
-				const float AccumulatedDavinciValue = PreviousPointDavinciValue + ExisitingDavinciValue + PointScale;
-				PreviousPointDavinciValue = AccumulatedDavinciValue;
-
-				BudLateralMeristemAttribute[PointIndex].LateralMeristem = 0.0;
-				BudLateralMeristemAttribute[PointIndex].Multiplier = 1;
-				BudLateralMeristemAttribute[PointIndex].Inactive = 0.0;
-				BudLateralMeristemAttribute[PointIndex].Davinci = AccumulatedDavinciValue;
-				BudLateralMeristemAttribute[PointIndex].ParentDot = ParentDot;
-				BudLateralMeristemAttribute[PointIndex].RootDistance = PointLengthFromRootAttribute[PointIndex];
-				BudLateralMeristemAttribute[PointIndex].Degredation = 0.0;
-			}
-
-			return EForEachResult::Continue;
-		}
-	);
-}
-
-void PV::ImportHelper::ComputeBudLightDetected(PV::FBudLightDetectedAttributeView BudLightDetectedAttribute)
-{
-	for (int32 i = 0; i < BudLightDetectedAttribute.Num(); ++i)
-	{
-		BudLightDetectedAttribute[i].Availible = 1.0;
-		BudLightDetectedAttribute[i].Resource = 1.0;
-		BudLightDetectedAttribute[i].Branch = 1.0;
-		BudLightDetectedAttribute[i].Collision = 0.0;
-	}
-}
-
-void PV::ImportHelper::InterpolatePointsScales(
-	PV::FPointScaleAttributeView PointScaleAttribute,
-	PV::FPointLengthFromRootAttributeConstView PointLengthFromRootAttribute,
-	PV::FBranchPointsAttributeConstView BranchPointsAttribute,
-	PV::FBranchParentNumberAttributeConstView BranchParentNumberAttribute,
-	PV::FBranchChildrenAttributeConstView BranchChildrenAttribute,
-	PV::FBranchNumberAttributeConstView BranchNumberAttribute
-)
-{
-	using namespace PV::AttributesHelper;
-
-	RecursiveWalkBranches(
-		BranchParentNumberAttribute,
-		BranchChildrenAttribute,
-		BranchNumberAttribute,
-		[&](int32 BranchIndex)->EForEachResult
-		{
-			const TArray<int32>& BranchPoints = BranchPointsAttribute[BranchIndex];
-
-			float CurrentPointScale = PointScaleAttribute[BranchPoints.Last()];
-			float CurrentLengthFromRoot = PointLengthFromRootAttribute[BranchPoints.Last()];
-
-			for (int32 i = BranchPoints.Num() - 1; i >= 0; --i)
-			{
-				float NextValidScale = -1;
-				float NextValidScale_LengthFromRoot = -1;
-				int32 NextValidScale_Index = INDEX_NONE;
-
-				for (int32 j = i - 1; j >= 0; --j)
-				{
-					const int32 PointIndex = BranchPoints[j];
-					const float PointScale = PointScaleAttribute[PointIndex];
-					if (PointScale > 0)
-					{
-						NextValidScale = PointScale;
-						NextValidScale_LengthFromRoot = PointLengthFromRootAttribute[PointIndex];
-						NextValidScale_Index = j;
-						break;
-					}
-				}
-
-				if (NextValidScale_Index == INDEX_NONE)
-				{
-					for (int32 j = i - 1; j >= 0; --j)
-					{
-						PointScaleAttribute[BranchPoints[j]] = CurrentPointScale;
-					}
-					break;
-				}
-
-				if (NextValidScale_Index == i - 1)
-				{
-					CurrentPointScale = PointScaleAttribute[BranchPoints[NextValidScale_Index]];
-					continue;
-				}
-
-				for (int32 j = i - 1; j >= NextValidScale_Index; --j)
-				{
-					const int32 PointIndex = BranchPoints[j];
-					const float PointLengthFromRoot = PointLengthFromRootAttribute[PointIndex];
-					const float Alpha = FMath::GetMappedRangeValueClamped(FVector2D(CurrentLengthFromRoot, NextValidScale_LengthFromRoot), FVector2D(0, 1), PointLengthFromRoot);
-					const float PointScale = FMath::Lerp(CurrentPointScale, NextValidScale, Alpha);
-					PointScaleAttribute[PointIndex] = PointScale;
-				}
-
-				CurrentPointScale = NextValidScale;
-				CurrentLengthFromRoot = NextValidScale_LengthFromRoot;
-				i = NextValidScale_Index + 1;
-			}
-
-			return EForEachResult::Continue;
-		}
-	);
-}
-
-bool PV::ImportHelper::SplitFusedPoints(FManagedArrayCollection& Collection)
-{
-	if (!Collection.HasGroup(PV::GroupNames::PointGroup) || !Collection.HasGroup(PV::GroupNames::BranchGroup))
-	{
-		return false;
-	}
-
-	auto BranchPointsAttribute = FBranchPointsAttribute::GetAttribute(Collection);
-	auto BranchParentNumberAttribute = FBranchParentNumberAttribute::GetAttribute(Collection);
-
-	auto PointPositionAttribute = FPointPositionAttribute::GetAttribute(Collection);
-	auto LengthFromRootAttribute = FPointLengthFromRootAttribute::GetAttribute(Collection);
-	auto LengthFromSeedAttribute = FPointLengthFromSeedAttribute::GetAttribute(Collection);
-	auto PointScaleGradientAttribute = FPointScaleGradientAttribute::GetAttribute(Collection);
-	auto PointHullGradientAttribute = FPointHullGradientAttribute::GetAttribute(Collection);
-	auto PointMainTrunkGradientAttribute = FPointMainTrunkGradientAttribute::GetAttribute(Collection);
-	auto GroundGradientAttribute = FPointGroundGradientAttribute::GetAttribute(Collection);
-	auto PointScaleAttribute = FPointScaleAttribute::GetAttribute(Collection);
-	auto BudDirectionAttribute = FBudDirectionAttribute::GetAttribute(Collection);
-	auto BudStatusAttribute = FBudStatusAttribute::GetAttribute(Collection);
-	auto BudLateralMeristemAttribute = FBudLateralMeristemAttribute::GetAttribute(Collection);
-	auto PointBranchGradientAttribute = FPointBranchGradientAttribute::GetAttribute(Collection);
-	auto BudHormoneLevelsAttribute = FBudHormoneLevelsAttribute::GetAttribute(Collection);
-	auto PlantGradientAttribute = FPointPlantGradientAttribute::GetAttribute(Collection);
-	auto BudNumberAttribute = FPointBudNumberAttribute::GetAttribute(Collection);
-	auto NjordPixelIndexAttribute = FPointNjordPixelIndexAttribute::GetAttribute(Collection);
-	auto BudLightDetectedAttribute = FBudLightDetectedAttribute::GetAttribute(Collection);
-	auto BudDevelopmentAttribute = FBudDevelopmentAttribute::GetAttribute(Collection);
-
-	const int32* MaxBudNumber = Algo::MaxElement(BudNumberAttribute);
-	check(MaxBudNumber != nullptr);
-
-	int32 NewBudNumber = *MaxBudNumber + 1;
-
-	for (int32 i = 0; i < BranchPointsAttribute.Num(); ++i)
-	{
-		const int32 BranchIndex = i;
-		if (BranchParentNumberAttribute[BranchIndex] == 0)
-		{
-			continue;
-		}
-
-		const int32 NumPoints = PointPositionAttribute.Num();
-		Collection.Resize(NumPoints + 1, PV::GroupNames::PointGroup);
-
-		TArray<int32>& BranchPoints = BranchPointsAttribute[BranchIndex];
-		const int32 PrevPointIndex = BranchPoints[0];
-		const int32 NewPointIndex = NumPoints;
-
-		BranchPoints[0] = NewPointIndex;
-
-		PointPositionAttribute[NewPointIndex] = PointPositionAttribute[PrevPointIndex];
-		LengthFromRootAttribute[NewPointIndex] = LengthFromRootAttribute[PrevPointIndex];
-		LengthFromSeedAttribute[NewPointIndex] = LengthFromSeedAttribute[PrevPointIndex];
-		PointScaleGradientAttribute[NewPointIndex] = PointScaleGradientAttribute[PrevPointIndex];
-		PointHullGradientAttribute[NewPointIndex] = PointHullGradientAttribute[PrevPointIndex];
-		PointMainTrunkGradientAttribute[NewPointIndex] = PointMainTrunkGradientAttribute[PrevPointIndex];
-		GroundGradientAttribute[NewPointIndex] = GroundGradientAttribute[PrevPointIndex];
-		PointScaleAttribute[NewPointIndex] = PointScaleAttribute[PrevPointIndex];
-		BudDirectionAttribute[NewPointIndex].Array = BudDirectionAttribute[PrevPointIndex].Array;
-		BudStatusAttribute[NewPointIndex].Array = BudStatusAttribute[PrevPointIndex].Array;
-		BudLateralMeristemAttribute[NewPointIndex].Array = BudLateralMeristemAttribute[PrevPointIndex].Array;
-		PointBranchGradientAttribute[NewPointIndex] = PointBranchGradientAttribute[PrevPointIndex];
-		BudHormoneLevelsAttribute[NewPointIndex].Array = BudHormoneLevelsAttribute[PrevPointIndex].Array;
-		PlantGradientAttribute[NewPointIndex] = PlantGradientAttribute[PrevPointIndex];
-		BudNumberAttribute[NewPointIndex] = NewBudNumber;
-		NjordPixelIndexAttribute[NewPointIndex] = NjordPixelIndexAttribute[PrevPointIndex];
-		BudLightDetectedAttribute[NewPointIndex].Array = BudLightDetectedAttribute[PrevPointIndex].Array;
-		BudDevelopmentAttribute[NewPointIndex].Array = BudDevelopmentAttribute[PrevPointIndex].Array;
-
-		NewBudNumber++;
-	}
-
-	return true;
-}
-
 void PV::ImportHelper::ConnectDisconnectedBranches(FPVBranchHierarchyDescription& InOutBranchHierarchy)
 {
 	struct FNewBranchRootIndex
@@ -1803,6 +1191,8 @@ void PV::ImportHelper::FuseCollinearBranches(FPVBranchHierarchyDescription& InOu
 
 bool PV::ImportHelper::GenerateGrowthDataFromBranchHierarchies(FManagedArrayCollection& OutCollection, TArrayView<const FPVBranchHierarchyDescription> InBranchHierarchies, bool bSplitPoints)
 {
+	using namespace PV::AttributesHelper;
+
 	if (InBranchHierarchies.Num() < 1)
 	{
 		return false;
diff --git a/Engine/Plugins/Experimental/ProceduralVegetationEditor/Source/ProceduralVegetationEditor/Private/Importer/PVImportHelpers.h b/Engine/Plugins/Experimental/ProceduralVegetationEditor/Source/ProceduralVegetationEditor/Private/Importer/PVImportHelpers.h
--- a/Engine/Plugins/Experimental/ProceduralVegetationEditor/Source/ProceduralVegetationEditor/Private/Importer/PVImportHelpers.h
+++ b/Engine/Plugins/Experimental/ProceduralVegetationEditor/Source/ProceduralVegetationEditor/Private/Importer/PVImportHelpers.h
@@ -120,7 +120,13 @@ namespace PV::DynamicMeshHelper
 namespace PV::ImportHelper
 {
 	// Creates a collection containing all attributes (of the supplied size) required for creating PV growth data.
-	void CreateEmptyGrowthData(FManagedArrayCollection& OutCollection, int32 NumBranchPoints, int32 NumFoliagePoints, int32 NumBranches, int32 NumFoliageMeshes);
+	void CreateEmptyGrowthData(
+		FManagedArrayCollection& OutCollection,
+		int32 NumBranchPoints,
+		int32 NumFoliagePoints,
+		int32 NumBranches,
+		int32 NumFoliageMeshes
+	);
 
 	void FillAttributesFromBranchHierarchy(
 		FManagedArrayCollection& OutCollection,
@@ -142,119 +148,6 @@ namespace PV::ImportHelper
 		TArrayView<const FPVBranchHierarchyDescription> InBranchHierarchies
 	);
 
-	// Re-computes the point bud numbers and writes them to the BudNumber and BranchSourceBudNumber attributes.
-	void ComputeBudNumbers(
-		FPointBudNumberAttributeView PointBudNumberAttribute,
-		FBranchSourceBudNumberAttributeView BranchSourceBudNumberAttribute,
-		FBranchPointsAttributeConstView BranchPointsAttribute
-	);
-
-	// Re-computes all njord pixel indices and writes them to the NjordPixelIndexAttribute.
-	void ComputeNjordPixelIndex(
-		PV::FPointNjordPixelIndexAttributeView PointNjordPixelIndexAttribute,
-		PV::FBranchPointsAttributeConstView BranchPointsAttribute,
-		PV::FBranchParentNumberAttributeConstView BranchParentNumberAttribute,
-		PV::FBranchNumberAttributeConstView BranchNumberAttribute,
-		PV::FBranchChildrenAttributeConstView BranchChildrenAttribute
-	);
-
-	// Re-computes the length from root and length from seed for all points.
-	// NOTE: Assumes fused points
-	void ComputeLengthFromRoot(
-		PV::FPointLengthFromRootAttributeView PointLengthFromRootAttribute,
-		PV::FPointLengthFromSeedAttributeView PointLengthFromSeedAttribute,
-		PV::FPointPositionAttributeConstView PointPositionAttribute,
-		PV::FBranchPointsAttributeConstView BranchPointsAttribute,
-		PV::FBranchParentNumberAttributeConstView BranchParentNumberAttribute,
-		PV::FBranchNumberAttributeConstView BranchNumberAttribute,
-		PV::FBranchChildrenAttributeConstView BranchChildrenAttribute
-	);
-
-	// Re-computes bud development for all points using the existing branch hierarchy and plant gradient. 
-	// NOTE: Assumes fused points
-	void ComputeBudDevelopment(
-		PV::FBudDevelopmentAttributeView BudDevelopmentAttribute,
-		PV::FPointPlantGradientAttributeConstView PointPlantGradientAttribute,
-		PV::FBranchPointsAttributeConstView BranchPointsAttribute,
-		PV::FBranchParentNumberAttributeConstView BranchParentNumberAttribute,
-		PV::FBranchParentsAttributeConstView BranchParentsAttribute,
-		PV::FBranchNumberAttributeConstView BranchNumberAttribute
-	);
-
-	// Re-computes the bud directions for all points.
-	// NOTE: Assumes fused points
-	void ComputeBudDirections(
-		PV::FBudDirectionAttributeView BudDirectionAttribute,
-		PV::FPointPositionAttributeConstView PointPositionAttribute,
-		PV::FBranchChildrenAttributeConstView BranchChildrenAttribute,
-		PV::FBranchParentNumberAttributeConstView BranchParentNumberAttribute,
-		PV::FBranchNumberAttributeConstView BranchNumberAttribute,
-		PV::FBranchPointsAttributeConstView BranchPointsAttribute
-	);
-
-	// Re-computes the bud hormone levels using the plant gradient for all points.
-	void ComputeBudHormoneLevels(
-		PV::FBudHormoneLevelsAttributeView BudHormoneLevelsAttribute,
-		PV::FPointPlantGradientAttributeConstView PointPlantGradientAttribute
-	);
-
-	// Re-computes the point- ground gradient, scale gradient, plant gradient, and branch gradient.
-	void ComputeGradients(
-		PV::FPointGroundGradientAttributeView PointGroundGradientAttribute,
-		PV::FPointScaleGradientAttributeView PointScaleGradientAttribute,
-		PV::FPointPlantGradientAttributeView PointPlantGradientAttribute,
-		PV::FPointBranchGradientAttributeView PointBranchGradientAttribute,
-		PV::FPointLengthFromRootAttributeConstView PointLengthFromRootAttribute,
-		PV::FPointScaleAttributeConstView PointScaleAttribute,
-		PV::FBranchPointsAttributeConstView BranchPointsAttribute,
-		PV::FBranchParentNumberAttributeConstView BranchParentNumberAttribute,
-		PV::FBranchNumberAttributeConstView BranchNumberAttribute,
-		PV::FBranchChildrenAttributeConstView BranchChildrenAttribute
-	);
-
-	// Re-computes the bud status for all points.
-	// NOTE: Assumes fused points
-	void ComputeBudStatus(
-		PV::FBudStatusAttributeView BudStatusAttribute,
-		PV::FBranchPointsAttributeConstView BranchPointsAttribute,
-		PV::FBranchParentNumberAttributeConstView BranchParentNumberAttribute,
-		PV::FBranchNumberAttributeConstView BranchNumberAttribute,
-		PV::FBranchChildrenAttributeConstView BranchChildrenAttribute
-	);
-
-	// Re-computes the bud lateral meristem for all points.
-	void ComputeBudLateralMeristem(
-		PV::FBudLateralMeristemAttributeView BudLateralMeristemAttribute,
-		PV::FPointPositionAttributeConstView PointPositionAttribute,
-		PV::FPointScaleAttributeConstView PointScaleAttribute,
-		PV::FPointLengthFromRootAttributeConstView PointLengthFromRootAttribute,
-		PV::FBudDirectionAttributeConstView BudDirectionAttribute,
-		PV::FBranchPointsAttributeConstView BranchPointsAttribute,
-		PV::FBranchParentNumberAttributeConstView BranchParentNumberAttribute,
-		PV::FBranchChildrenAttributeConstView BranchChildrenAttribute,
-		PV::FBranchNumberAttributeConstView BranchNumberAttribute
-	);
-
-	// Re-computes the bud light detected for all points.
-	void ComputeBudLightDetected(PV::FBudLightDetectedAttributeView BudLightDetectedAttribute);
-
-	// This function ensures all points have a valid point scale by walking the point hierarchy
-	// and finding points with missing point scales. It then interpolates the missing point scales
-	// using the point scale of the parent/child points.
-	void InterpolatePointsScales(
-		PV::FPointScaleAttributeView PointScaleAttribute,
-		PV::FPointLengthFromRootAttributeConstView PointLengthFromRootAttribute,
-		PV::FBranchPointsAttributeConstView BranchPointsAttribute,
-		PV::FBranchParentNumberAttributeConstView BranchParentNumberAttribute,
-		PV::FBranchChildrenAttributeConstView BranchChildrenAttribute,
-		PV::FBranchNumberAttributeConstView BranchNumberAttribute
-	);
-
-	// Converts a growth data collection with fused points to a split points collection.
-	// Will not validate that the input collection has split points.
-	// NOTE: Assumes input collection is valid growth data.
-	bool SplitFusedPoints(FManagedArrayCollection& Collection);
-
 	void ConnectDisconnectedBranches(FPVBranchHierarchyDescription& InOutBranchHierarchy);
 
 	void FuseCollinearBranches(FPVBranchHierarchyDescription& InOutBranchHierarchy);
diff --git a/Engine/Plugins/Experimental/ProceduralVegetationEditor/Source/ProceduralVegetationEditor/Private/Tests/Importer/PVImportTests.cpp b/Engine/Plugins/Experimental/ProceduralVegetationEditor/Source/ProceduralVegetationEditor/Private/Tests/Importer/PVImportTests.cpp
--- a/Engine/Plugins/Experimental/ProceduralVegetationEditor/Source/ProceduralVegetationEditor/Private/Tests/Importer/PVImportTests.cpp
+++ b/Engine/Plugins/Experimental/ProceduralVegetationEditor/Source/ProceduralVegetationEditor/Private/Tests/Importer/PVImportTests.cpp
@@ -2,15 +2,8 @@
 
 #if WITH_DEV_AUTOMATION_TESTS
 
-#include "CoreTypes.h"
-#include "Misc/AutomationTest.h"
-#include "HAL/PreprocessorHelpers.h"
-#include "GeometryCollection/ManagedArrayCollection.h"
-
 #include "Tests/PVTestsCommon.h"
-#include "Tests/PVMockTree.h"
 #include "Importer/PVImportHelpers.h"
-
 #include "Facades/PVAttributesNames.h"
 
 #define PV_IMPORTER_TEST(TestName) PV_SIMPLE_AUTOMATION_TEST(Importer, TestName)
@@ -133,682 +126,6 @@ PV_IMPORTER_TEST(FillAttributesFromBranchHierarchy)
 	return true;
 }
 
-PV_IMPORTER_TEST(ComputeLengthFromRoot)
-{
-	const TSharedRef<FManagedArrayCollection> Collection = PVMockTreeCollection::CreateCollection();
-	
-	PV::FPointLengthFromRootAttributeConstView LengthFromRootAttribute = PV::FPointLengthFromRootAttribute::FindAttribute(*Collection);
-	UTEST_TRUE("HasAttribute", LengthFromRootAttribute.IsValid());
-	UTEST_TRUE("Has 12 elements", LengthFromRootAttribute.Num() >= 12);
-
-	UTEST_NEARLY_EQUAL("Point0_LengthFromRoot",  LengthFromRootAttribute[0],  PVMockTree::Point0_LengthFromRoot, 0.01f);
-	UTEST_NEARLY_EQUAL("Point1_LengthFromRoot",  LengthFromRootAttribute[1],  PVMockTree::Point1_LengthFromRoot, 0.01f);
-	UTEST_NEARLY_EQUAL("Point2_LengthFromRoot",  LengthFromRootAttribute[2],  PVMockTree::Point2_LengthFromRoot, 0.01f);
-	UTEST_NEARLY_EQUAL("Point3_LengthFromRoot",  LengthFromRootAttribute[3],  PVMockTree::Point3_LengthFromRoot, 0.01f);
-	UTEST_NEARLY_EQUAL("Point4_LengthFromRoot",  LengthFromRootAttribute[4],  PVMockTree::Point4_LengthFromRoot, 0.01f);
-	UTEST_NEARLY_EQUAL("Point5_LengthFromRoot",  LengthFromRootAttribute[5],  PVMockTree::Point5_LengthFromRoot, 0.01f);
-	UTEST_NEARLY_EQUAL("Point6_LengthFromRoot",  LengthFromRootAttribute[6],  PVMockTree::Point6_LengthFromRoot, 0.01f);
-	UTEST_NEARLY_EQUAL("Point7_LengthFromRoot",  LengthFromRootAttribute[7],  PVMockTree::Point7_LengthFromRoot, 0.01f);
-	UTEST_NEARLY_EQUAL("Point8_LengthFromRoot",  LengthFromRootAttribute[8],  PVMockTree::Point8_LengthFromRoot, 0.01f);
-	UTEST_NEARLY_EQUAL("Point9_LengthFromRoot",  LengthFromRootAttribute[9],  PVMockTree::Point9_LengthFromRoot, 0.01f);
-	UTEST_NEARLY_EQUAL("Point10_LengthFromRoot", LengthFromRootAttribute[10], PVMockTree::Point10_LengthFromRoot, 0.01f);
-	UTEST_NEARLY_EQUAL("Point11_LengthFromRoot", LengthFromRootAttribute[11], PVMockTree::Point11_LengthFromRoot, 0.01f);
-
-	const PV::FPointLengthFromSeedAttributeView LengthFromSeedAttribute = PV::FPointLengthFromSeedAttribute::FindAttribute(*Collection);
-	UTEST_TRUE("HasAttribute", LengthFromSeedAttribute.IsValid());
-	UTEST_TRUE("Has 12 elements", LengthFromSeedAttribute.Num() >= 12);
-
-	UTEST_NEARLY_EQUAL("Point0_LengthFromRoot",  LengthFromSeedAttribute[0],  PVMockTree::Point0_LengthFromRoot,  0.01f);
-	UTEST_NEARLY_EQUAL("Point1_LengthFromRoot",  LengthFromSeedAttribute[1],  PVMockTree::Point1_LengthFromRoot,  0.01f);
-	UTEST_NEARLY_EQUAL("Point2_LengthFromRoot",  LengthFromSeedAttribute[2],  PVMockTree::Point2_LengthFromRoot,  0.01f);
-	UTEST_NEARLY_EQUAL("Point3_LengthFromRoot",  LengthFromSeedAttribute[3],  PVMockTree::Point3_LengthFromRoot,  0.01f);
-	UTEST_NEARLY_EQUAL("Point4_LengthFromRoot",  LengthFromSeedAttribute[4],  PVMockTree::Point4_LengthFromRoot,  0.01f);
-	UTEST_NEARLY_EQUAL("Point5_LengthFromRoot",  LengthFromSeedAttribute[5],  PVMockTree::Point5_LengthFromRoot,  0.01f);
-	UTEST_NEARLY_EQUAL("Point6_LengthFromRoot",  LengthFromSeedAttribute[6],  PVMockTree::Point6_LengthFromRoot,  0.01f);
-	UTEST_NEARLY_EQUAL("Point7_LengthFromRoot",  LengthFromSeedAttribute[7],  PVMockTree::Point7_LengthFromRoot,  0.01f);
-	UTEST_NEARLY_EQUAL("Point8_LengthFromRoot",  LengthFromSeedAttribute[8],  PVMockTree::Point8_LengthFromRoot,  0.01f);
-	UTEST_NEARLY_EQUAL("Point9_LengthFromRoot",  LengthFromSeedAttribute[9],  PVMockTree::Point9_LengthFromRoot,  0.01f);
-	UTEST_NEARLY_EQUAL("Point10_LengthFromRoot", LengthFromSeedAttribute[10], PVMockTree::Point10_LengthFromRoot, 0.01f);
-	UTEST_NEARLY_EQUAL("Point11_LengthFromRoot", LengthFromSeedAttribute[11], PVMockTree::Point11_LengthFromRoot, 0.01f);
-
-	return true;
-}
-
-PV_IMPORTER_TEST(ComputeBudDevelopment)
-{
-	struct FBudDevelopmentDataLayout
-	{
-		int32 Generation;
-		int32 Age;
-		int32 Branch_Age;
-		int32 AgeSenescence;
-		int32 Light_Senescence;
-		int32 Relative_Age;
-
-		bool operator==(const TArray<int32>& Other) const
-		{
-			return Other.Num() == 6
-				&& Other[0] == Generation
-				&& Other[1] == Age
-				&& Other[2] == Branch_Age
-				&& Other[3] == AgeSenescence
-				&& Other[4] == Light_Senescence
-				&& Other[5] == Relative_Age;
-		}
-	};
-
-	const FBudDevelopmentDataLayout Point0_ExpectedBudDevelopment = {
-		.Generation = 1,
-		.Age = 6,
-		.Branch_Age = 6,
-		.AgeSenescence = 0,
-		.Light_Senescence = 0,
-		.Relative_Age = 6
-	};
-	const FBudDevelopmentDataLayout Point1_ExpectedBudDevelopment = {
-		.Generation = 1,
-		.Age = 5,
-		.Branch_Age = 6,
-		.AgeSenescence = 0,
-		.Light_Senescence = 0,
-		.Relative_Age = 5
-	};
-	const FBudDevelopmentDataLayout Point2_ExpectedBudDevelopment = {
-		.Generation = 1,
-		.Age = 4,
-		.Branch_Age = 6,
-		.AgeSenescence = 0,
-		.Light_Senescence = 0,
-		.Relative_Age = 4
-	};
-	const FBudDevelopmentDataLayout Point3_ExpectedBudDevelopment = {
-		.Generation = 1,
-		.Age = 3,
-		.Branch_Age = 6,
-		.AgeSenescence = 0,
-		.Light_Senescence = 0,
-		.Relative_Age = 3
-	};
-	const FBudDevelopmentDataLayout Point4_ExpectedBudDevelopment = {
-		.Generation = 1,
-		.Age = 2,
-		.Branch_Age = 6,
-		.AgeSenescence = 0,
-		.Light_Senescence = 0,
-		.Relative_Age = 2
-	};
-	const FBudDevelopmentDataLayout Point5_ExpectedBudDevelopment = {
-		.Generation = 1,
-		.Age = 1,
-		.Branch_Age = 6,
-		.AgeSenescence = 0,
-		.Light_Senescence = 0,
-		.Relative_Age = 1
-	};
-	const FBudDevelopmentDataLayout Point6_ExpectedBudDevelopment = {
-		.Generation = 2,
-		.Age = 1,
-		.Branch_Age = 1,
-		.AgeSenescence = 0,
-		.Light_Senescence = 0,
-		.Relative_Age = 1
-	};
-	const FBudDevelopmentDataLayout Point7_ExpectedBudDevelopment = {
-		.Generation = 2,
-		.Age = 3,
-		.Branch_Age = 3,
-		.AgeSenescence = 0,
-		.Light_Senescence = 0,
-		.Relative_Age = 3
-	};
-	const FBudDevelopmentDataLayout Point8_ExpectedBudDevelopment = {
-		.Generation = 2,
-		.Age = 2,
-		.Branch_Age = 3,
-		.AgeSenescence = 0,
-		.Light_Senescence = 0,
-		.Relative_Age = 2
-	};
-	const FBudDevelopmentDataLayout Point9_ExpectedBudDevelopment = {
-		.Generation = 2,
-		.Age = 1,
-		.Branch_Age = 3,
-		.AgeSenescence = 0,
-		.Light_Senescence = 0,
-		.Relative_Age = 1
-	};
-	const FBudDevelopmentDataLayout Point10_ExpectedBudDevelopment = {
-		.Generation = 3,
-		.Age = 2,
-		.Branch_Age = 2,
-		.AgeSenescence = 0,
-		.Light_Senescence = 0,
-		.Relative_Age = 2
-	};
-	const FBudDevelopmentDataLayout Point11_ExpectedBudDevelopment = {
-		.Generation = 3,
-		.Age = 1,
-		.Branch_Age = 2,
-		.AgeSenescence = 0,
-		.Light_Senescence = 0,
-		.Relative_Age = 1
-	};
-
-	const TSharedRef<FManagedArrayCollection> Collection = PVMockTreeCollection::CreateCollection();
-	const PV::FBudDevelopmentAttributeConstView BudDevelopmentAttribute = PV::FBudDevelopmentAttribute::FindAttribute(*Collection);
-	UTEST_TRUE("HasAttribute", BudDevelopmentAttribute);
-	UTEST_TRUE("Has 12 elements", BudDevelopmentAttribute.Num() >= 12);
-
-	UTEST_TRUE("Point0_BudDevelopment",  BudDevelopmentAttribute[0].Array  == Point0_ExpectedBudDevelopment);
-	UTEST_TRUE("Point1_BudDevelopment",  BudDevelopmentAttribute[1].Array  == Point1_ExpectedBudDevelopment);
-	UTEST_TRUE("Point2_BudDevelopment",  BudDevelopmentAttribute[2].Array  == Point2_ExpectedBudDevelopment);
-	UTEST_TRUE("Point3_BudDevelopment",  BudDevelopmentAttribute[3].Array  == Point3_ExpectedBudDevelopment);
-	UTEST_TRUE("Point4_BudDevelopment",  BudDevelopmentAttribute[4].Array  == Point4_ExpectedBudDevelopment);
-	UTEST_TRUE("Point5_BudDevelopment",  BudDevelopmentAttribute[5].Array  == Point5_ExpectedBudDevelopment);
-	UTEST_TRUE("Point6_BudDevelopment",  BudDevelopmentAttribute[6].Array  == Point6_ExpectedBudDevelopment);
-	UTEST_TRUE("Point7_BudDevelopment",  BudDevelopmentAttribute[7].Array  == Point7_ExpectedBudDevelopment);
-	UTEST_TRUE("Point8_BudDevelopment",  BudDevelopmentAttribute[8].Array  == Point8_ExpectedBudDevelopment);
-	UTEST_TRUE("Point9_BudDevelopment",  BudDevelopmentAttribute[9].Array  == Point9_ExpectedBudDevelopment);
-	UTEST_TRUE("Point10_BudDevelopment", BudDevelopmentAttribute[10].Array == Point10_ExpectedBudDevelopment);
-	UTEST_TRUE("Point11_BudDevelopment", BudDevelopmentAttribute[11].Array == Point11_ExpectedBudDevelopment);
-
-	return true;
-}
-
-PV_IMPORTER_TEST(ComputeGradients)
-{
-	const float Point0_ExpectedPlantGradient  = 1;
-	const float Point1_ExpectedPlantGradient  = 0.8;
-	const float Point2_ExpectedPlantGradient  = 0.6;
-	const float Point3_ExpectedPlantGradient  = 0.4;
-	const float Point4_ExpectedPlantGradient  = 0.2;
-	const float Point5_ExpectedPlantGradient  = 0;
-	const float Point6_ExpectedPlantGradient  = 0;
-	const float Point7_ExpectedPlantGradient  = FMath::Lerp(Point2_ExpectedPlantGradient, 0.f, 1.f / 3);
-	const float Point8_ExpectedPlantGradient  = FMath::Lerp(Point2_ExpectedPlantGradient, 0.f, 2.f / 3);
-	const float Point9_ExpectedPlantGradient  = 0;
-	const float Point10_ExpectedPlantGradient = FMath::Lerp(Point8_ExpectedPlantGradient, 0.f, 1.f / 2);
-	const float Point11_ExpectedPlantGradient = 0;
-
-	const TSharedRef<FManagedArrayCollection> Collection = PVMockTreeCollection::CreateCollection();
-	
-	const PV::FPointPlantGradientAttributeView PointPlantGradientAttribute = PV::FPointPlantGradientAttribute::FindAttribute(*Collection);
-	UTEST_TRUE("HasAttribute", PointPlantGradientAttribute.IsValid());
-	UTEST_TRUE("Has 12 elements", PointPlantGradientAttribute.Num() >= 12);
-
-	UTEST_TRUE("Point0_ExpectedPlantGradient",  FMath::IsNearlyEqual(PointPlantGradientAttribute[0],  Point0_ExpectedPlantGradient,  0.01f));
-	UTEST_TRUE("Point1_ExpectedPlantGradient",  FMath::IsNearlyEqual(PointPlantGradientAttribute[1],  Point1_ExpectedPlantGradient,  0.01f));
-	UTEST_TRUE("Point2_ExpectedPlantGradient",  FMath::IsNearlyEqual(PointPlantGradientAttribute[2],  Point2_ExpectedPlantGradient,  0.01f));
-	UTEST_TRUE("Point3_ExpectedPlantGradient",  FMath::IsNearlyEqual(PointPlantGradientAttribute[3],  Point3_ExpectedPlantGradient,  0.01f));
-	UTEST_TRUE("Point4_ExpectedPlantGradient",  FMath::IsNearlyEqual(PointPlantGradientAttribute[4],  Point4_ExpectedPlantGradient,  0.01f));
-	UTEST_TRUE("Point5_ExpectedPlantGradient",  FMath::IsNearlyEqual(PointPlantGradientAttribute[5],  Point5_ExpectedPlantGradient,  0.01f));
-	UTEST_TRUE("Point6_ExpectedPlantGradient",  FMath::IsNearlyEqual(PointPlantGradientAttribute[6],  Point6_ExpectedPlantGradient,  0.01f));
-	UTEST_TRUE("Point7_ExpectedPlantGradient",  FMath::IsNearlyEqual(PointPlantGradientAttribute[7],  Point7_ExpectedPlantGradient,  0.01f));
-	UTEST_TRUE("Point8_ExpectedPlantGradient",  FMath::IsNearlyEqual(PointPlantGradientAttribute[8],  Point8_ExpectedPlantGradient,  0.01f));
-	UTEST_TRUE("Point9_ExpectedPlantGradient",  FMath::IsNearlyEqual(PointPlantGradientAttribute[9],  Point9_ExpectedPlantGradient,  0.01f));
-	UTEST_TRUE("Point10_ExpectedPlantGradient", FMath::IsNearlyEqual(PointPlantGradientAttribute[10], Point10_ExpectedPlantGradient, 0.01f));
-	UTEST_TRUE("Point11_ExpectedPlantGradient", FMath::IsNearlyEqual(PointPlantGradientAttribute[11], Point11_ExpectedPlantGradient, 0.01f));
-
-	return true;
-}
-
-PV_IMPORTER_TEST(ComputeBudStatus)
-{
-	struct FBudStatusDataLayout {
-		int32 ApicalMeristem;
-		int32 Codominant;
-		int32 Axillary;
-		int32 Seed;
-		int32 Dormant;
-		int32 Triggered;
-		int32 NumTriggered;
-		int32 Inactive;
-		int32 BrokenTip;
-		int32 Broken;
-
-		bool operator==(const TArray<int32>& Other) const
-		{
-			return Other.Num() == 10
-				&& ApicalMeristem == Other[0]
-				&& Codominant == Other[1]
-				&& Axillary == Other[2]
-				&& Seed == Other[3]
-				&& Dormant == Other[4]
-				&& Triggered == Other[5]
-				&& NumTriggered == Other[6]
-				&& Inactive == Other[7]
-				&& BrokenTip == Other[8]
-				&& Broken == Other[9];
-		}
-	};
-
-	const FBudStatusDataLayout Point0_ExpectedBudStatus = {
-		.ApicalMeristem = 0,
-		.Codominant = 0,
-		.Axillary = 0,
-		.Seed = 1,
-		.Dormant = 0,
-		.Triggered = 1,
-		.NumTriggered = 1,
-		.Inactive = 1,
-		.BrokenTip = 0,
-		.Broken = 0,
-	};
-
-	const FBudStatusDataLayout Point1_ExpectedBudStatus = {
-		.ApicalMeristem = 0,
-		.Codominant = 0,
-		.Axillary = 0,
-		.Seed = 0,
-		.Dormant = 0,
-		.Triggered = 1,
-		.NumTriggered = 1,
-		.Inactive = 1,
-		.BrokenTip = 0,
-		.Broken = 0,
-	};
-
-	const FBudStatusDataLayout Point2_ExpectedBudStatus = {
-		.ApicalMeristem = 0,
-		.Codominant = 0,
-		.Axillary = 1,
-		.Seed = 0,
-		.Dormant = 0,
-		.Triggered = 1,
-		.NumTriggered = 1,
-		.Inactive = 1,
-		.BrokenTip = 0,
-		.Broken = 0,
-	};
-
-	const FBudStatusDataLayout Point3_ExpectedBudStatus = {
-		.ApicalMeristem = 0,
-		.Codominant = 0,
-		.Axillary = 1,
-		.Seed = 0,
-		.Dormant = 0,
-		.Triggered = 1,
-		.NumTriggered = 1,
-		.Inactive = 1,
-		.BrokenTip = 0,
-		.Broken = 0,
-	};
-
-	const FBudStatusDataLayout Point4_ExpectedBudStatus = {
-		.ApicalMeristem = 0,
-		.Codominant = 0,
-		.Axillary = 0,
-		.Seed = 0,
-		.Dormant = 0,
-		.Triggered = 1,
-		.NumTriggered = 1,
-		.Inactive = 1,
-		.BrokenTip = 0,
-		.Broken = 0,
-	};
-
-	const FBudStatusDataLayout Point5_ExpectedBudStatus = {
-		.ApicalMeristem = 1,
-		.Codominant = 0,
-		.Axillary = 0,
-		.Seed = 0,
-		.Dormant = 1,
-		.Triggered = 0,
-		.NumTriggered = 0,
-		.Inactive = 0,
-		.BrokenTip = 0,
-		.Broken = 0,
-	};
-
-	const FBudStatusDataLayout Point6_ExpectedBudStatus = {
-		.ApicalMeristem = 1,
-		.Codominant = 0,
-		.Axillary = 0,
-		.Seed = 0,
-		.Dormant = 1,
-		.Triggered = 0,
-		.NumTriggered = 0,
-		.Inactive = 0,
-		.BrokenTip = 0,
-		.Broken = 0,
-	};
-
-	const FBudStatusDataLayout Point7_ExpectedBudStatus = {
-		.ApicalMeristem = 0,
-		.Codominant = 0,
-		.Axillary = 0,
-		.Seed = 0,
-		.Dormant = 0,
-		.Triggered = 1,
-		.NumTriggered = 1,
-		.Inactive = 1,
-		.BrokenTip = 0,
-		.Broken = 0,
-	};
-
-	const FBudStatusDataLayout Point8_ExpectedBudStatus = {
-		.ApicalMeristem = 0,
-		.Codominant = 0,
-		.Axillary = 1,
-		.Seed = 0,
-		.Dormant = 0,
-		.Triggered = 1,
-		.NumTriggered = 1,
-		.Inactive = 1,
-		.BrokenTip = 0,
-		.Broken = 0,
-	};
-
-	const FBudStatusDataLayout Point9_ExpectedBudStatus = {
-		.ApicalMeristem = 1,
-		.Codominant = 0,
-		.Axillary = 0,
-		.Seed = 0,
-		.Dormant = 1,
-		.Triggered = 0,
-		.NumTriggered = 0,
-		.Inactive = 0,
-		.BrokenTip = 0,
-		.Broken = 0,
-	};
-
-	const FBudStatusDataLayout Point10_ExpectedBudStatus = {
-		.ApicalMeristem = 0,
-		.Codominant = 0,
-		.Axillary = 0,
-		.Seed = 0,
-		.Dormant = 0,
-		.Triggered = 1,
-		.NumTriggered = 1,
-		.Inactive = 1,
-		.BrokenTip = 0,
-		.Broken = 0,
-	};
-
-	const FBudStatusDataLayout Point11_ExpectedBudStatus = {
-		.ApicalMeristem = 1,
-		.Codominant = 0,
-		.Axillary = 0,
-		.Seed = 0,
-		.Dormant = 1,
-		.Triggered = 0,
-		.NumTriggered = 0,
-		.Inactive = 0,
-		.BrokenTip = 0,
-		.Broken = 0,
-	};
-
-	const TSharedRef<FManagedArrayCollection> Collection = PVMockTreeCollection::CreateCollection();
-
-	const PV::FBudStatusAttributeConstView BudStatusAttribute = PV::FBudStatusAttribute::FindAttribute(*Collection);
-	UTEST_TRUE("HasAttribute", BudStatusAttribute);
-	UTEST_TRUE("Has 12 elements", BudStatusAttribute.Num() >= 12);
-
-	UTEST_TRUE("Point0_BudStatus",  BudStatusAttribute[0].Array  == Point0_ExpectedBudStatus);
-	UTEST_TRUE("Point1_BudStatus",  BudStatusAttribute[1].Array  == Point1_ExpectedBudStatus);
-	UTEST_TRUE("Point2_BudStatus",  BudStatusAttribute[2].Array  == Point2_ExpectedBudStatus);
-	UTEST_TRUE("Point3_BudStatus",  BudStatusAttribute[3].Array  == Point3_ExpectedBudStatus);
-	UTEST_TRUE("Point4_BudStatus",  BudStatusAttribute[4].Array  == Point4_ExpectedBudStatus);
-	UTEST_TRUE("Point5_BudStatus",  BudStatusAttribute[5].Array  == Point5_ExpectedBudStatus);
-	UTEST_TRUE("Point6_BudStatus",  BudStatusAttribute[6].Array  == Point6_ExpectedBudStatus);
-	UTEST_TRUE("Point7_BudStatus",  BudStatusAttribute[7].Array  == Point7_ExpectedBudStatus);
-	UTEST_TRUE("Point8_BudStatus",  BudStatusAttribute[8].Array  == Point8_ExpectedBudStatus);
-	UTEST_TRUE("Point9_BudStatus",  BudStatusAttribute[9].Array  == Point9_ExpectedBudStatus);
-	UTEST_TRUE("Point10_BudStatus", BudStatusAttribute[10].Array == Point10_ExpectedBudStatus);
-	UTEST_TRUE("Point11_BudStatus", BudStatusAttribute[11].Array == Point11_ExpectedBudStatus);
-
-	return true;
-}
-
-PV_IMPORTER_TEST(ComputeBudHormoneLevels)
-{
-	struct FBudHormoneLevelsDataLayout
-	{
-		float Apical;
-		float Axillary;
-		float AxillaryInhibition;
-		float Radical;
-		float Ethylene;
-		float Cykininin;
-
-		bool operator==(const TArray<float>& Other) const
-		{
-			return Other.Num() == 6
-				&& FMath::IsNearlyEqual(Other[0], Apical, 0.01f)
-				&& FMath::IsNearlyEqual(Other[1], Axillary, 0.01f)
-				&& FMath::IsNearlyEqual(Other[2], AxillaryInhibition, 0.01f)
-				&& FMath::IsNearlyEqual(Other[3], Radical, 0.01f)
-				&& FMath::IsNearlyEqual(Other[4], Ethylene, 0.01f)
-				&& FMath::IsNearlyEqual(Other[5], Cykininin, 0.01f);
-		}
-	};
-
-	const FBudHormoneLevelsDataLayout Point0_ExpectedBudHormoneLevels = {
-		.Apical = 0,
-		.Axillary = 0,
-		.AxillaryInhibition = 0,
-		.Radical = 0,
-		.Ethylene = 1,
-		.Cykininin = 0,
-	};
-	const FBudHormoneLevelsDataLayout Point5_ExpectedBudHormoneLevels = {
-		.Apical = 1,
-		.Axillary = 0,
-		.AxillaryInhibition = 0,
-		.Radical = 0,
-		.Ethylene = 0,
-		.Cykininin = 0,
-	};
-	const FBudHormoneLevelsDataLayout Point6_ExpectedBudHormoneLevels = {
-		.Apical = 1,
-		.Axillary = 0,
-		.AxillaryInhibition = 0,
-		.Radical = 0,
-		.Ethylene = 0,
-		.Cykininin = 0,
-	};
-	const FBudHormoneLevelsDataLayout Point9_ExpectedBudHormoneLevels = {
-		.Apical = 1,
-		.Axillary = 0,
-		.AxillaryInhibition = 0,
-		.Radical = 0,
-		.Ethylene = 0,
-		.Cykininin = 0,
-	};
-	const FBudHormoneLevelsDataLayout Point11_ExpectedBudHormoneLevels = {
-		.Apical = 1,
-		.Axillary = 0,
-		.AxillaryInhibition = 0,
-		.Radical = 0,
-		.Ethylene = 0,
-		.Cykininin = 0,
-	};
-
-	const TSharedRef<FManagedArrayCollection> Collection = PVMockTreeCollection::CreateCollection();
-
-	const PV::FBudHormoneLevelsAttributeConstView BudHormoneLevelsAttribute = PV::FBudHormoneLevelsAttribute::FindAttribute(*Collection);
-	UTEST_TRUE("HasAttribute", BudHormoneLevelsAttribute.IsValid());
-	UTEST_TRUE("Has 12 elements", BudHormoneLevelsAttribute.Num() >= 12);
-
-	UTEST_TRUE("Point0_ExpectedBudHormoneLevels",  BudHormoneLevelsAttribute[0].Array  == Point0_ExpectedBudHormoneLevels);
-	UTEST_TRUE("Point5_ExpectedBudHormoneLevels",  BudHormoneLevelsAttribute[5].Array  == Point5_ExpectedBudHormoneLevels);
-	UTEST_TRUE("Point6_ExpectedBudHormoneLevels",  BudHormoneLevelsAttribute[6].Array  == Point6_ExpectedBudHormoneLevels);
-	UTEST_TRUE("Point9_ExpectedBudHormoneLevels",  BudHormoneLevelsAttribute[9].Array  == Point9_ExpectedBudHormoneLevels);
-	UTEST_TRUE("Point11_ExpectedBudHormoneLevels", BudHormoneLevelsAttribute[11].Array == Point11_ExpectedBudHormoneLevels);
-
-	return true;
-}
-
-PV_IMPORTER_TEST(ComputeBudLateralMeristem)
-{
-	const TSharedRef<FManagedArrayCollection> Collection = PVMockTreeCollection::CreateCollection();
-	
-	const PV::FBudLateralMeristemAttributeConstView BudLateralMeristemAttribute = PV::FBudLateralMeristemAttribute::FindAttribute(*Collection);
-	UTEST_TRUE("HasAttribute", BudLateralMeristemAttribute.IsValid());
-	UTEST_TRUE("Has 12 elements", BudLateralMeristemAttribute.Num() >= 12);
-
-	UTEST_TRUE("Point0_DavinciIsNotZero",  !FMath::IsNearlyZero(BudLateralMeristemAttribute[0].Davinci));
-	UTEST_TRUE("Point1_DavinciIsNotZero",  !FMath::IsNearlyZero(BudLateralMeristemAttribute[1].Davinci));
-	UTEST_TRUE("Point2_DavinciIsNotZero",  !FMath::IsNearlyZero(BudLateralMeristemAttribute[2].Davinci));
-	UTEST_TRUE("Point3_DavinciIsNotZero",  !FMath::IsNearlyZero(BudLateralMeristemAttribute[3].Davinci));
-	UTEST_TRUE("Point4_DavinciIsNotZero",  !FMath::IsNearlyZero(BudLateralMeristemAttribute[4].Davinci));
-	UTEST_TRUE("Point7_DavinciIsNotZero",  !FMath::IsNearlyZero(BudLateralMeristemAttribute[7].Davinci));
-	UTEST_TRUE("Point8_DavinciIsNotZero",  !FMath::IsNearlyZero(BudLateralMeristemAttribute[8].Davinci));
-	UTEST_TRUE("Point10_DavinciIsNotZero", !FMath::IsNearlyZero(BudLateralMeristemAttribute[10].Davinci));
-
-	UTEST_TRUE("Point5_DavinciIsZero",  FMath::IsNearlyZero(BudLateralMeristemAttribute[5].Davinci));
-	UTEST_TRUE("Point6_DavinciIsZero",  FMath::IsNearlyZero(BudLateralMeristemAttribute[6].Davinci));
-	UTEST_TRUE("Point9_DavinciIsZero",  FMath::IsNearlyZero(BudLateralMeristemAttribute[9].Davinci));
-	UTEST_TRUE("Point11_DavinciIsZero", FMath::IsNearlyZero(BudLateralMeristemAttribute[11].Davinci));
-
-	return true;
-}
-
-PV_IMPORTER_TEST(ComputeBudLightDetected)
-{
-	struct FBudLightDetectedDataLayout
-	{
-		float Availible;
-		float Resource;
-		float Branch;
-		float Collision;
-
-		bool operator==(const TArray<float>& Other) const
-		{
-			return Other.Num() == 4
-				&& Other[0] == Availible
-				&& Other[1] == Resource
-				&& Other[2] == Branch
-				&& Other[3] == Collision;
-		}
-	};
-
-	FBudLightDetectedDataLayout ExpectedBudLightDetected = {
-		.Availible = 1.0,
-		.Resource = 1.0,
-		.Branch = 1.0,
-		.Collision = 0.0,
-	};
-
-	const TSharedRef<FManagedArrayCollection> Collection = PVMockTreeCollection::CreateCollection();
-
-	const PV::FBudLightDetectedAttributeConstView BudLightDetectedAttribute = PV::FBudLightDetectedAttribute::FindAttribute(*Collection);
-	UTEST_TRUE("HasAttribute", BudLightDetectedAttribute.IsValid());
-	UTEST_TRUE("Has 12 elements", BudLightDetectedAttribute.Num() >= 12);
-
-	UTEST_TRUE("Point0_BudLightDetected",  BudLightDetectedAttribute[0].Array  == ExpectedBudLightDetected);
-	UTEST_TRUE("Point1_BudLightDetected",  BudLightDetectedAttribute[1].Array  == ExpectedBudLightDetected);
-	UTEST_TRUE("Point2_BudLightDetected",  BudLightDetectedAttribute[2].Array  == ExpectedBudLightDetected);
-	UTEST_TRUE("Point3_BudLightDetected",  BudLightDetectedAttribute[3].Array  == ExpectedBudLightDetected);
-	UTEST_TRUE("Point4_BudLightDetected",  BudLightDetectedAttribute[4].Array  == ExpectedBudLightDetected);
-	UTEST_TRUE("Point5_BudLightDetected",  BudLightDetectedAttribute[5].Array  == ExpectedBudLightDetected);
-	UTEST_TRUE("Point6_BudLightDetected",  BudLightDetectedAttribute[6].Array  == ExpectedBudLightDetected);
-	UTEST_TRUE("Point7_BudLightDetected",  BudLightDetectedAttribute[7].Array  == ExpectedBudLightDetected);
-	UTEST_TRUE("Point8_BudLightDetected",  BudLightDetectedAttribute[8].Array  == ExpectedBudLightDetected);
-	UTEST_TRUE("Point9_BudLightDetected",  BudLightDetectedAttribute[9].Array  == ExpectedBudLightDetected);
-	UTEST_TRUE("Point10_BudLightDetected", BudLightDetectedAttribute[10].Array == ExpectedBudLightDetected);
-	UTEST_TRUE("Point11_BudLightDetected", BudLightDetectedAttribute[11].Array == ExpectedBudLightDetected);
-
-	return true;
-}
-
-PV_IMPORTER_TEST(ComputeBudDirections)
-{
-	const FVector3f ExpectedApical0 = (PVMockTree::Point1 - PVMockTree::Point0).GetSafeNormal();
-	const FVector3f ExpectedApical1 = (PVMockTree::Point2 - PVMockTree::Point1).GetSafeNormal();
-	const FVector3f ExpectedApical2 = (PVMockTree::Point3 - PVMockTree::Point2).GetSafeNormal();
-	const FVector3f ExpectedApical3 = (PVMockTree::Point4 - PVMockTree::Point3).GetSafeNormal();
-	const FVector3f ExpectedApical4 = (PVMockTree::Point5 - PVMockTree::Point4).GetSafeNormal();
-	const FVector3f ExpectedApical5 = (PVMockTree::Point5 - PVMockTree::Point4).GetSafeNormal();
-	const FVector3f ExpectedApical6 = (PVMockTree::Point6 - PVMockTree::Point3).GetSafeNormal();
-	const FVector3f ExpectedApical7 = (PVMockTree::Point8 - PVMockTree::Point7).GetSafeNormal();
-	const FVector3f ExpectedApical8 = (PVMockTree::Point9 - PVMockTree::Point8).GetSafeNormal();
-	const FVector3f ExpectedApical9 = (PVMockTree::Point9 - PVMockTree::Point8).GetSafeNormal();
-	const FVector3f ExpectedApical10 = (PVMockTree::Point11 - PVMockTree::Point10).GetSafeNormal();
-	const FVector3f ExpectedApical11 = (PVMockTree::Point11 - PVMockTree::Point10).GetSafeNormal();
-
-	const int32 ExpectedNumBudDirectionValues = 6;
-
-	const TSharedRef<FManagedArrayCollection> Collection = PVMockTreeCollection::CreateCollection();
-
-	const PV::FBudDirectionAttributeConstView BudDirectionAttribute = PV::FBudDirectionAttribute::FindAttribute(*Collection);
-	UTEST_TRUE("HasAttribute", BudDirectionAttribute.IsValid());
-	UTEST_TRUE("Has 12 elements", BudDirectionAttribute.Num() >= 12);
-
-	UTEST_NEARLY_EQUAL("Point0_Apical", FVector(BudDirectionAttribute[0].Apical), FVector(ExpectedApical0), UE_KINDA_SMALL_NUMBER);
-	UTEST_NEARLY_EQUAL("Point1_Apical", FVector(BudDirectionAttribute[1].Apical), FVector(ExpectedApical1), UE_KINDA_SMALL_NUMBER);
-	UTEST_NEARLY_EQUAL("Point2_Apical", FVector(BudDirectionAttribute[2].Apical), FVector(ExpectedApical2), UE_KINDA_SMALL_NUMBER);
-	UTEST_NEARLY_EQUAL("Point3_Apical", FVector(BudDirectionAttribute[3].Apical), FVector(ExpectedApical3), UE_KINDA_SMALL_NUMBER);
-	UTEST_NEARLY_EQUAL("Point4_Apical", FVector(BudDirectionAttribute[4].Apical), FVector(ExpectedApical4), UE_KINDA_SMALL_NUMBER);
-	UTEST_NEARLY_EQUAL("Point5_Apical", FVector(BudDirectionAttribute[5].Apical), FVector(ExpectedApical5), UE_KINDA_SMALL_NUMBER);
-	UTEST_NEARLY_EQUAL("Point6_Apical", FVector(BudDirectionAttribute[6].Apical), FVector(ExpectedApical6), UE_KINDA_SMALL_NUMBER);
-	UTEST_NEARLY_EQUAL("Point7_Apical", FVector(BudDirectionAttribute[7].Apical), FVector(ExpectedApical7), UE_KINDA_SMALL_NUMBER);
-	UTEST_NEARLY_EQUAL("Point8_Apical", FVector(BudDirectionAttribute[8].Apical), FVector(ExpectedApical8), UE_KINDA_SMALL_NUMBER);
-	UTEST_NEARLY_EQUAL("Point9_Apical", FVector(BudDirectionAttribute[9].Apical), FVector(ExpectedApical9), UE_KINDA_SMALL_NUMBER);
-	UTEST_NEARLY_EQUAL("Point10_Apical", FVector(BudDirectionAttribute[10].Apical), FVector(ExpectedApical10), UE_KINDA_SMALL_NUMBER);
-	UTEST_NEARLY_EQUAL("Point11_Apical", FVector(BudDirectionAttribute[11].Apical), FVector(ExpectedApical11), UE_KINDA_SMALL_NUMBER);
-
-	// TODO: Test UpVector and Axillary
-
-	return true;
-}
-
-PV_IMPORTER_TEST(SplitFusedPoints)
-{
-	const TSharedRef<FManagedArrayCollection> Collection = PVMockTreeCollection::CreateCollection(false);
-
-	const bool bSplitFusedPointsResult = PV::ImportHelper::SplitFusedPoints(*Collection);
-	UTEST_TRUE("SplitFusedPointsSucceeded", bSplitFusedPointsResult);
-
-	PV::FPointPositionAttributeConstView PointPositionAttribute = PV::FPointPositionAttribute::FindAttribute(*Collection);
-	UTEST_TRUE("HasAttribute", PointPositionAttribute.IsValid());
-	UTEST_TRUE("Has 15 elements", PointPositionAttribute.Num() == 15);
-
-	PV::FPointLengthFromRootAttributeConstView PointLengthFromRootAttribute = PV::FPointLengthFromRootAttribute::FindAttribute(*Collection);
-	UTEST_TRUE("HasAttribute", PointLengthFromRootAttribute.IsValid());
-	UTEST_TRUE("Has 15 elements", PointLengthFromRootAttribute.Num() == 15);
-
-	const PV::FPointLengthFromSeedAttributeView PointLengthFromSeedAttribute = PV::FPointLengthFromSeedAttribute::FindAttribute(*Collection);
-	UTEST_TRUE("HasAttribute", PointLengthFromSeedAttribute.IsValid());
-	UTEST_TRUE("Has 15 elements", PointLengthFromSeedAttribute.Num() == 15);
-
-	const PV::FBranchPointsAttributeView BranchPointsAttribute = PV::FBranchPointsAttribute::FindAttribute(*Collection);
-	UTEST_TRUE("HasAttribute", BranchPointsAttribute.IsValid());
-	UTEST_TRUE("Has 4 elements", BranchPointsAttribute.Num() == 4);
-
-	const int32 Branch0_OldRootPointIndex = 0;
-	const int32 Branch1_OldRootPointIndex = 2;
-	const int32 Branch2_OldRootPointIndex = 8;
-	const int32 Branch3_OldRootPointIndex = 3;
-
-	const int32 Branch0_NewRootPointIndex = BranchPointsAttribute[0][0];
-	const int32 Branch1_NewRootPointIndex = BranchPointsAttribute[1][0];
-	const int32 Branch2_NewRootPointIndex = BranchPointsAttribute[2][0];
-	const int32 Branch3_NewRootPointIndex = BranchPointsAttribute[3][0];
-
-	UTEST_TRUE("Branch0_HasSamePointIndex",  Branch0_NewRootPointIndex == Branch0_OldRootPointIndex);
-	UTEST_TRUE("Branch1_HasUniqePointIndex", Branch1_NewRootPointIndex != Branch1_OldRootPointIndex && Branch1_NewRootPointIndex > 11);
-	UTEST_TRUE("Branch2_HasUniqePointIndex", Branch2_NewRootPointIndex != Branch2_OldRootPointIndex && Branch2_NewRootPointIndex > 11);
-	UTEST_TRUE("Branch3_HasUniqePointIndex", Branch3_NewRootPointIndex != Branch3_OldRootPointIndex && Branch3_NewRootPointIndex > 11);
-
-	UTEST_EQUAL("Branch0_NewPointPositionIsSame", PointPositionAttribute[Branch0_NewRootPointIndex], PointPositionAttribute[Branch0_OldRootPointIndex]);
-	UTEST_EQUAL("Branch1_NewPointPositionIsSame", PointPositionAttribute[Branch1_NewRootPointIndex], PointPositionAttribute[Branch1_OldRootPointIndex]);
-	UTEST_EQUAL("Branch2_NewPointPositionIsSame", PointPositionAttribute[Branch2_NewRootPointIndex], PointPositionAttribute[Branch2_OldRootPointIndex]);
-	UTEST_EQUAL("Branch3_NewPointPositionIsSame", PointPositionAttribute[Branch3_NewRootPointIndex], PointPositionAttribute[Branch3_OldRootPointIndex]);
-
-	UTEST_EQUAL("Branch0_NewPointLengthFromRootIsSame", PointLengthFromRootAttribute[Branch0_NewRootPointIndex], PointLengthFromRootAttribute[Branch0_OldRootPointIndex]);
-	UTEST_EQUAL("Branch1_NewPointLengthFromRootIsSame", PointLengthFromRootAttribute[Branch1_NewRootPointIndex], PointLengthFromRootAttribute[Branch1_OldRootPointIndex]);
-	UTEST_EQUAL("Branch2_NewPointLengthFromRootIsSame", PointLengthFromRootAttribute[Branch2_NewRootPointIndex], PointLengthFromRootAttribute[Branch2_OldRootPointIndex]);
-	UTEST_EQUAL("Branch3_NewPointLengthFromRootIsSame", PointLengthFromRootAttribute[Branch3_NewRootPointIndex], PointLengthFromRootAttribute[Branch3_OldRootPointIndex]);
-
-	UTEST_EQUAL("Branch0_NewPointLengthFromSeedIsSame", PointLengthFromSeedAttribute[Branch0_NewRootPointIndex], PointLengthFromSeedAttribute[Branch0_OldRootPointIndex]);
-	UTEST_EQUAL("Branch1_NewPointLengthFromSeedIsSame", PointLengthFromSeedAttribute[Branch1_NewRootPointIndex], PointLengthFromSeedAttribute[Branch1_OldRootPointIndex]);
-	UTEST_EQUAL("Branch2_NewPointLengthFromSeedIsSame", PointLengthFromSeedAttribute[Branch2_NewRootPointIndex], PointLengthFromSeedAttribute[Branch2_OldRootPointIndex]);
-	UTEST_EQUAL("Branch3_NewPointLengthFromSeedIsSame", PointLengthFromSeedAttribute[Branch3_NewRootPointIndex], PointLengthFromSeedAttribute[Branch3_OldRootPointIndex]);
-
-	return true;
-}
-
 PV_IMPORTER_TEST(FindRootBranch)
 {
 	const TArray<TArray<int32>>& Branches = PVMockTree::BranchPointIndices;
