[PVE] Added Graph parameters and embedded graphs tab in the default editor layout.

#Synced-CL 50704972

diff --git a/Engine/Plugins/Experimental/ProceduralVegetationEditor/Source/ProceduralVegetationEditor/Private/PVEditor.cpp b/Engine/Plugins/Experimental/ProceduralVegetationEditor/Source/ProceduralVegetationEditor/Private/PVEditor.cpp
--- a/Engine/Plugins/Experimental/ProceduralVegetationEditor/Source/ProceduralVegetationEditor/Private/PVEditor.cpp
+++ b/Engine/Plugins/Experimental/ProceduralVegetationEditor/Source/ProceduralVegetationEditor/Private/PVEditor.cpp
@@ -206,7 +206,7 @@ TAttribute<FGraphAppearanceInfo> FPVEditor::GetAppearanceInfo() const
 
 TSharedRef<FTabManager::FLayout> FPVEditor::GetDefaultLayout() const
 {
-	return FTabManager::NewLayout("Standalone_PVGraphEditor_DefaultLayout_v1.0.5")
+	return FTabManager::NewLayout("Standalone_PVGraphEditor_DefaultLayout_v1.0.6")
 	->AddArea // Main PCG Graph Editor Area
 	(
 		FTabManager::NewPrimaryArea()->SetOrientation(Orient_Vertical)
@@ -218,14 +218,21 @@ TSharedRef<FTabManager::FLayout> FPVEditor::GetDefaultLayout() const
 			(
 				FTabManager::NewSplitter()->SetOrientation(Orient_Vertical)
 				->SetSizeCoefficient(0.1f)
-				->Split // Viewport and AVL
+				->Split // Viewport
 				(
 					FTabManager::NewStack()
-					->SetSizeCoefficient(0.95f)
+					->SetSizeCoefficient(0.75f)
 					->AddTab(PCGEditorTabs::ViewportID[0], ETabState::OpenedTab)
 					->SetForegroundTab(FTabId(PCGEditorTabs::ViewportID[0]))
 					->SetHideTabWell(true)
 				)
+				->Split // Embedded graph
+				(
+					FTabManager::NewStack()
+					->SetSizeCoefficient(0.25f)
+					->AddTab(PCGEditorTabs::EmbeddedSubgraphsID, ETabState::OpenedTab)
+					->SetForegroundTab(FTabId(PCGEditorTabs::EmbeddedSubgraphsID))
+				)
 			)
 			->Split // Graph and details panel
 			(
@@ -237,7 +244,6 @@ TSharedRef<FTabManager::FLayout> FPVEditor::GetDefaultLayout() const
 					FTabManager::NewStack()
 					->SetSizeCoefficient(0.6f)
 					->AddTab(PCGEditorTabs::GraphEditorID, ETabState::OpenedTab)
-					->AddTab(PCGEditorTabs::UserParamsID, ETabState::ClosedTab)
 					->SetHideTabWell(true)
 				)
 				->Split
@@ -245,6 +251,8 @@ TSharedRef<FTabManager::FLayout> FPVEditor::GetDefaultLayout() const
 					FTabManager::NewStack()
 					->SetSizeCoefficient(0.4f)
 					->AddTab(PCGEditorTabs::PropertyDetailsID[0], ETabState::OpenedTab)
+					->AddTab(PCGEditorTabs::UserParamsID, ETabState::OpenedTab)
+					->SetForegroundTab(FTabId(PCGEditorTabs::PropertyDetailsID[0]))
 				)
 			)
 		)
