[PVE] Seed Generator

Added two nodes Point Scatter and Convert to Seed points
- Point Scatter node is a PCGSubGraphDerived Node that sets a default Scatter SubGraph, The default sub graph exposes multiple paramters on node and in detail panel for manipulation of scatter points.
- Convert to Seed points node takes the PCG Point Data input and outputs a Manged Array Collection, that can be chained togther with Grower node.
- Added Content includes Cone mesh and materials for visualization and PCG Point Scatter Graph.

#virtualized
#Synced-CL 50393760

diff --git a/Engine/Plugins/Experimental/ProceduralVegetationEditor/Source/ProceduralVegetation/Private/Facades/PVPointFacade.cpp b/Engine/Plugins/Experimental/ProceduralVegetationEditor/Source/ProceduralVegetation/Private/Facades/PVPointFacade.cpp
--- a/Engine/Plugins/Experimental/ProceduralVegetationEditor/Source/ProceduralVegetation/Private/Facades/PVPointFacade.cpp
+++ b/Engine/Plugins/Experimental/ProceduralVegetationEditor/Source/ProceduralVegetation/Private/Facades/PVPointFacade.cpp
@@ -21,6 +21,8 @@ namespace PV::Facades
 		, BudHormoneLevels(InCollection, AttributeNames::BudHormoneLevels, GroupNames::PointGroup)
 		, PlantGradients(InCollection, AttributeNames::PlantGradient, GroupNames::PointGroup)
 		, LengthFromSeed(InCollection, AttributeNames::LengthFromSeed, GroupNames::PointGroup)
+		, SeedPScale(InCollection, AttributeNames::SeedPScale, GroupNames::PointGroup)
+		, SeedPScaleRatio(InCollection, AttributeNames::SeedPScaleRatio, GroupNames::PointGroup)
 		, BudNumber(InCollection, AttributeNames::BudNumber, GroupNames::PointGroup)
 		, NjordPixelIndex(InCollection, AttributeNames::NjordPixelIndex, GroupNames::PointGroup)
 		, BudLightDetected(InCollection, AttributeNames::BudLightDetected, GroupNames::PointGroup)
@@ -46,6 +48,8 @@ namespace PV::Facades
 		, BudHormoneLevels(InCollection, AttributeNames::BudHormoneLevels, GroupNames::PointGroup)
 		, PlantGradients(InCollection, AttributeNames::PlantGradient, GroupNames::PointGroup)
 		, LengthFromSeed(InCollection, AttributeNames::LengthFromSeed, GroupNames::PointGroup)
+		, SeedPScale(InCollection, AttributeNames::SeedPScale, GroupNames::PointGroup)
+		, SeedPScaleRatio(InCollection, AttributeNames::SeedPScaleRatio, GroupNames::PointGroup)
 		, BudNumber(InCollection, AttributeNames::BudNumber, GroupNames::PointGroup)
 		, NjordPixelIndex(InCollection, AttributeNames::NjordPixelIndex, GroupNames::PointGroup)
 		, BudLightDetected(InCollection, AttributeNames::BudLightDetected, GroupNames::PointGroup)
@@ -514,6 +518,16 @@ namespace PV::Facades
 			{
 				BudLateralMeristem.ModifyAt(ToIndex, BudLateralMeristem[FromIndex]);
 			}
+
+			if (SeedPScale.IsValid() && SeedPScale.IsValidIndex(ToIndex) && SeedPScale.IsValidIndex(FromIndex))
+			{
+				SeedPScale.ModifyAt(ToIndex, SeedPScale[FromIndex]);
+			}
+
+			if (SeedPScaleRatio.IsValid() && SeedPScaleRatio.IsValidIndex(ToIndex) && SeedPScaleRatio.IsValidIndex(FromIndex))
+			{
+				SeedPScaleRatio.ModifyAt(ToIndex, SeedPScaleRatio[FromIndex]);
+			}
 		}
 	}
 
@@ -625,6 +639,16 @@ namespace PV::Facades
 		FILL_COLLECTION_ATTRIBUTE(BudLateralMeristem, InputArray);
 	}
 
+	void FPointFacade::FillSeedPScale(const TArray<float>& InputArray)
+	{
+		FILL_COLLECTION_ATTRIBUTE(SeedPScale, InputArray);
+	}
+
+	void FPointFacade::FillSeedPScaleRatio(const TArray<float>& InputArray)
+	{
+		FILL_COLLECTION_ATTRIBUTE(SeedPScaleRatio, InputArray);
+	}
+
 	void FPointFacade::GetLFR(int32 Index, float& OutValue) const
 	{
 		GetValueFromAccessor<float>(LengthFromRoot, OutValue, Index, -1);
@@ -704,6 +728,17 @@ namespace PV::Facades
 	{
 		GetValueFromAccessor<TArray<float>>(BudLateralMeristem, OutValue, Index, TArray<float>());
 	}
+
+	void FPointFacade::GetSeedPScale(int32 Index, float& OutValue) const
+	{
+		GetValueFromAccessor<float>(SeedPScale, OutValue, Index, 0.1);
+	}
+
+	void FPointFacade::GetSeedPScaleRatio(int32 Index, float& OutValue) const
+	{
+		GetValueFromAccessor<float>(SeedPScaleRatio, OutValue, Index, 1);
+	}
+	
 	const TManagedArray<TArray<float>>& FPointFacade::GetBudLightDetectedArrays() const
 	{
 		return BudLightDetected.Get();
diff --git a/Engine/Plugins/Experimental/ProceduralVegetationEditor/Source/ProceduralVegetation/Private/Implementations/PVGrower.cpp b/Engine/Plugins/Experimental/ProceduralVegetationEditor/Source/ProceduralVegetation/Private/Implementations/PVGrower.cpp
--- a/Engine/Plugins/Experimental/ProceduralVegetationEditor/Source/ProceduralVegetation/Private/Implementations/PVGrower.cpp
+++ b/Engine/Plugins/Experimental/ProceduralVegetationEditor/Source/ProceduralVegetation/Private/Implementations/PVGrower.cpp
@@ -4,6 +4,7 @@
 #include <random>
 #include "ProceduralVegetationModule.h"
 #include "Algo/Accumulate.h"
+#include "Facades/PVAttributesNames.h"
 #include "Facades/PVBranchFacade.h"
 #include "Facades/PVPointFacade.h"
 #include "GeometryCollection/ManagedArrayCollection.h"
@@ -112,6 +113,8 @@ void FPVGrower::FillCollection(FManagedArrayCollection& Collection, UPVGrowerDat
 	TArray<float> MainTrunkGradients;
 	TArray<float> GroundGradients;
 	TArray<float> PointScales;
+	TArray<float> SeedPScales;
+	TArray<float> SeedPScaleRatios;
 	TArray<TArray<FVector3f>> BudDirections;
 	TArray<float> BranchGradients;
 	TArray<TArray<float>> BudHormoneLevels;
@@ -130,6 +133,8 @@ void FPVGrower::FillCollection(FManagedArrayCollection& Collection, UPVGrowerDat
 		Points.Add(FVector3f(Point->Position));//FVector3f(-Point->Position.X , Point->Position.Z, Point->Position.Y));
 		LengthFromRoots.Add(Point->LengthFromRoot);
 		PointScales.Add(Point->PScale);
+		SeedPScales.Add(Point->SeedInfo.PScale);
+		SeedPScaleRatios.Add(Point->SeedInfo.PScaleRatio);
 		LengthFromSeeds.Add(Point->LengthFromRoot);
 		BudNumbers.Add(Point->Bud.BudNumber);
 		
@@ -162,6 +167,8 @@ void FPVGrower::FillCollection(FManagedArrayCollection& Collection, UPVGrowerDat
 	PointFacade.FillBudLightDetected(BudLightDetected);
 	PointFacade.FillBudDevelopment(BudDevelopment);
 	PointFacade.FillNjordPixelID(NjordPixelIndices);
+	PointFacade.FillSeedPScale(SeedPScales);
+	PointFacade.FillSeedPScaleRatio(SeedPScaleRatios);
 
 	TArray<TArray<int32>> AllParents;
 	TArray<TArray<int32>> AllChildren;
@@ -285,8 +292,8 @@ void FPVGrower::FillPointsFromCollection(const FManagedArrayCollection& Collecti
 			
 			Point->Bud = Bud;
 
-			Point->SeedInfo.PScale = 0.1;
-			Point->SeedInfo.PScaleRatio = 1;
+			PointFacade.GetSeedPScale(i, Point->SeedInfo.PScale);
+			PointFacade.GetSeedPScaleRatio(i, Point->SeedInfo.PScaleRatio);
 
 			Skeleton->AddPoint(Point, Bud.BudNumber);
 		}
@@ -2789,6 +2796,7 @@ void FPVGrower::SetDegradationAmount(const FPVGrowerParams& InGrowerParams, UPVG
 			float Retention = Bud.LateralMeristem.LateralMeristem / Denominator;
 			float AdditionalDegradation = RandomValue(BudNumber + Cycle) - Retention;
 			AdditionalDegradation = FMath::Clamp(AdditionalDegradation,0,1);
+			//TODO: Investigate degradation Should be just and Add.
 			AdditionalDegradation *= Primitive->Senescence.DegradationAmount;
 			Bud.LateralMeristem.Inactive = 1.0;
 			Bud.LateralMeristem.Degredation += AdditionalDegradation;
diff --git a/Engine/Plugins/Experimental/ProceduralVegetationEditor/Source/ProceduralVegetation/Private/Implementations/PVSeedGenerator.cpp b/Engine/Plugins/Experimental/ProceduralVegetationEditor/Source/ProceduralVegetation/Private/Implementations/PVSeedGenerator.cpp
--- a/Engine/Plugins/Experimental/ProceduralVegetationEditor/Source/ProceduralVegetation/Private/Implementations/PVSeedGenerator.cpp
+++ b/Engine/Plugins/Experimental/ProceduralVegetationEditor/Source/ProceduralVegetation/Private/Implementations/PVSeedGenerator.cpp
@@ -1,7 +1,236 @@
 // Copyright Epic Games, Inc. All Rights Reserved.
 
 #include "PVSeedGenerator.h"
+#include "PVGrowerData.h"
+#include "Facades/PVBranchFacade.h"
+#include "Facades/PVPointFacade.h"
+#include "Utils/PCGValueRange.h"
 
 #define LOCTEXT_NAMESPACE "PVSeedGenerator"
 
+TArray<FTransform> FPVConvertToSeedPointImplementation::FindNeighbours(const FTransform& InTransform, const UPCGPointData* PointData, float MaxDist, int MaxPoints)
+{
+	TArray<FTransform> Neighbours;
+
+	TConstPCGValueRange<FTransform> PointTransforms = PointData->GetConstTransformValueRange();
+	for (const FTransform& Point : PointTransforms)
+	{
+		if (Point.GetLocation() != InTransform.GetLocation() && Neighbours.Num() < MaxPoints
+			&& FVector::Dist(Point.GetLocation(), InTransform.GetLocation()) < MaxDist)
+		{
+			Neighbours.Add(Point);
+		}
+	}
+	
+	return Neighbours;
+}
+
+void FPVConvertToSeedPointImplementation::SetBaseDirection(const UPCGPointData* PointData, const FPVConvertToSeedPointParams& Params, TArray<TArray<FVector3f>>& DirectionArray)
+{
+	int PrimaryAxis = 0;
+	TConstPCGValueRange<FTransform> PointTransforms = PointData->GetConstTransformValueRange();
+
+	float MinDistance = FLT_MAX;
+	float MaxDistance = -FLT_MAX;
+	
+	for (const FTransform& Point : PointTransforms)
+	{
+		float Distance = FVector::Dist(Point.GetLocation(), FVector::Zero());
+		
+		MinDistance = FMath::Min(Distance, MinDistance);
+		MaxDistance = FMath::Max(Distance, MaxDistance);
+	}
+
+	int SeedPointCount = 0;
+	for (const FTransform& Point : PointTransforms)
+	{
+		TArray<FTransform> Neighbours = FindNeighbours(Point, PointData, 200, 5);
+		FVector AvoidanceDirection = FVector::ZeroVector;
+		
+		for (const FTransform& Neighbour : Neighbours)
+		{
+			FVector NearAvoidanceDirection = (Point.GetLocation() - Neighbour.GetLocation()).GetSafeNormal();
+			AvoidanceDirection += NearAvoidanceDirection;
+		}
+
+		if (!Neighbours.IsEmpty())
+		{
+			AvoidanceDirection = AvoidanceDirection / Neighbours.Num();
+		}
+
+		FVector ApicalDirection;
+		FVector AxillaryDirection;
+		FVector BorderAdjustedApicalDirection;
+		FVector UpVector;
+		FVector CentroidDirection = Point.GetLocation().GetSafeNormal();
+		//float RandomSeed = SeedPointCount + 1;
+		
+		float CenterReduction = FMath::GetMappedRangeValueClamped(FVector2f(0, MaxDistance), FVector2f(0, 1), FVector::Dist(Point.GetLocation(), FVector::Zero()));
+		
+		if (PrimaryAxis == 0)
+		{
+			AxillaryDirection = FMath::Lerp(CentroidDirection,AvoidanceDirection,Params.AvoidanceBias).GetSafeNormal();
+			ApicalDirection = FVector::UpVector;
+			UpVector = FVector::CrossProduct(AxillaryDirection,ApicalDirection);
+
+			FVector AxillaryAdjustedApicalDirection = FMath::Lerp(ApicalDirection,AxillaryDirection, Params.Avoidance);
+			BorderAdjustedApicalDirection =  FMath::Lerp(ApicalDirection,AxillaryAdjustedApicalDirection,CenterReduction);
+			BorderAdjustedApicalDirection =  FMath::Lerp(AxillaryAdjustedApicalDirection,BorderAdjustedApicalDirection,Params.AvoidanceCenterReduction);
+		}
+		else
+		{
+			ApicalDirection = FVector::UpVector;
+			ApicalDirection =  FMath::Lerp(ApicalDirection,AvoidanceDirection,Params.Avoidance).GetSafeNormal();
+			UpVector = FVector::RightVector;
+			AxillaryDirection =   FVector::CrossProduct(UpVector,ApicalDirection);
+
+			BorderAdjustedApicalDirection = ApicalDirection;
+		}
+
+		//float ApicalRandomAngle = FMath::DegreesToRadians(Params.ApicalRandom);
+		//FVector RandomAngleApicalDirection = sample_direction_cone(BorderAdjustedApicalDirection, ApicalRandomAngle,FVector2f(FMath::FRand(RandomSeed)));
+
+		FQuat Dihedral = FQuat::FindBetweenVectors(ApicalDirection, BorderAdjustedApicalDirection);
+		AxillaryDirection = Dihedral.RotateVector(AxillaryDirection);
+		UpVector = Dihedral.RotateVector(UpVector);
+
+		ApicalDirection = BorderAdjustedApicalDirection;
+
+		DirectionArray.Add(FPVBudDirection{ApicalDirection,AxillaryDirection,ApicalDirection, ApicalDirection, ApicalDirection, UpVector}.GetDataArray());
+		SeedPointCount++;
+	}
+
+}
+
+void FPVConvertToSeedPointImplementation::FillCollection(FManagedArrayCollection& Collection, const FPVConvertToSeedPointParams& Settings, const UPCGPointData* SeedPoints)
+{
+	PV::Facades::FPointFacade PointFacade = PV::Facades::FPointFacade(Collection);
+	PV::Facades::FBranchFacade BranchFacade = PV::Facades::FBranchFacade(Collection);
+
+	TArray<FVector3f> Points;
+	TArray<float> LengthFromRoots;
+	TArray<float> PointScaleGradients;
+	TArray<float> HullGradients;
+	TArray<float> MainTrunkGradients;
+	TArray<float> GroundGradients;
+	TArray<float> PointScales;
+	TArray<float> SeedPScales;
+	TArray<float> SeedPScaleRatios;
+	TArray<TArray<FVector3f>> BudDirections;
+	TArray<float> BranchGradients;
+	TArray<TArray<float>> BudHormoneLevels;
+	TArray<float> PlantGradients;
+	TArray<float> LengthFromSeeds;
+	TArray<int32> BudNumbers;
+	TArray<TArray<float>> BudLightDetected;
+	TArray<TArray<int32>> BudDevelopment;
+	TArray<TArray<int32>> BudStatus;
+	TArray<TArray<float>> BudLateralMeristem;
+	TArray<float> NjordPixelIndices;
+
+	TArray<TArray<int32>> AllParents;
+	TArray<TArray<int32>> AllChildren;
+	TArray<TArray<int32>> AllPoints;
+	TArray<int32> BranchNumbers;
+	TArray<int32> BranchSourceBudNumbers;
+	TArray<int32> BranchHierarchyNumbers;
+	TArray<int32> BranchGenerationNumbers;
+	TArray<int32> BranchParentNumbers;
+	TArray<int32> PlantNumbers;
+	
+	int PointIndex = 0;
+	TConstPCGValueRange<FTransform> PointTransforms = SeedPoints->GetConstTransformValueRange();
+
+	float MinScale = FLT_MAX;
+	float MaxScale = -FLT_MAX;
+	for (const FTransform& Point : PointTransforms)
+	{
+		MinScale = FMath::Min(Point.GetScale3D().X, MinScale);
+		MaxScale = FMath::Max(Point.GetScale3D().X, MaxScale);
+	}
+
+	int BudNumber = 0;
+	int PlantNumber = 0;
+	int BranchNumber = 0;
+	for (const FTransform& Point : PointTransforms)
+	{ 
+		Points.Add(FVector3f(Point.GetLocation()));//FVector3f(-Point->Position.X , Point->Position.Z, Point->Position.Y));
+		LengthFromRoots.Add(0);
+		PointScales.Add(Point.GetScale3D().X);
+		SeedPScales.Add(Point.GetScale3D().X);
+		float PScaleRatio = FMath::GetMappedRangeValueClamped(FVector2f(0, MaxScale), FVector2f(0, 1), Point.GetScale3D().X);
+		SeedPScaleRatios.Add(PScaleRatio);
+		LengthFromSeeds.Add(0);
+		BudNumbers.Add(BudNumber);
+
+		FPVBudHarmoneLevels HarmoneLevels;
+		FPVBudStatus Status;
+		FPVBudLateralMeristem LateralMeristem;
+		FPVBudLightDetected LightDetected;
+		FPVBudDevelopment Development;
+		
+		LateralMeristem = {0,1,0,0,0,0,0};
+		Status = FPVBudStatus{1,0,0,1,1,0,0,0,0,0};
+		HarmoneLevels = FPVBudHarmoneLevels{1,1, 0,0,0 };
+		LightDetected = {0,0,0,0};
+		Development = { 1,0,0,0,0,1};
+		
+		BudHormoneLevels.Add(HarmoneLevels.GetDataArray());
+		BudDevelopment.Add(Development.GetDataArray());
+		BudLightDetected.Add(LightDetected.GetDataArray());
+		BudStatus.Add(Status.GetDataArray());
+		BudLateralMeristem.Add(LateralMeristem.GetDataArray());
+		
+		AllParents.Add({0});
+		AllPoints.Add({BudNumber});
+		BranchNumbers.Add(BudNumber);
+		PlantNumbers.Add(PlantNumber);
+		BranchSourceBudNumbers.Add(BudNumber);
+		BranchParentNumbers.Add(0);
+		BranchGenerationNumbers.Add(1);
+		BranchHierarchyNumbers.Add(1);
+		
+		BudNumber++;
+		PlantNumber++;
+		BranchNumber++;
+	}
+
+	SetBaseDirection(SeedPoints, Settings, BudDirections);
+
+	PointFacade.AddElements(Points.Num());
+	PointFacade.FillPositions(Points);
+	PointFacade.FillLFR(LengthFromRoots);
+	PointFacade.FillPointScaleGradients(PointScaleGradients);
+	PointFacade.FillHullGradients(HullGradients);
+	PointFacade.FillMainTrunkGradients(MainTrunkGradients);
+	PointFacade.FillGroundGradients(GroundGradients);
+	PointFacade.FillPointScales(PointScales);
+	PointFacade.FillBudDirections(BudDirections);
+	PointFacade.FillBudStatus(BudStatus);
+	PointFacade.FillBudLateralMeristem(BudLateralMeristem);
+	PointFacade.FillBranchGradients(BranchGradients);
+	PointFacade.FillBudHormoneLevels(BudHormoneLevels);
+	PointFacade.FillPlantGradients(PlantGradients);
+	PointFacade.FillLengthFromSeeds(LengthFromSeeds);
+	PointFacade.FillBudNumbers(BudNumbers);
+	PointFacade.FillBudLightDetected(BudLightDetected);
+	PointFacade.FillBudDevelopment(BudDevelopment);
+	PointFacade.FillNjordPixelID(NjordPixelIndices);
+	PointFacade.FillSeedPScale(SeedPScales);
+	PointFacade.FillSeedPScaleRatio(SeedPScaleRatios);
+
+	BranchFacade.AddElements(AllParents.Num());
+	BranchFacade.FillParents(AllParents);
+	BranchFacade.FillChildren(AllChildren);
+	BranchFacade.FillBranchPoints(AllPoints);
+	BranchFacade.FillBranchNumbers(BranchNumbers);
+	BranchFacade.FillBranchSourceBudNumber(BranchSourceBudNumbers);
+	BranchFacade.FillBranchHierarchyNumber(BranchHierarchyNumbers);
+	BranchFacade.FillBranchParentNumbers(BranchParentNumbers);
+	BranchFacade.FillPlantNumbers(PlantNumbers);
+
+	check(PointFacade.IsValid());
+	check(BranchFacade.IsValid());
+}
+
 #undef LOCTEXT_NAMESPACE
diff --git a/Engine/Plugins/Experimental/ProceduralVegetationEditor/Source/ProceduralVegetation/Private/Implementations/PVSeedGenerator.h b/Engine/Plugins/Experimental/ProceduralVegetationEditor/Source/ProceduralVegetation/Private/Implementations/PVSeedGenerator.h
--- a/Engine/Plugins/Experimental/ProceduralVegetationEditor/Source/ProceduralVegetation/Private/Implementations/PVSeedGenerator.h
+++ b/Engine/Plugins/Experimental/ProceduralVegetationEditor/Source/ProceduralVegetation/Private/Implementations/PVSeedGenerator.h
@@ -3,6 +3,8 @@
 #pragma once
 
 #include "CoreMinimal.h"
+#include "Data/PCGPointData.h"
+#include "GeometryCollection/ManagedArrayCollection.h"
 #include "PVSeedGenerator.generated.h"
 
 USTRUCT()
@@ -72,3 +74,33 @@ private:
 	UPROPERTY(EditAnywhere, Category = "Seed")
 	FPVSeedPoint SeedPoint;
 };
+
+USTRUCT()
+struct FPVConvertToSeedPointParams
+{
+	GENERATED_BODY()
+	
+	UPROPERTY(EditAnywhere, Category = AvoidanceSettings , meta = (UIMin = "0.0", UIMax = "1.0"))
+	float Avoidance = 0;
+	
+	UPROPERTY(EditAnywhere, Category = AvoidanceSettings, meta = (UIMin = "0.0", UIMax = "1.0"))
+	float AvoidanceBias = 0;
+
+	UPROPERTY(EditAnywhere, Category = AvoidanceSettings, meta = (UIMin = "0.0", UIMax = "1.0"))
+	float AvoidanceCenterReduction = 0;
+
+	UPROPERTY(VisibleAnywhere, Category = AvoidanceSettings, meta = (UIMin = "0.0", UIMax = "360.0"))
+	float ApicalRandom = 0;
+
+	UPROPERTY(VisibleAnywhere, Category = AvoidanceSettings, meta = (UIMin = "0.0", UIMax = "360.0"))
+	float AxillaryRandom = 0;
+};
+
+struct FPVConvertToSeedPointImplementation
+{
+	static TArray<FTransform> FindNeighbours(const FTransform& InTransform, const UPCGPointData* PointData, float MaxDist, int MaxPoints);
+	
+	static void SetBaseDirection(const UPCGPointData* PointData, const FPVConvertToSeedPointParams& Params, TArray<TArray<FVector3f>>& DirectionArray);
+	
+	static void FillCollection(FManagedArrayCollection& Collection, const FPVConvertToSeedPointParams& Params, const UPCGPointData* SeedPoints);
+};
diff --git a/Engine/Plugins/Experimental/ProceduralVegetationEditor/Source/ProceduralVegetation/Private/Nodes/PVConvertToSeedPointsSettings.cpp b/Engine/Plugins/Experimental/ProceduralVegetationEditor/Source/ProceduralVegetation/Private/Nodes/PVConvertToSeedPointsSettings.cpp
new file mode 100644
--- /dev/null
+++ b/Engine/Plugins/Experimental/ProceduralVegetationEditor/Source/ProceduralVegetation/Private/Nodes/PVConvertToSeedPointsSettings.cpp
@@ -0,0 +1,57 @@
+// Copyright Epic Games, Inc. All Rights Reserved.
+
+#include "Nodes/PVConvertToSeedPointsSettings.h"
+
+#include "PCGContext.h"
+#include "DataTypes/PVGrowthData.h"
+#include "Data/PCGBasePointData.h"
+#include "Data/PCGPointData.h"
+#include "Facades/PVPointFacade.h"
+
+#define LOCTEXT_NAMESPACE "PVConvertToSeedPointsSettings"
+
+FPCGDataTypeIdentifier UPVConvertToSeedPointsSettings::GetInputPinTypeIdentifier() const
+{
+	return FPCGDataTypeIdentifier{ FPCGDataTypeInfoPoint::AsId() };
+}
+
+FPCGDataTypeIdentifier UPVConvertToSeedPointsSettings::GetOutputPinTypeIdentifier() const
+{
+	return FPCGDataTypeIdentifier{ FPVDataTypeInfoGrowth::AsId() };
+}
+
+FPCGElementPtr UPVConvertToSeedPointsSettings::CreateElement() const
+{
+	return MakeShared<FPVConvertToSeedPointsElement>();
+}
+
+bool FPVConvertToSeedPointsElement::ExecuteInternal(FPCGContext* InContext) const
+{
+	check(InContext);
+
+	const UPVConvertToSeedPointsSettings* Settings = InContext->GetInputSettings<UPVConvertToSeedPointsSettings>();
+	check(Settings);
+	
+	for (const FPCGTaggedData& Input : InContext->InputData.GetInputsByPin(PCGPinConstants::DefaultInputLabel))
+	{
+		if(const UPCGPointData* InputData = Cast<UPCGPointData>(Input.Data))
+		{
+			FManagedArrayCollection OutCollection;
+			UPVGrowthData* OutManagedArrayCollectionData = FPCGContext::NewObject_AnyThread<UPVGrowthData>(InContext);
+			
+			FPVConvertToSeedPointImplementation::FillCollection(OutCollection, Settings->Params, InputData);
+			
+			OutManagedArrayCollectionData->Initialize(MoveTemp(OutCollection));
+			InContext->OutputData.TaggedData.Emplace(OutManagedArrayCollectionData);
+		}
+		else
+		{
+			PCGLog::InputOutput::LogInvalidInputDataError(InContext);
+			return true;
+		}
+	}
+	
+	return true;
+}
+
+#undef LOCTEXT_NAMESPACE
diff --git a/Engine/Plugins/Experimental/ProceduralVegetationEditor/Source/ProceduralVegetation/Private/Nodes/PVConvertToSeedPointsSettings.h b/Engine/Plugins/Experimental/ProceduralVegetationEditor/Source/ProceduralVegetation/Private/Nodes/PVConvertToSeedPointsSettings.h
new file mode 100644
--- /dev/null
+++ b/Engine/Plugins/Experimental/ProceduralVegetationEditor/Source/ProceduralVegetation/Private/Nodes/PVConvertToSeedPointsSettings.h
@@ -0,0 +1,39 @@
+// Copyright Epic Games, Inc. All Rights Reserved.
+#pragma once
+
+#include "CoreMinimal.h"
+#include "PCGSubgraph.h"
+#include "Nodes/PVBaseSettings.h"
+#include "DataTypes/PVData.h"
+#include "Implementations/PVSeedGenerator.h"
+#include "PVConvertToSeedPointsSettings.generated.h"
+
+UCLASS(BlueprintType, ClassGroup = (Procedural))
+class UPVConvertToSeedPointsSettings : public UPVBaseSettings
+{
+	GENERATED_BODY()
+public:
+#if WITH_EDITOR
+	virtual FName GetDefaultNodeName() const override { return FName(TEXT("ConvertToSeedPoints")); }
+	virtual FText GetDefaultNodeTitle() const override { return NSLOCTEXT("UPVConvertToSeedPointsSettings", "NodeTitle", "Convert To Seed Points"); }
+	virtual FText GetNodeTooltipText() const override { return NSLOCTEXT("UPVConvertToSeedPointsSettings", "NodeTooltip", "Create Seeds, that are used to grow plant"); }
+	virtual FLinearColor GetNodeTitleColor() const override { return FLinearColor(FColor::Black); }
+	virtual TArray<EPVRenderType> GetDefaultRenderType() const override { return TArray{ EPVRenderType::Seed }; }
+#endif
+
+protected:
+	virtual FPCGDataTypeIdentifier GetInputPinTypeIdentifier() const override;
+	virtual FPCGDataTypeIdentifier GetOutputPinTypeIdentifier() const override;
+	
+	virtual FPCGElementPtr CreateElement() const override;
+
+public:
+	UPROPERTY(EditAnywhere, Category = AvoidanceSettings, meta=(ShowOnlyInnerProperties))
+	FPVConvertToSeedPointParams Params;
+};
+
+class FPVConvertToSeedPointsElement : public FPVBaseElement
+{
+protected:
+	virtual bool ExecuteInternal(FPCGContext* InContext) const override;
+};
diff --git a/Engine/Plugins/Experimental/ProceduralVegetationEditor/Source/ProceduralVegetation/Private/Nodes/PVPointScatterSettings.cpp b/Engine/Plugins/Experimental/ProceduralVegetationEditor/Source/ProceduralVegetation/Private/Nodes/PVPointScatterSettings.cpp
new file mode 100644
--- /dev/null
+++ b/Engine/Plugins/Experimental/ProceduralVegetationEditor/Source/ProceduralVegetation/Private/Nodes/PVPointScatterSettings.cpp
@@ -0,0 +1,44 @@
+// Copyright Epic Games, Inc. All Rights Reserved.
+
+#include "Nodes/PVPointScatterSettings.h"
+#include "ProceduralVegetationPreset.h"
+#include "Data/PCGBasePointData.h"
+#include "Helpers/PCGHelpers.h"
+#include "UObject/SoftObjectPath.h"
+#include "UObject/ConstructorHelpers.h"
+
+#define LOCTEXT_NAMESPACE "PVPointScatterSettings"
+
+UPVPointScatterSettings::UPVPointScatterSettings(const FObjectInitializer& InObjectInitializer)
+	: Super(InObjectInitializer)
+{
+	if (PCGHelpers::IsNewObjectAndNotDefault(this))
+	{
+		TWeakObjectPtr WeakThis(this);
+		TSoftObjectPtr<UPCGGraphInterface> DefaultGraphSoftPtr = TSoftObjectPtr<UPCGGraphInterface>(FSoftObjectPath(TEXT("/ProceduralVegetationEditor/DefaultAssets/PCG/SeedPointScatter.SeedPointScatter")));
+		DefaultGraphSoftPtr.LoadAsync(FLoadSoftObjectPathAsyncDelegate::CreateLambda([WeakThis](const FSoftObjectPath&, UObject* InObject)
+		{
+			if (!WeakThis.IsValid())
+			{
+				return;
+			}
+				
+			if (InObject)
+			{
+				WeakThis.Get()->SetSubgraph(Cast<UPCGGraphInterface>(InObject));
+			}
+		}));
+	}
+}
+
+FPCGElementPtr UPVPointScatterSettings::CreateElement() const
+{
+	return MakeShared<FPVPointScatterElement>();
+}
+
+bool FPVPointScatterElement::ExecuteInternal(FPCGContext* InContext) const
+{
+	return FPCGSubgraphElement::ExecuteInternal(InContext);
+}
+
+#undef LOCTEXT_NAMESPACE
diff --git a/Engine/Plugins/Experimental/ProceduralVegetationEditor/Source/ProceduralVegetation/Public/DataTypes/PVData.h b/Engine/Plugins/Experimental/ProceduralVegetationEditor/Source/ProceduralVegetation/Public/DataTypes/PVData.h
--- a/Engine/Plugins/Experimental/ProceduralVegetationEditor/Source/ProceduralVegetation/Public/DataTypes/PVData.h
+++ b/Engine/Plugins/Experimental/ProceduralVegetationEditor/Source/ProceduralVegetation/Public/DataTypes/PVData.h
@@ -22,6 +22,7 @@ enum class EPVRenderType : uint8
 	Bones 				= 1 << 3,
 	FoliageGrid 		= 1 << 4,
 	FoliageAttachments	= 1 << 5,
+	Seed				= 1 << 6,
 
 	// Add new renter type above
 	Count	
diff --git a/Engine/Plugins/Experimental/ProceduralVegetationEditor/Source/ProceduralVegetation/Public/Facades/PVAttributesNames.h b/Engine/Plugins/Experimental/ProceduralVegetationEditor/Source/ProceduralVegetation/Public/Facades/PVAttributesNames.h
--- a/Engine/Plugins/Experimental/ProceduralVegetationEditor/Source/ProceduralVegetation/Public/Facades/PVAttributesNames.h
+++ b/Engine/Plugins/Experimental/ProceduralVegetationEditor/Source/ProceduralVegetation/Public/Facades/PVAttributesNames.h
@@ -62,6 +62,8 @@ namespace PV::AttributeNames
 	inline constexpr TCHAR PointScaleGradientStr[] = TEXT("LOD_totalPscaleGradient");
 	inline constexpr TCHAR PointScaleStr[] = TEXT("Scale");
 	inline constexpr TCHAR LengthFromSeedStr[] = TEXT("lengthFromSeed");
+	inline constexpr TCHAR SeedPScaleStr[] = TEXT("seedPScale");
+	inline constexpr TCHAR SeedPScaleRatioStr[] = TEXT("seedPScaleRatio");
 	inline constexpr TCHAR BudLightDetectedStr[] = TEXT("budLightDetected");
 	inline constexpr TCHAR BudDevelopmentStr[] = TEXT("budDevelopment");
 	inline constexpr TCHAR BudLateralMeristemStr[] = TEXT("budLateralMeristem");
@@ -101,6 +103,10 @@ namespace PV::AttributeNames
 	inline static const FName PlantGradient = FName(PlantGradientStr);
 	inline static const FName NjordPixelIndex = FName(NjordPixelIndexStr);
 
+	inline static const FName SeedPScale = FName(SeedPScaleStr);
+	inline static const FName SeedPScaleRatio = FName(SeedPScaleRatioStr);
+	
+
 	//Foliage Attributes
 	inline constexpr TCHAR FoliageNameIDStr[] = TEXT("FoliageNameID");
 	inline constexpr TCHAR FoliageBranchIDStr[] = TEXT("FoliageBranchID");
diff --git a/Engine/Plugins/Experimental/ProceduralVegetationEditor/Source/ProceduralVegetation/Public/Facades/PVPointFacade.h b/Engine/Plugins/Experimental/ProceduralVegetationEditor/Source/ProceduralVegetation/Public/Facades/PVPointFacade.h
--- a/Engine/Plugins/Experimental/ProceduralVegetationEditor/Source/ProceduralVegetation/Public/Facades/PVPointFacade.h
+++ b/Engine/Plugins/Experimental/ProceduralVegetationEditor/Source/ProceduralVegetation/Public/Facades/PVPointFacade.h
@@ -173,6 +173,10 @@ Accessor.ModifyAt(i, InputArray[i]); \
 
 		void FillBudLateralMeristem(const TArray<TArray<float>>& InputArray);
 
+		void FillSeedPScale(const TArray<float>& InputArray);
+
+		void FillSeedPScaleRatio(const TArray<float>& InputArray);
+
 		template<typename T>
 		FORCEINLINE bool GetValueFromAccessor(TManagedArrayAccessor<T> Accessor, T& OutValue, int32 Index, const T DefaultValue) const
 		{
@@ -221,6 +225,10 @@ Accessor.ModifyAt(i, InputArray[i]); \
 		void GetBudStatus(int32 Index, TArray<int>& OutValue) const;
 		
 		void GetBudLateralMeristem(int32 Index, TArray<float>& OutValue) const;
+		
+		void GetSeedPScale(int32 Index, float& OutValue) const;
+		
+		void GetSeedPScaleRatio(int32 Index, float& OutValue) const;
 
 	private:
 		TManagedArrayAccessor<FVector3f> Positions;
@@ -237,6 +245,8 @@ Accessor.ModifyAt(i, InputArray[i]); \
 		TManagedArrayAccessor<TArray<float>> BudHormoneLevels;
 		TManagedArrayAccessor<float> PlantGradients;
 		TManagedArrayAccessor<float> LengthFromSeed;
+		TManagedArrayAccessor<float> SeedPScale;
+		TManagedArrayAccessor<float> SeedPScaleRatio;
 		TManagedArrayAccessor<int32> BudNumber;
 		TManagedArrayAccessor<float> NjordPixelIndex;
 		TManagedArrayAccessor<TArray<float>> BudLightDetected;
diff --git a/Engine/Plugins/Experimental/ProceduralVegetationEditor/Source/ProceduralVegetation/Public/Nodes/PVPointScatterSettings.h b/Engine/Plugins/Experimental/ProceduralVegetationEditor/Source/ProceduralVegetation/Public/Nodes/PVPointScatterSettings.h
new file mode 100644
--- /dev/null
+++ b/Engine/Plugins/Experimental/ProceduralVegetationEditor/Source/ProceduralVegetation/Public/Nodes/PVPointScatterSettings.h
@@ -0,0 +1,37 @@
+// Copyright Epic Games, Inc. All Rights Reserved.
+#pragma once
+
+#include "CoreMinimal.h"
+#include "PCGSubgraph.h"
+
+#include "PVPointScatterSettings.generated.h"
+
+UCLASS(BlueprintType, ClassGroup = (Procedural))
+class UPVPointScatterSettings : public UPCGSubgraphSettings
+{
+	GENERATED_BODY()
+public:
+	UPVPointScatterSettings(const FObjectInitializer& InObjectInitializer);
+	
+	virtual bool IsDynamicGraph() const override { return true; }
+	
+#if WITH_EDITOR
+	virtual FName GetDefaultNodeName() const override { return FName(TEXT("PointScatter")); }
+	virtual FText GetDefaultNodeTitle() const override { return NSLOCTEXT("UPVPointScatterSettings", "NodeTitle", "Point Scatter"); }
+	virtual FText GetNodeTooltipText() const override { return NSLOCTEXT("UPVPointScatterSettings", "NodeTooltip", "Scatter Points using a PCG Graph"); }
+	virtual FLinearColor GetNodeTitleColor() const override { return FLinearColor(FColor::Black); }
+	virtual EPCGSettingsType GetType() const override { return EPCGSettingsType::Subgraph; }
+	virtual bool GetPinExtraIcon(const UPCGPin* InPin, FName& OutExtraIcon, FText& OutTooltip) const override { return false; }
+	// TODO: rendering stuff
+	// virtual TArray<EPVRenderType> GetDefaultRenderType() const override { return TArray{ EPVRenderType::PointData }; }
+#endif
+
+protected:
+	virtual FPCGElementPtr CreateElement() const override;
+};
+
+class FPVPointScatterElement : public FPCGSubgraphElement
+{
+protected:
+	virtual bool ExecuteInternal(FPCGContext* InContext) const override;
+};
diff --git a/Engine/Plugins/Experimental/ProceduralVegetationEditor/Source/ProceduralVegetationEditor/Private/PVEditorSchema.cpp b/Engine/Plugins/Experimental/ProceduralVegetationEditor/Source/ProceduralVegetationEditor/Private/PVEditorSchema.cpp
--- a/Engine/Plugins/Experimental/ProceduralVegetationEditor/Source/ProceduralVegetationEditor/Private/PVEditorSchema.cpp
+++ b/Engine/Plugins/Experimental/ProceduralVegetationEditor/Source/ProceduralVegetationEditor/Private/PVEditorSchema.cpp
@@ -7,6 +7,7 @@
 #include "Elements/PCGCreateAttribute.h"
 #include "Elements/Metadata/PCGMetadataMathsOpElement.h"
 #include "Nodes/PVBaseSettings.h"
+#include "Nodes/PVPointScatterSettings.h"
 #include "Nodes/PVSubgraphSettings.h"
 
 EPCGElementType UPVEditorSchema::GetElementTypeFiltering() const
@@ -25,6 +26,7 @@ const FPCGGraphEditorCustomization& UPVEditorSchema::GetGraphEditorCustomization
 			UPCGCreateAttributeSetSettings::StaticClass(),
 			UPCGMetadataMathsSettings::StaticClass(),
 			UPVSubgraphSettings::StaticClass(),
+			UPVPointScatterSettings::StaticClass(),
 		}
 	};
 
diff --git a/Engine/Plugins/Experimental/ProceduralVegetationEditor/Source/ProceduralVegetationEditor/Private/Visualizations/PVDataVisualization.cpp b/Engine/Plugins/Experimental/ProceduralVegetationEditor/Source/ProceduralVegetationEditor/Private/Visualizations/PVDataVisualization.cpp
--- a/Engine/Plugins/Experimental/ProceduralVegetationEditor/Source/ProceduralVegetationEditor/Private/Visualizations/PVDataVisualization.cpp
+++ b/Engine/Plugins/Experimental/ProceduralVegetationEditor/Source/ProceduralVegetationEditor/Private/Visualizations/PVDataVisualization.cpp
@@ -28,7 +28,8 @@
 #include "Engine/StaticMesh.h"
  
 #include "Facades/PVFoliageFacade.h"
- 
+#include "Facades/PVPointFacade.h"
+
 #include "GeometryCollection/GeometryCollectionObject.h"
  
 #include "Helpers/PVUtilities.h"
@@ -62,6 +63,7 @@ void FPVDataVisualization::RegisterVisualizations()
 	RenderMap.Add(EPVRenderType::Bones, FPVRenderCallback::CreateRaw(this, &FPVDataVisualization::BonesRenderer));
 	RenderMap.Add(EPVRenderType::FoliageGrid, FPVRenderCallback::CreateRaw(this, &FPVDataVisualization::FoliageGridRenderer));
 	RenderMap.Add(EPVRenderType::FoliageAttachments, FPVRenderCallback::CreateRaw(this, &FPVDataVisualization::FoliageAttachmentPointsRenderer));
+	RenderMap.Add(EPVRenderType::Seed, FPVRenderCallback::CreateRaw(this, &FPVDataVisualization::SeedRenderer));
 }
  
 void FPVDataVisualization::ExecuteDebugDisplay(FPCGContext* Context, const UPCGSettingsInterface* SettingsInterface, const UPCGData* Data,
@@ -466,5 +468,59 @@ void FPVDataVisualization::FoliageAttachmentPointsRenderer(const FManagedArrayCo
 	
 	OutBounds = OutBounds + Bounds;
 }
+
+void FPVDataVisualization::SeedRenderer(const FManagedArrayCollection& InCollection, FPCGSceneSetupParams& InOutParams,
+	FBoxSphereBounds& OutBounds)
+{
+	const PV::Facades::FPointFacade PointFacade(InCollection);
+
+	TArray<TPair<FColor, FText>> LegendItems;
+	LegendItems.Add({FColor::Blue,  FText::FromString("Apical Direction")});
+	LegendItems.Add({FColor::Red,  FText::FromString("Axillary Direction")});
+	LegendItems.Add({FColor::Green,  FText::FromString("Up Vector")});
+	
+	const TObjectPtr<UInstancedStaticMeshComponent> ConeInstancer = NewObject<UInstancedStaticMeshComponent>(GetTransientPackage(), NAME_None, RF_Transient);
+	ConeInstancer->SetStaticMesh(LoadObject<UStaticMesh>(nullptr, TEXT("/ProceduralVegetationEditor/DefaultAssets/Visualization/SM_Cone.SM_Cone")));
+	FBox Bounds;
+	for (int32 PointIndex = 0; PointIndex < PointFacade.GetElementCount(); PointIndex++)
+	{
+		FVector3f Position = PointFacade.GetPosition(PointIndex);
+		float SeedScaleRatio = 1.0;
+		PointFacade.GetSeedPScaleRatio(PointIndex, SeedScaleRatio);
+		SeedScaleRatio *= 0.1;
+
+		TArray<FVector3f> Directions = PointFacade.GetBudDirection(PointIndex);
+
+		FQuat FullRotation = FQuat::Identity;
+		
+		if (Directions.Num() >= 2)
+		{
+			FQuat Quat = FQuat::FindBetweenVectors(FVector::UpVector, FVector(Directions[0]));
+			FQuat RollQuat = FQuat::FindBetweenVectors(Quat.GetForwardVector(), FVector(Directions[1]));
+			FullRotation = RollQuat * Quat;
+		}
+		
+		Bounds += static_cast<FVector>(Position);
+		
+		const int32 InstanceIndex = ConeInstancer->AddInstance(
+			FTransform(
+				FullRotation,
+				static_cast<FVector>(Position),
+				FVector(SeedScaleRatio)
+			)
+		);
+	}
+ 
+	ConeInstancer->MarkRenderStateDirty();
+ 
+	InOutParams.ManagedResources.Add(ConeInstancer);
+	InOutParams.Scene->AddComponent(ConeInstancer, FTransform::Identity);
+	
+	
+	TSharedPtr<SPVEditorViewport> ViewportWidget = StaticCastSharedPtr<SPVEditorViewport>(InOutParams.EditorViewportClient->GetEditorViewportWidget());
+	ViewportWidget->PopulateLegendOverlayText(LegendItems);
+	
+	OutBounds = OutBounds + Bounds;
+}
  
 #undef LOCTEXT_NAMESPACE
\ No newline at end of file
diff --git a/Engine/Plugins/Experimental/ProceduralVegetationEditor/Source/ProceduralVegetationEditor/Private/Visualizations/PVDataVisualization.h b/Engine/Plugins/Experimental/ProceduralVegetationEditor/Source/ProceduralVegetationEditor/Private/Visualizations/PVDataVisualization.h
--- a/Engine/Plugins/Experimental/ProceduralVegetationEditor/Source/ProceduralVegetationEditor/Private/Visualizations/PVDataVisualization.h
+++ b/Engine/Plugins/Experimental/ProceduralVegetationEditor/Source/ProceduralVegetationEditor/Private/Visualizations/PVDataVisualization.h
@@ -35,6 +35,7 @@ protected:
 	void BonesRenderer(const FManagedArrayCollection& InCollection, FPCGSceneSetupParams& InOutParams, FBoxSphereBounds& OutBounds);
 	void FoliageGridRenderer(const FManagedArrayCollection& InCollection, FPCGSceneSetupParams& InOutParams, FBoxSphereBounds& OutBounds);
 	void FoliageAttachmentPointsRenderer(const FManagedArrayCollection& InCollection, FPCGSceneSetupParams& InOutParams, FBoxSphereBounds& OutBounds);
+	void SeedRenderer(const FManagedArrayCollection& InCollection, FPCGSceneSetupParams& InOutParams, FBoxSphereBounds& OutBounds);
 };
 
 DECLARE_LOG_CATEGORY_EXTERN(LogProceduralVegetationDataVisualization, Log, All);
\ No newline at end of file
