Add a new node for rotating the branches along the axial direction of their parents

This CL adds a new node which allows you to rotate the branches of a plant along the axial direction of their parents. The main purpose for this node is to allow the user to add depth to a plant imported using the 2D importer, however, the node is implemented as a standart post-deformer node so it can be used with any growth data output.

Code changes:
A new helper file has also been added which supplies a helper function for performing recursive rotation operations as that is a commonly performed operation (slope, gravity, object avoidance).

Main files of interest:
- PVRotateAxillaryBranches.h/.cpp: The core implementation of the node logic
- PVRotateAxillaryBranchesSettings.h/.cpp: Node class for the branch rotation
- PVTransformHelper: A new helper file for performing transformations of PV skeleton data. Ensures dependent attributes are updated, and that the rotations are applied correctly.
- PVSlope.cpp: Update the slope node to use the new PVTransformHelper::RecursiveRotateBranchPoints_SplitPoints function

Main functions of interest:
- RecursiveRotateBranchPoints_SplitPoints: Recursivly traverses the tree and multiplies the rotation matrices to ensure rotations propagate correctly
- FPVRotateAxillaryBranches::ApplyRotateAxillaryBranches: Core logic for the new node

#Synced-CL 50813759

diff --git a/Engine/Plugins/Experimental/ProceduralVegetationEditor/Source/ProceduralVegetation/Private/Helpers/PVTransformHelper.cpp b/Engine/Plugins/Experimental/ProceduralVegetationEditor/Source/ProceduralVegetation/Private/Helpers/PVTransformHelper.cpp
new file mode 100644
--- /dev/null
+++ b/Engine/Plugins/Experimental/ProceduralVegetationEditor/Source/ProceduralVegetation/Private/Helpers/PVTransformHelper.cpp
@@ -0,0 +1,141 @@
+// Copyright Epic Games, Inc. All Rights Reserved.
+
+#include "Helpers/PVTransformHelper.h"
+#include "Helpers/PVAttributesHelper.h"
+
+bool PV::Transform::RecursiveRotateBranchPoints_SplitPoints(
+	const FRotateBranchPointsAttributeCollection& Attributes,
+	int32 BranchIndex,
+	const FRotateBranchPointsFunction& RotatePointsBody,
+	const FVector3f& TrunkPivot
+)
+{
+	using namespace PV::AttributesHelper;
+
+	if (!Attributes.IsValid())
+	{
+		return false;
+	}
+
+	struct FPointState
+	{
+		int32     ParentPointIndex = INDEX_NONE;
+		FVector3f ParentPivotPoint = FVector3f::ZeroVector;
+		FVector3f ParentIdentityPosition = FVector3f::ZeroVector;
+		FMatrix44f TransformMatrix = FMatrix44f::Identity;
+	};
+	RecursiveWalkBranchPoints_SplitPoints<FPointState, false>(
+		Attributes.BranchChildren,
+		Attributes.BranchParentNumber,
+		Attributes.BranchNumber,
+		Attributes.BranchPoints,
+		Attributes.PointLengthFromRoot,
+		BranchIndex,
+		FPointState(),
+		[&](int32 BranchIndex, int32 BranchPointIndex, int32 PointIndex, const TArray<int32>& ChildBranchIndices, FPointState& PointState)
+		{
+			const FRotateBranchPointParams Params = FRotateBranchPointParams(
+				BranchIndex,
+				BranchPointIndex,
+				PointIndex,
+				PointState.ParentPointIndex,
+				ChildBranchIndices
+			);
+			FRotateBranchPointResult Result = {
+				.Rotation = FQuat4f::Identity,
+				.bAbsoluteRotation = false,
+			};
+			RotatePointsBody(Params, Result);
+
+			const FVector3f& PointPosition = Attributes.PointPosition[PointIndex];
+			const FVector3f RelativePosition = PointState.ParentPointIndex == INDEX_NONE
+				? TrunkPivot
+				: PointPosition - PointState.ParentIdentityPosition;
+
+			const FMatrix44f LocalRotationMatrix = Result.Rotation.ToMatrix();
+			const FMatrix44f LocalTranslationMatrix = FTranslationMatrix44f(RelativePosition);
+			const FMatrix44f LocalTransformMatrix = LocalTranslationMatrix * LocalRotationMatrix;
+
+			if (Result.bAbsoluteRotation)
+			{
+				const FMatrix44f TranslationOnly = FTranslationMatrix44f(PointState.TransformMatrix.GetOrigin() + LocalTransformMatrix.GetOrigin());
+				PointState.TransformMatrix = LocalRotationMatrix * TranslationOnly;
+			}
+			else
+			{
+				PointState.TransformMatrix = LocalTransformMatrix * PointState.TransformMatrix;
+			}
+
+			PointState.ParentIdentityPosition = PointPosition;
+
+			const FVector3f PointTranslation = PointState.TransformMatrix.GetOrigin() - PointPosition;
+			const FQuat4f PointRotation = PointState.TransformMatrix.ToQuat();
+			TransformPoint(Attributes, BranchIndex, PointIndex, PointTranslation, PointRotation);
+
+			PointState.ParentPointIndex = PointIndex;
+
+			return EForEachResult::Continue;
+		}
+	);
+
+	return true;
+}
+
+void PV::Transform::TransformPoint(
+	const FTransformPointAttributeCollection& Attributes,
+	int32 BranchIndex,
+	int32 PointIndex,
+	const FVector3f& Translation,
+	const FQuat4f& Rotation
+)
+{
+	if (!Attributes.IsValid())
+	{
+		return;
+	}
+
+	Attributes.PointPosition[PointIndex] += Translation;
+
+	if (Attributes.OptionalBudDirection.IsValid())
+	{
+		PV::FBudDirectionView BudDirections = Attributes.OptionalBudDirection[PointIndex];
+		BudDirections.Apical = Rotation.RotateVector(BudDirections.Apical);
+		BudDirections.Axillary = Rotation.RotateVector(BudDirections.Axillary);
+		BudDirections.GuideCurve = Rotation.RotateVector(BudDirections.GuideCurve);
+		BudDirections.UpVector = Rotation.RotateVector(BudDirections.UpVector);
+	}
+
+	const bool bUpdateFoliateRotation = Attributes.OptionalBranchPoints.IsValid()
+		&& Attributes.OptionalPointLengthFromRoot.IsValid()
+		&& Attributes.OptionalBranchFoliageIDs.IsValid()
+		&& Attributes.OptionalFoliageLengthFromRoot.IsValid();
+
+	const TArray<int32> BranchSegmentFoliageEntryIds = bUpdateFoliateRotation
+		? PV::AttributesHelper::GetBranchPointFoliageEntryIds(
+			Attributes.OptionalBranchPoints,
+			Attributes.OptionalPointLengthFromRoot,
+			Attributes.OptionalBranchFoliageIDs,
+			Attributes.OptionalFoliageLengthFromRoot,
+			BranchIndex,
+			PointIndex
+		)
+		: TArray<int32>();
+
+	for (int32 FoliageEntryId : BranchSegmentFoliageEntryIds)
+	{
+		if (Attributes.OptionalFoliagePivotPoint.IsValid())
+		{
+			Attributes.OptionalFoliagePivotPoint[FoliageEntryId] += Translation;
+		}
+
+		if (Attributes.OptionalFoliageUpVector.IsValid())
+		{
+			Attributes.OptionalFoliageUpVector[FoliageEntryId] = Rotation.RotateVector(Attributes.OptionalFoliageUpVector[FoliageEntryId]);
+		}
+
+		if (Attributes.OptionalFoliageNormalVector.IsValid())
+		{
+			Attributes.OptionalFoliageNormalVector[FoliageEntryId] = Rotation.RotateVector(Attributes.OptionalFoliageNormalVector[FoliageEntryId]);
+		}
+	}
+}
diff --git a/Engine/Plugins/Experimental/ProceduralVegetationEditor/Source/ProceduralVegetation/Private/Helpers/PVTransformHelper.h b/Engine/Plugins/Experimental/ProceduralVegetationEditor/Source/ProceduralVegetation/Private/Helpers/PVTransformHelper.h
new file mode 100644
--- /dev/null
+++ b/Engine/Plugins/Experimental/ProceduralVegetationEditor/Source/ProceduralVegetation/Private/Helpers/PVTransformHelper.h
@@ -0,0 +1,177 @@
+// Copyright Epic Games, Inc. All Rights Reserved.
+#pragma once
+
+#include "Facades/PVAttributes.h"
+
+namespace PV::Transform
+{
+
+struct FRotateBranchPointParams
+{
+	int32 BranchIndex;
+	int32 BranchPointIndex; // The index of the point in the branch points attribute
+	int32 PointIndex;		// The global point index of this point
+	int32 ParentPointIndex; // INDEX_NONE for the first point on the trunk
+	const TArray<int32>& PointChildBranchIndices;
+
+	FRotateBranchPointParams() = delete;
+	FRotateBranchPointParams(
+		int32 InBranchIndex,
+		int32 InBranchPointIndex,
+		int32 InPointIndex,
+		int32 InParentPointIndex,
+		const TArray<int32>& InPointChildBranchIndices
+	)
+		: BranchIndex(InBranchIndex)
+		, BranchPointIndex(InBranchPointIndex)
+		, PointIndex(InPointIndex)
+		, ParentPointIndex(InParentPointIndex)
+		, PointChildBranchIndices(InPointChildBranchIndices)
+	{
+	}
+};
+
+struct FRotateBranchPointResult
+{
+	FQuat4f Rotation;
+	bool bAbsoluteRotation;  // If true, then the rotation of the parent point will be ignored when rotating this point.
+};
+
+using FRotateBranchPointsFunction = TFunction<void(const FRotateBranchPointParams& Params, FRotateBranchPointResult& Result)>;
+
+/*
+* Attributes required to traverse a tree and rotate the points around its parent point.
+* Attributes named as Optional are not required to peform this operation, but if present will
+* be updated to reflect the point transformations.
+*/
+struct FRotateBranchPointsAttributeCollection
+{
+	PV::FPointPositionAttributeView PointPosition;
+	PV::FBudDirectionAttributeView BudDirection;
+	PV::FPointLengthFromRootAttributeConstView PointLengthFromRoot;
+	PV::FBranchChildrenAttributeConstView BranchChildren;
+	PV::FBranchParentNumberAttributeConstView BranchParentNumber;
+	PV::FBranchNumberAttributeConstView BranchNumber;
+	PV::FBranchPointsAttributeConstView BranchPoints;
+
+	PV::FFoliagePivotPointAttributeView OptionalFoliagePivotPoint;
+	PV::FFoliageUpVectorAttributeView OptionalFoliageUpVector;
+	PV::FFoliageNormalVectorAttributeView OptionalFoliageNormalVector;
+	PV::FFoliageLengthFromRootAttributeConstView OptionalFoliageLengthFromRoot;
+	PV::FBranchFoliageIDsAttributeConstView OptionalBranchFoliageIDs;
+
+	FRotateBranchPointsAttributeCollection() = default;
+	FRotateBranchPointsAttributeCollection(FManagedArrayCollection& Collection)
+		: PointPosition(PV::FPointPositionAttribute::FindAttribute(Collection))
+		, BudDirection(PV::FBudDirectionAttribute::FindAttribute(Collection))
+		, PointLengthFromRoot(PV::FPointLengthFromRootAttribute::FindAttribute(Collection))
+		, BranchChildren(PV::FBranchChildrenAttribute::FindAttribute(Collection))
+		, BranchParentNumber(PV::FBranchParentNumberAttribute::FindAttribute(Collection))
+		, BranchNumber(PV::FBranchNumberAttribute::FindAttribute(Collection))
+		, BranchPoints(PV::FBranchPointsAttribute::FindAttribute(Collection))
+		, OptionalFoliagePivotPoint(PV::FFoliagePivotPointAttribute::FindAttribute(Collection))
+		, OptionalFoliageUpVector(PV::FFoliageUpVectorAttribute::FindAttribute(Collection))
+		, OptionalFoliageNormalVector(PV::FFoliageNormalVectorAttribute::FindAttribute(Collection))
+		, OptionalFoliageLengthFromRoot(PV::FFoliageLengthFromRootAttribute::FindAttribute(Collection))
+		, OptionalBranchFoliageIDs(PV::FBranchFoliageIDsAttribute::FindAttribute(Collection))
+	{
+	}
+
+	bool IsValid() const
+	{
+		return PointPosition
+			&& BudDirection
+			&& PointLengthFromRoot
+			&& BranchChildren
+			&& BranchParentNumber
+			&& BranchNumber
+			&& BranchPoints;
+	}
+};
+
+/**
+ * Recursively rotates all points in the branch hierarchy starting from the specified BranchIndex.
+ * 
+ * Parent rotations are multiplied into child transforms ensuring that child points inherit all 
+ * parent rotations (unless bAbsoluteRotation is set to true). The behavior mirrors bone 
+ * transform accumulation in a skeletal hierarchy.
+ *
+ * This function performs a depth-first traversal of the branch hierarchy:
+ *  1. It walks all points in the current branch in order.
+ *  2. For each point, it invokes RotatePointsBody, allowing the caller to define
+ *     a local rotation for that point.
+ *  3. The computed rotation is accumulated with the parent’s transform, ensuring
+ *     that rotations propagate hierarchically (similar to skeletal animation).
+ *     Each point is effectively rotated around its parent pivot (or TrunkPivot
+ *     for the root).
+ *  4. After finishing all points in the current branch, the function recursively
+ *     processes any axillary (child) branches originating from points in this branch.
+ *
+ * @param Attributes       The attributes requried to traverse the tree and update the point rotations.
+ * @param BranchIndex      Index of the branch to start traversal from.
+ * @param RotatePointsBody Callback invoked for each point to compute its rotation.
+ * @param TrunkPivot       Pivot used for the root branch.
+ *
+ * @return true if the traversal and rotation succeeded; false otherwise.
+ */
+PROCEDURALVEGETATION_API extern bool RecursiveRotateBranchPoints_SplitPoints(
+	const FRotateBranchPointsAttributeCollection& Attributes,
+	int32 BranchIndex,
+	const FRotateBranchPointsFunction& RotatePointsBody,
+	const FVector3f& TrunkPivot = FVector3f::ZeroVector
+);
+
+struct FTransformPointAttributeCollection
+{
+	PV::FPointPositionAttributeView PointPosition;
+	PV::FBudDirectionAttributeView OptionalBudDirection;
+	PV::FBranchPointsAttributeConstView OptionalBranchPoints;
+	PV::FPointLengthFromRootAttributeConstView OptionalPointLengthFromRoot;
+	PV::FFoliagePivotPointAttributeView OptionalFoliagePivotPoint;
+	PV::FFoliageUpVectorAttributeView OptionalFoliageUpVector;
+	PV::FFoliageNormalVectorAttributeView OptionalFoliageNormalVector;
+	PV::FFoliageLengthFromRootAttributeConstView OptionalFoliageLengthFromRoot;
+	PV::FBranchFoliageIDsAttributeConstView OptionalBranchFoliageIDs;
+
+	FTransformPointAttributeCollection() = default;
+	FTransformPointAttributeCollection(FManagedArrayCollection& Collection)
+		: PointPosition(PV::FPointPositionAttribute::FindAttribute(Collection))
+		, OptionalBudDirection(PV::FBudDirectionAttribute::FindAttribute(Collection))
+		, OptionalBranchPoints(PV::FBranchPointsAttribute::FindAttribute(Collection))
+		, OptionalPointLengthFromRoot(PV::FPointLengthFromRootAttribute::FindAttribute(Collection))
+		, OptionalFoliagePivotPoint(PV::FFoliagePivotPointAttribute::FindAttribute(Collection))
+		, OptionalFoliageUpVector(PV::FFoliageUpVectorAttribute::FindAttribute(Collection))
+		, OptionalFoliageNormalVector(PV::FFoliageNormalVectorAttribute::FindAttribute(Collection))
+		, OptionalFoliageLengthFromRoot(PV::FFoliageLengthFromRootAttribute::FindAttribute(Collection))
+		, OptionalBranchFoliageIDs(PV::FBranchFoliageIDsAttribute::FindAttribute(Collection))
+	{
+	}
+
+	FTransformPointAttributeCollection(const FRotateBranchPointsAttributeCollection& Other)
+		: PointPosition(Other.PointPosition)
+		, OptionalBudDirection(Other.BudDirection)
+		, OptionalBranchPoints(Other.BranchPoints)
+		, OptionalPointLengthFromRoot(Other.PointLengthFromRoot)
+		, OptionalFoliagePivotPoint(Other.OptionalFoliagePivotPoint)
+		, OptionalFoliageUpVector(Other.OptionalFoliageUpVector)
+		, OptionalFoliageNormalVector(Other.OptionalFoliageNormalVector)
+		, OptionalFoliageLengthFromRoot(Other.OptionalFoliageLengthFromRoot)
+		, OptionalBranchFoliageIDs(Other.OptionalBranchFoliageIDs)
+	{
+	}
+
+	bool IsValid() const
+	{
+		return PointPosition.IsValid();
+	}
+};
+
+PROCEDURALVEGETATION_API extern void TransformPoint(
+	const FTransformPointAttributeCollection& Attributes,
+	int32 BranchIndex,
+	int32 PointIndex,
+	const FVector3f& Translation,
+	const FQuat4f& Rotation
+);
+
+};
\ No newline at end of file
diff --git a/Engine/Plugins/Experimental/ProceduralVegetationEditor/Source/ProceduralVegetation/Private/Implementations/PVRotateAxillaryBranches.cpp b/Engine/Plugins/Experimental/ProceduralVegetationEditor/Source/ProceduralVegetation/Private/Implementations/PVRotateAxillaryBranches.cpp
new file mode 100644
--- /dev/null
+++ b/Engine/Plugins/Experimental/ProceduralVegetationEditor/Source/ProceduralVegetation/Private/Implementations/PVRotateAxillaryBranches.cpp
@@ -0,0 +1,92 @@
+// Copyright Epic Games, Inc. All Rights Reserved.
+
+#include "PVRotateAxillaryBranches.h"
+#include "Helpers/PVTransformHelper.h"
+#include "Helpers/PVAttributesHelper.h"
+
+void FPVRotateAxillaryBranches::ApplyRotateAxillaryBranches(const FPVRotateAxillaryBranchesParams& InRotateAxillaryBranchesParams, FManagedArrayCollection& OutCollection)
+{
+	using namespace PV::Transform;
+
+	const FRotateBranchPointsAttributeCollection Attributes(OutCollection);
+	if (!Attributes.IsValid())
+	{
+		return;
+	}
+	
+	PV::FBudDevelopmentAttributeConstView BudDevelopmentAttribute = PV::FBudDevelopmentAttribute::FindAttribute(OutCollection);
+	if (!BudDevelopmentAttribute.IsValid())
+	{
+		return;
+	}
+
+	const int32 StartGeneration = InRotateAxillaryBranchesParams.StartGeneration;
+	const int32 EndGeneration = InRotateAxillaryBranchesParams.NumGenerations <= 0 
+		? MAX_int32 
+		: InRotateAxillaryBranchesParams.StartGeneration + InRotateAxillaryBranchesParams.NumGenerations;
+
+	const int32 StartLocalBranchGeneration = InRotateAxillaryBranchesParams.StartLocalBranchGeneration;
+	const int32 EndLocalBranchGeneration = InRotateAxillaryBranchesParams.NumLocalBranchGenerations <= 0 
+		? MAX_int32 
+		: InRotateAxillaryBranchesParams.StartLocalBranchGeneration + InRotateAxillaryBranchesParams.NumLocalBranchGenerations;
+
+	int32 CurrentBranchIndex = INDEX_NONE;
+	int32 CurrentLocalBranchGeneration = 0;
+	TArray<int32> BranchesLocalBranchGenerations;
+	BranchesLocalBranchGenerations.SetNumZeroed(Attributes.BranchNumber.Num());
+
+	float RotationSign = 1;
+	for (const int32 TrunkIndex : PV::AttributesHelper::GetTrunkIndices(Attributes.BranchParentNumber))
+	{
+		const TArray<int32>& CurrentBranchPoints = Attributes.BranchPoints[TrunkIndex];
+		if (CurrentBranchPoints.Num() == 0)
+		{
+			continue;
+		}
+
+		FVector PivotLocation = FVector::ZeroVector;
+		const FRotateBranchPointsFunction CalcRotationFunction = [&](const FRotateBranchPointParams& Params, FRotateBranchPointResult& Result)
+		{
+			if (Params.BranchIndex != CurrentBranchIndex)
+			{
+				CurrentLocalBranchGeneration = 0;
+				CurrentBranchIndex = Params.BranchIndex;
+			}
+
+			if (Params.PointChildBranchIndices.Num() > 0)
+			{
+				CurrentLocalBranchGeneration += 1;
+
+				for (int32 ChildBranchIndex : Params.PointChildBranchIndices)
+				{
+					BranchesLocalBranchGenerations[ChildBranchIndex] = CurrentLocalBranchGeneration;
+				}
+			}
+
+			if (Params.BranchPointIndex == 0 && Params.BranchIndex != TrunkIndex)
+			{
+				const int32 Generation = BudDevelopmentAttribute[Params.PointIndex].Generation;
+				if (Generation < StartGeneration || Generation >= EndGeneration)
+				{
+					return;
+				}
+
+				const int32 LocalBranchGeneration = BranchesLocalBranchGenerations[Params.BranchIndex];
+				if (LocalBranchGeneration < StartLocalBranchGeneration || LocalBranchGeneration >= EndLocalBranchGeneration)
+				{
+					return;
+				}
+
+				const FVector3f& ParentApical = Attributes.BudDirection[Params.ParentPointIndex].Apical;
+				Result.Rotation = FQuat4f(ParentApical, FMath::DegreesToRadians(InRotateAxillaryBranchesParams.Rotation * RotationSign));
+
+				if (InRotateAxillaryBranchesParams.bAlternatingRotations)
+				{
+					RotationSign *= -1;
+				}
+			}
+		};
+
+		RecursiveRotateBranchPoints_SplitPoints(Attributes, TrunkIndex, CalcRotationFunction);
+	}
+}
\ No newline at end of file
diff --git a/Engine/Plugins/Experimental/ProceduralVegetationEditor/Source/ProceduralVegetation/Private/Implementations/PVRotateAxillaryBranches.h b/Engine/Plugins/Experimental/ProceduralVegetationEditor/Source/ProceduralVegetation/Private/Implementations/PVRotateAxillaryBranches.h
new file mode 100644
--- /dev/null
+++ b/Engine/Plugins/Experimental/ProceduralVegetationEditor/Source/ProceduralVegetation/Private/Implementations/PVRotateAxillaryBranches.h
@@ -0,0 +1,41 @@
+// Copyright Epic Games, Inc. All Rights Reserved.
+
+#pragma once
+
+#include "CoreMinimal.h"
+#include "GeometryCollection/ManagedArrayCollection.h"
+#include "PVRotateAxillaryBranches.generated.h"
+
+USTRUCT()
+struct FPVRotateAxillaryBranchesParams
+{
+	GENERATED_BODY()
+
+	UPROPERTY(EditAnywhere, Category = "Settings", meta = (Units = Degrees, ClampMin = -180.0f, ClampMax = 180.0f, UIMin = -180.0f, UIMax = 180.0f))
+	float Rotation = 0;
+
+	// The generation to start applying the rotation
+	UPROPERTY(EditAnywhere, Category = "Settings", meta = (ClampMin = 1, UIMin = 1, UIMax = 20))
+	int32 StartGeneration = 1;
+
+	// The number of generations to apply the rotation for (0 = no limit)
+	UPROPERTY(EditAnywhere, Category = "Settings", meta = (ClampMin = 0, UIMin = 0, UIMax = 20))
+	int32 NumGenerations = 0;
+
+	// The generation of the branch on its parent.
+	// A branch with 3 child-branches will have generations 1, 2, 3. Branches starting on the same point have the same generation.
+	UPROPERTY(EditAnywhere, Category = "Settings", meta = (ClampMin = 1, UIMin = 1, UIMax = 20))
+	int32 StartLocalBranchGeneration = 1;
+
+	// The number of local branch generations to apply the rotation for (0 = no limit)
+	UPROPERTY(EditAnywhere, Category = "Settings", meta = (ClampMin = 0, UIMin = 0, UIMax = 20))
+	int32 NumLocalBranchGenerations = 0;
+
+	UPROPERTY(EditAnywhere, Category = "Settings")
+	bool bAlternatingRotations = true;
+};
+
+struct FPVRotateAxillaryBranches
+{
+	static void ApplyRotateAxillaryBranches(const FPVRotateAxillaryBranchesParams& InRotateAxillaryBranchesParams, FManagedArrayCollection& OutCollection);
+};
diff --git a/Engine/Plugins/Experimental/ProceduralVegetationEditor/Source/ProceduralVegetation/Private/Implementations/PVSlope.cpp b/Engine/Plugins/Experimental/ProceduralVegetationEditor/Source/ProceduralVegetation/Private/Implementations/PVSlope.cpp
--- a/Engine/Plugins/Experimental/ProceduralVegetationEditor/Source/ProceduralVegetation/Private/Implementations/PVSlope.cpp
+++ b/Engine/Plugins/Experimental/ProceduralVegetationEditor/Source/ProceduralVegetation/Private/Implementations/PVSlope.cpp
@@ -1,8 +1,8 @@
 // Copyright Epic Games, Inc. All Rights Reserved.
 
 #include "PVSlope.h"
+#include "Helpers/PVTransformHelper.h"
 #include "Helpers/PVAttributesHelper.h"
-#include "Facades/PVAttributes.h"
 
 namespace PVSlope
 {
@@ -16,15 +16,18 @@ namespace PVSlope
 		return YawRotation * LeanRotation;
 	}
 
-	float CalcBranchLength(
+	float CalculateTreeHeight(
 		PV::FBranchChildrenAttributeConstView BranchChildrenAttribute,
 		PV::FBranchParentNumberAttributeConstView BranchParentNumberAttribute,
 		PV::FBranchNumberAttributeConstView BranchNumberAttribute,
 		PV::FBranchPointsAttributeConstView BranchPointsAttribute,
 		PV::FPointPositionAttributeConstView PointPositionAttribute,
-		int32 BranchIndex)
+		PV::FPointLengthFromRootAttributeConstView PointLengthFromRootAttribute,
+		int32 TrunkIndex)
 	{
-		const TArray<int32>& CurrentBranchPoints = BranchPointsAttribute[BranchIndex];
+		using namespace PV::AttributesHelper;
+
+		const TArray<int32>& CurrentBranchPoints = BranchPointsAttribute[TrunkIndex];
 		if (CurrentBranchPoints.Num() == 0)
 		{
 			return 0;
@@ -33,13 +36,14 @@ namespace PVSlope
 		const FVector3f& RootPosition = PointPositionAttribute[CurrentBranchPoints[0]];
 
 		float BranchLengthSquared = 0;
-		PV::AttributesHelper::RecursiveWalkBranchPoints_FusedPoints(
+		RecursiveWalkBranchPoints_SplitPoints(
 			BranchChildrenAttribute,
 			BranchParentNumberAttribute,
 			BranchNumberAttribute,
 			BranchPointsAttribute,
-			BranchIndex,
-			[&](int32 BranchIndex, int32 BranchPointIndex, int32 PointIndex, const TArray<int32>& ChildPointIndices)->PV::AttributesHelper::EForEachResult
+			PointLengthFromRootAttribute,
+			TrunkIndex,
+			[&](int32 BranchIndex, int32 BranchPointIndex, int32 PointIndex, const TArray<int32>& ChildPointIndices)->EForEachResult
 			{
 				const FVector3f& PointPosition = PointPositionAttribute[PointIndex];
 				const float PointDistSquared = (PointPosition - RootPosition).SizeSquared();
@@ -47,188 +51,53 @@ namespace PVSlope
 				{
 					BranchLengthSquared = PointDistSquared;
 				}
-				return PV::AttributesHelper::EForEachResult::Continue;
+				return EForEachResult::Continue;
 			}
 		);
 
 		return BranchLengthSquared > 0 ? FMath::Sqrt(BranchLengthSquared) : 0.f;
 	}
-
-	struct FApplySlopeAttributes
-	{
-		PV::FPointPositionAttributeView PointPosition;
-		PV::FBudDirectionAttributeView BudDirection;
-		PV::FPointLengthFromRootAttributeView PointLengthFromRoot;
-		PV::FFoliagePivotPointAttributeView FoliagePivotPoint;
-		PV::FFoliageUpVectorAttributeView FoliageUpVector;
-		PV::FFoliageNormalVectorAttributeView FoliageNormalVector;
-		PV::FFoliageLengthFromRootAttributeConstView FoliageLengthFromRoot;
-		PV::FBranchChildrenAttributeConstView BranchChildren;
-		PV::FBranchParentNumberAttributeConstView BranchParentNumber;
-		PV::FBranchNumberAttributeConstView BranchNumber;
-		PV::FBranchPointsAttributeConstView BranchPoints;
-		PV::FBranchFoliageIDsAttributeConstView BranchFoliageIDs;
-	};
-
-	void ApplySlope_Recursive(
-		const FApplySlopeAttributes& Attributes,
-		const int32 BranchIndex,
-		const FPVSlopeParams& SlopeParams,
-		float TreeHeight,
-		const FVector3f& TrunkPosition,
-		FVector3f ParentPosition,
-		FVector3f ParentIdentityPosition
-	)
-	{
-		const TArray<int32>& CurrentBranchPoints = Attributes.BranchPoints[BranchIndex];
-
-		for (int32 Index = 0; Index < CurrentBranchPoints.Num(); ++Index)
-		{
-			const int32 PointIndex = CurrentBranchPoints[Index];
-			const int32 PrevPointIndex = Index > 0 ? CurrentBranchPoints[Index - 1] : INDEX_NONE;
-			FVector3f& Position = Attributes.PointPosition[PointIndex];
-
-			const float TreeAngleAlpha = 1.0 - (TreeHeight > 0 ? ((Position - TrunkPosition).Size() / TreeHeight) : 0);
-			const float AngleStrength = FMath::Clamp(FMath::Pow(TreeAngleAlpha, SlopeParams.BendStrength), 0.0, 1.0);
-
-			const FQuat4f SlopeRotation = CalculateBranchRotation(SlopeParams.SlopeDirection, SlopeParams.SlopeAngle * AngleStrength);
-
-			PV::FBudDirectionView BudDirections = Attributes.BudDirection[PointIndex];
-			BudDirections.Apical = SlopeRotation.RotateVector(BudDirections.Apical);
-			BudDirections.Axillary = SlopeRotation.RotateVector(BudDirections.Axillary);
-			BudDirections.GuideCurve = SlopeRotation.RotateVector(BudDirections.GuideCurve);
-			BudDirections.UpVector = SlopeRotation.RotateVector(BudDirections.UpVector);
-
-			const TArray<int32> BranchSegmentFoliageEntryIds = PV::AttributesHelper::GetBranchPointFoliageEntryIds(
-				Attributes.BranchPoints, 
-				Attributes.PointLengthFromRoot,
-				Attributes.BranchFoliageIDs, 
-				Attributes.FoliageLengthFromRoot,
-				BranchIndex, 
-				PointIndex
-			);
-
-			for (int32 FoliageEntryId : BranchSegmentFoliageEntryIds)
-			{
-				const FVector3f& FoliageEntryPivotPoint = Attributes.FoliagePivotPoint[FoliageEntryId];
-				const FVector3f NewFoliageEntryPivotPoint = SlopeRotation.RotateVector(FoliageEntryPivotPoint - ParentIdentityPosition) + ParentPosition;
-				Attributes.FoliagePivotPoint[FoliageEntryId] = NewFoliageEntryPivotPoint;
-
-				const FVector3f FoliageEntryUpVector = Attributes.FoliageUpVector[FoliageEntryId];
-				const FVector3f NewFoliageEntryUpVector = SlopeRotation.RotateVector(FoliageEntryUpVector);
-				Attributes.FoliageUpVector[FoliageEntryId] = NewFoliageEntryUpVector;
-
-				const FVector3f FoliageEntryNormalVector = Attributes.FoliageNormalVector[FoliageEntryId];
-				const FVector3f NewFoliageEntryNormalVector = SlopeRotation.RotateVector(FoliageEntryNormalVector);
-				Attributes.FoliageNormalVector[FoliageEntryId] = NewFoliageEntryNormalVector;
-			}
-
-			const FVector3f NewPosition = SlopeRotation.RotateVector(Position - ParentIdentityPosition) + ParentPosition;
-
-			ParentIdentityPosition = Position;
-			Position = NewPosition;
-			ParentPosition = Position;
-
-			const TArray<int32> BranchSegmentChildren = PV::AttributesHelper::GetBranchPointChildBranchIndices_SplitPoints(
-				Attributes.BranchChildren,
-				Attributes.BranchParentNumber,
-				Attributes.BranchNumber,
-				Attributes.BranchPoints,
-				Attributes.PointLengthFromRoot,
-				BranchIndex, 
-				PointIndex
-			);
-
-			for (const int32 ChildIndex : BranchSegmentChildren)
-			{
-				ApplySlope_Recursive(
-					Attributes,
-					ChildIndex,
-					SlopeParams,
-					TreeHeight,
-					TrunkPosition,
-					ParentPosition,
-					ParentIdentityPosition
-				);
-			}
-		}
-	}
 }
 
 void FPVSlope::ApplySlope(const FPVSlopeParams& InSlopeParams, FManagedArrayCollection& OutCollection)
 {
-	const PV::FBranchChildrenAttributeConstView BranchChildrenAttribute = PV::FBranchChildrenAttribute::FindAttribute(OutCollection);
-	const PV::FBranchParentNumberAttributeConstView BranchParentNumberAttribute = PV::FBranchParentNumberAttribute::FindAttribute(OutCollection);
-	const PV::FBranchNumberAttributeConstView BranchNumberAttribute = PV::FBranchNumberAttribute::FindAttribute(OutCollection);
-	const PV::FBranchPointsAttributeConstView BranchPointsAttribute = PV::FBranchPointsAttribute::FindAttribute(OutCollection);
-	const PV::FBranchFoliageIDsAttributeConstView BranchFoliageIDsAttribute = PV::FBranchFoliageIDsAttribute::FindAttribute(OutCollection);
-	const PV::FFoliageLengthFromRootAttributeConstView FoliageLengthFromRootAttribute = PV::FFoliageLengthFromRootAttribute::FindAttribute(OutCollection);
-
-	PV::FPointPositionAttributeView PointPositionAttribute = PV::FPointPositionAttribute::FindAttribute(OutCollection);
-	PV::FBudDirectionAttributeView BudDirectionAttribute = PV::FBudDirectionAttribute::FindAttribute(OutCollection);
-	PV::FPointLengthFromRootAttributeView PointLengthFromRootAttribute = PV::FPointLengthFromRootAttribute::FindAttribute(OutCollection);
-	PV::FFoliagePivotPointAttributeView FoliagePivotPointAttribute = PV::FFoliagePivotPointAttribute::FindAttribute(OutCollection);
-	PV::FFoliageUpVectorAttributeView FoliageUpVectorAttribute = PV::FFoliageUpVectorAttribute::FindAttribute(OutCollection);
-	PV::FFoliageNormalVectorAttributeView FoliageNormalVectorAttribute = PV::FFoliageNormalVectorAttribute::FindAttribute(OutCollection);
+	using namespace PV::Transform;
 
-	if (!BranchChildrenAttribute
-	 || !BranchParentNumberAttribute
-	 || !BranchNumberAttribute
-	 || !BranchPointsAttribute
-	 || !BranchFoliageIDsAttribute
-	 || !PointPositionAttribute
-	 || !BudDirectionAttribute
-	 || !PointLengthFromRootAttribute
-	 || !FoliagePivotPointAttribute
-	 || !FoliageUpVectorAttribute
-	 || !FoliageNormalVectorAttribute
-	 || !FoliageLengthFromRootAttribute)
+	const FRotateBranchPointsAttributeCollection Attributes(OutCollection);
+	if (!Attributes.IsValid())
 	{
 		return;
 	}
 
-	const PVSlope::FApplySlopeAttributes Attributes = {
-		PointPositionAttribute,
-		BudDirectionAttribute,
-		PointLengthFromRootAttribute,
-		FoliagePivotPointAttribute,
-		FoliageUpVectorAttribute,
-		FoliageNormalVectorAttribute,
-		FoliageLengthFromRootAttribute,
-		BranchChildrenAttribute,
-		BranchParentNumberAttribute,
-		BranchNumberAttribute,
-		BranchPointsAttribute,
-		BranchFoliageIDsAttribute
-	};
-
-	for (const int32 TrunkIndex : PV::AttributesHelper::GetTrunkIndices(BranchParentNumberAttribute))
+	for (const int32 TrunkIndex : PV::AttributesHelper::GetTrunkIndices(Attributes.BranchParentNumber))
 	{
-		const TArray<int32>& CurrentBranchPoints = BranchPointsAttribute[TrunkIndex];
+		const TArray<int32>& CurrentBranchPoints = Attributes.BranchPoints[TrunkIndex];
 		if (CurrentBranchPoints.Num() == 0)
 		{
 			continue;
 		}
 
-		const float BranchLength = PVSlope::CalcBranchLength(
-			BranchChildrenAttribute,
-			BranchParentNumberAttribute,
-			BranchNumberAttribute,
-			BranchPointsAttribute,
-			PointPositionAttribute,
+		const float TreeHeight = PVSlope::CalculateTreeHeight(
+			Attributes.BranchChildren,
+			Attributes.BranchParentNumber,
+			Attributes.BranchNumber,
+			Attributes.BranchPoints,
+			Attributes.PointPosition,
+			Attributes.PointLengthFromRoot,
 			TrunkIndex
 		);
-		const FVector3f TrunkPosition = PointPositionAttribute[CurrentBranchPoints[0]];
-		const FVector3f PivotPoint = InSlopeParams.TrunkPivotPoint == EPVSlopeTrunkPivotPoint::Origin ? FVector3f::ZeroVector : TrunkPosition;
+		const FVector3f TrunkPosition = Attributes.PointPosition[CurrentBranchPoints[0]];
+		const FVector3f TrunkPivot = InSlopeParams.TrunkPivotPoint == EPVSlopeTrunkPivotPoint::Origin ? FVector3f::ZeroVector : TrunkPosition;
+		const FRotateBranchPointsFunction CalcRotationFunction = [&](const FRotateBranchPointParams& Params, FRotateBranchPointResult& Result)
+		{
+			const FVector3f Position = Attributes.PointPosition[Params.PointIndex];
+			const float TreeAngleAlpha = 1.0 - (TreeHeight > 0 ? ((Position - TrunkPosition).Size() / TreeHeight) : 0);
+			const float AngleStrength = FMath::Clamp(FMath::Pow(TreeAngleAlpha, InSlopeParams.BendStrength), 0.0, 1.0);
 
-		PVSlope::ApplySlope_Recursive(
-			Attributes,
-			TrunkIndex,
-			InSlopeParams,
-			BranchLength,
-			TrunkPosition,
-			PivotPoint,
-			PivotPoint
-		);
+			Result.Rotation = PVSlope::CalculateBranchRotation(InSlopeParams.SlopeDirection, InSlopeParams.SlopeAngle * AngleStrength);
+			Result.bAbsoluteRotation = true;
+		};
+
+		RecursiveRotateBranchPoints_SplitPoints(Attributes, TrunkIndex, CalcRotationFunction, TrunkPivot);
 	}
 }
diff --git a/Engine/Plugins/Experimental/ProceduralVegetationEditor/Source/ProceduralVegetation/Private/Nodes/PVRotateAxillaryBranchesSettings.cpp b/Engine/Plugins/Experimental/ProceduralVegetationEditor/Source/ProceduralVegetation/Private/Nodes/PVRotateAxillaryBranchesSettings.cpp
new file mode 100644
--- /dev/null
+++ b/Engine/Plugins/Experimental/ProceduralVegetationEditor/Source/ProceduralVegetation/Private/Nodes/PVRotateAxillaryBranchesSettings.cpp
@@ -0,0 +1,74 @@
+// Copyright Epic Games, Inc. All Rights Reserved.
+
+#include "PVRotateAxillaryBranchesSettings.h"
+#include "DataTypes/PVData.h"
+#include "DataTypes/PVGrowthData.h"
+#include "Data/PCGBasePointData.h"
+#include "GeometryCollection/ManagedArrayCollection.h"
+#include "Helpers/PCGSettingsHelpers.h"
+
+#define LOCTEXT_NAMESPACE "PVRotateAxillaryBranchesSettings"
+
+#if WITH_EDITOR
+FText UPVRotateAxillaryBranchesSettings::GetDefaultNodeTitle() const
+{ 
+	return LOCTEXT("NodeTitle", "RotateAxillaryBranches"); 
+}
+
+FText UPVRotateAxillaryBranchesSettings::GetNodeTooltipText() const
+{ 
+	return LOCTEXT("NodeTooltip", 
+		"Rotates the axillary branches around the apical direction of its parent branch."
+		"\n\nPress Ctrl + L to lock/unlock node output"
+	); 
+}
+#endif
+
+FPCGDataTypeIdentifier UPVRotateAxillaryBranchesSettings::GetInputPinTypeIdentifier() const
+{
+	return FPCGDataTypeIdentifier{ FPVDataTypeInfoGrowth::AsId() };
+}
+
+FPCGDataTypeIdentifier UPVRotateAxillaryBranchesSettings::GetOutputPinTypeIdentifier() const
+{
+	return FPCGDataTypeIdentifier{ FPVDataTypeInfoGrowth::AsId() };
+}
+
+FPCGElementPtr UPVRotateAxillaryBranchesSettings::CreateElement() const
+{
+	return MakeShared<FPVRotateAxillaryBranchesElement>();
+}
+
+bool FPVRotateAxillaryBranchesElement::ExecuteInternal(FPCGContext* InContext) const
+{
+	TRACE_CPUPROFILER_EVENT_SCOPE(FPVRotateAxillaryBranchesElement::Execute);
+
+	check(InContext);
+
+	const UPVRotateAxillaryBranchesSettings* Settings = InContext->GetInputSettings<UPVRotateAxillaryBranchesSettings>();
+	check(Settings);
+
+	for (const FPCGTaggedData& Input : InContext->InputData.GetInputsByPin(PCGPinConstants::DefaultInputLabel))
+	{
+		if (const UPVGrowthData* InputData = Cast<UPVGrowthData>(Input.Data))
+		{
+			FManagedArrayCollection OutCollection = InputData->GetCollection();
+
+			UPVGrowthData* OutManagedArrayCollectionData = FPCGContext::NewObject_AnyThread<UPVGrowthData>(InContext);
+
+			FPVRotateAxillaryBranches::ApplyRotateAxillaryBranches(Settings->RotateAxillaryBranchesParams, OutCollection);
+
+			OutManagedArrayCollectionData->Initialize(MoveTemp(OutCollection));
+			InContext->OutputData.TaggedData.Emplace(OutManagedArrayCollectionData);
+		}
+		else
+		{
+			PCGLog::InputOutput::LogInvalidInputDataError(InContext);
+			return true;
+		}
+	}
+
+	return true;
+}
+
+#undef LOCTEXT_NAMESPACE
\ No newline at end of file
diff --git a/Engine/Plugins/Experimental/ProceduralVegetationEditor/Source/ProceduralVegetation/Private/Nodes/PVRotateAxillaryBranchesSettings.h b/Engine/Plugins/Experimental/ProceduralVegetationEditor/Source/ProceduralVegetation/Private/Nodes/PVRotateAxillaryBranchesSettings.h
new file mode 100644
--- /dev/null
+++ b/Engine/Plugins/Experimental/ProceduralVegetationEditor/Source/ProceduralVegetation/Private/Nodes/PVRotateAxillaryBranchesSettings.h
@@ -0,0 +1,39 @@
+// Copyright Epic Games, Inc. All Rights Reserved.
+#pragma once
+
+#include "CoreMinimal.h"
+#include "Nodes/PVBaseSettings.h"
+#include "Implementations/PVRotateAxillaryBranches.h"
+#include "PVRotateAxillaryBranchesSettings.generated.h"
+
+UCLASS(BlueprintType, ClassGroup = (Procedural))
+class UPVRotateAxillaryBranchesSettings : public UPVBaseSettings
+{
+	GENERATED_BODY()
+
+public:
+#if WITH_EDITOR
+	virtual FName GetDefaultNodeName() const override { return FName(TEXT("ProceduralVegetationRotateAxillaryBranches")); }
+	virtual FText GetDefaultNodeTitle() const override;
+	virtual FText GetNodeTooltipText() const override;
+	virtual FLinearColor GetNodeTitleColor() const override { return FLinearColor::Yellow; }
+
+	virtual TArray<EPVRenderType> GetDefaultRenderType() const override { return TArray{ EPVRenderType::PointData }; }
+#endif
+
+protected:
+	virtual FPCGDataTypeIdentifier GetInputPinTypeIdentifier() const override;
+	virtual FPCGDataTypeIdentifier GetOutputPinTypeIdentifier() const override;
+
+	virtual FPCGElementPtr CreateElement() const override;
+
+public:
+	UPROPERTY(EditAnywhere, Category = "RotateAxillaryBranches", meta = (ShowOnlyInnerProperties, PCG_Overridable))
+	FPVRotateAxillaryBranchesParams RotateAxillaryBranchesParams;
+};
+
+class FPVRotateAxillaryBranchesElement : public FPVBaseElement
+{
+protected:
+	virtual bool ExecuteInternal(FPCGContext* InContext) const override;
+};
diff --git a/Engine/Plugins/Experimental/ProceduralVegetationEditor/Source/ProceduralVegetation/Public/Facades/PVAttributes.h b/Engine/Plugins/Experimental/ProceduralVegetationEditor/Source/ProceduralVegetation/Public/Facades/PVAttributes.h
--- a/Engine/Plugins/Experimental/ProceduralVegetationEditor/Source/ProceduralVegetation/Public/Facades/PVAttributes.h
+++ b/Engine/Plugins/Experimental/ProceduralVegetationEditor/Source/ProceduralVegetation/Public/Facades/PVAttributes.h
@@ -70,6 +70,10 @@ public:
 
 	ManagedArrayType* ManagedArray;
 
+	TAttributeView()
+		: ManagedArray(nullptr)
+	{}
+
 	TAttributeView(ManagedArrayType* InManagedArray)
 		: ManagedArray(InManagedArray)
 	{}
diff --git a/Engine/Plugins/Experimental/ProceduralVegetationEditor/Source/ProceduralVegetation/Public/Helpers/PVAttributesHelper.h b/Engine/Plugins/Experimental/ProceduralVegetationEditor/Source/ProceduralVegetation/Public/Helpers/PVAttributesHelper.h
--- a/Engine/Plugins/Experimental/ProceduralVegetationEditor/Source/ProceduralVegetation/Public/Helpers/PVAttributesHelper.h
+++ b/Engine/Plugins/Experimental/ProceduralVegetationEditor/Source/ProceduralVegetation/Public/Helpers/PVAttributesHelper.h
@@ -311,7 +311,7 @@ namespace PV::AttributesHelper
 		return ChildrenIndices;
 	}
 
-	template <typename StateType>
+	template <typename StateType, bool bWalkAxillaryFirst = false>
 	inline EForEachResult RecursiveWalkBranchPoints_SplitPoints(
 		PV::FBranchChildrenAttributeConstView BranchChildrenAttribute,
 		PV::FBranchParentNumberAttributeConstView BranchParentNumberAttribute,
@@ -323,45 +323,76 @@ namespace PV::AttributesHelper
 		const TFunction<EForEachResult(int32 BranchIndex, int32 BranchPointIndex, int32 PointIndex, const TArray<int32>& ChildBranchIndices, StateType& PointState)>& Body
 	)
 	{
-		const TFunction<EForEachResult(int32, const StateType&)> WalkPoints = [&](int32 BranchIndex, const StateType& ParentState) -> EForEachResult
+		const TFunction<EForEachResult(int32, int32, const StateType&)> WalkPoints = [&](int32 BranchIndex, int32 BranchPointIndex, const StateType& ParentState) -> EForEachResult
 		{
 			StateType PointState = ParentState;
+			
 			const TArray<int32>& CurrentBranchPoints = BranchPointsAttribute[BranchIndex];
+			const int32 PointIndex = CurrentBranchPoints[BranchPointIndex];
+
+			const TArray<int32> ChildBranchIndices = GetBranchPointChildBranchIndices_SplitPoints(
+				BranchChildrenAttribute,
+				BranchParentNumberAttribute,
+				BranchNumberAttribute,
+				BranchPointsAttribute,
+				PointLengthFromRootAttribute,
+				BranchIndex,
+				PointIndex
+			);
 
-			for (int32 Index = 0; Index < CurrentBranchPoints.Num(); ++Index)
+			if (Body(BranchIndex, BranchPointIndex, PointIndex, ChildBranchIndices, PointState) == EForEachResult::Break)
 			{
-				const int32 PointIndex = CurrentBranchPoints[Index];
-				const TArray<int32> ChildBranchIndices = GetBranchPointChildBranchIndices_SplitPoints(
-					BranchChildrenAttribute,
-					BranchParentNumberAttribute,
-					BranchNumberAttribute,
-					BranchPointsAttribute,
-					PointLengthFromRootAttribute,
-					BranchIndex,
-					PointIndex
-				);
+				return EForEachResult::Break;
+			}
 
-				if (Body(BranchIndex, Index, PointIndex, ChildBranchIndices, PointState) == EForEachResult::Break)
+			struct FChildPoint
+			{
+				int32 BranchIndex;
+				int32 BranchPointIndex;
+			};
+			TArray<FChildPoint, TInlineAllocator<3>> Children;
+			Children.Reserve(1 + ChildBranchIndices.Num());
+
+			if constexpr (bWalkAxillaryFirst)
+			{
+				for (int32 ChildBranchIndex : ChildBranchIndices)
 				{
-					return EForEachResult::Break;
+					Children.Emplace(ChildBranchIndex, 0);
+				}
+
+				if (BranchPointIndex < CurrentBranchPoints.Num() - 1)
+				{
+					Children.Emplace(BranchIndex, BranchPointIndex + 1);
+				}
+			}
+			else
+			{
+				if (BranchPointIndex < CurrentBranchPoints.Num() - 1)
+				{
+					Children.Emplace(BranchIndex, BranchPointIndex + 1);
 				}
 
-				for (const int32 ChildBranchIndex : ChildBranchIndices)
+				for (int32 ChildBranchIndex : ChildBranchIndices)
+				{
+					Children.Emplace(ChildBranchIndex, 0);
+				}
+			}
+
+			for (const auto& [ChildBranchIndex, ChildBranchPointIndex] : Children)
+			{
+				if (WalkPoints(ChildBranchIndex, ChildBranchPointIndex, PointState) == EForEachResult::Break)
 				{
-					if (WalkPoints(ChildBranchIndex, PointState) == EForEachResult::Break)
-					{
-						return EForEachResult::Break;
-					}
+					return EForEachResult::Break;
 				}
 			}
 
 			return EForEachResult::Continue;
 		};
 
-		return WalkPoints(BranchIndex, RootPointState);
+		return WalkPoints(BranchIndex, 0, RootPointState);
 	}
 
-	template <typename StateType>
+	template <typename StateType, bool bWalkAxillaryFirst = false>
 	inline EForEachResult RecursiveWalkBranchPoints_FusedPoints(
 		PV::FBranchChildrenAttributeConstView BranchChildrenAttribute,
 		PV::FBranchParentNumberAttributeConstView BranchParentNumberAttribute,
@@ -372,44 +403,75 @@ namespace PV::AttributesHelper
 		const TFunction<EForEachResult(int32 BranchIndex, int32 BranchPointIndex, int32 PointIndex, const TArray<int32>& ChildBranchIndices, StateType& PointState)>& Body
 	)
 	{
-		const TFunction<EForEachResult(int32, bool, const StateType&)> WalkPoints = [&](int32 BranchIndex, bool bSkipFirstBranchPoint, const StateType& ParentState) -> EForEachResult
+		const TFunction<EForEachResult(int32, int32, const StateType&)> WalkPoints = [&](int32 BranchIndex, int32 BranchPointIndex, const StateType& ParentState) -> EForEachResult
 		{
 			StateType PointState = ParentState;
+			
 			const TArray<int32>& CurrentBranchPoints = BranchPointsAttribute[BranchIndex];
+			const int32 PointIndex = CurrentBranchPoints[BranchPointIndex];
+
+			const TArray<int32> ChildBranchIndices = GetBranchPointChildBranchIndices_FusedPoints(
+				BranchChildrenAttribute,
+				BranchParentNumberAttribute,
+				BranchNumberAttribute,
+				BranchPointsAttribute,
+				BranchIndex,
+				PointIndex
+			);
+
+			if (Body(BranchIndex, BranchPointIndex, PointIndex, ChildBranchIndices, PointState) == EForEachResult::Break)
+			{
+				return EForEachResult::Break;
+			}
+
+			struct FChildPoint
+			{
+				int32 BranchIndex;
+				int32 BranchPointIndex;
+			};
+			TArray<FChildPoint, TInlineAllocator<3>> Children;
+			Children.Reserve(1 + ChildBranchIndices.Num());
 
-			const int32 StartIndex = bSkipFirstBranchPoint ? 1 : 0;
-			for (int32 Index = StartIndex; Index < CurrentBranchPoints.Num(); ++Index)
+			if constexpr (bWalkAxillaryFirst)
 			{
-				const int32 PointIndex = CurrentBranchPoints[Index];
-				const TArray<int32> ChildBranchIndices = GetBranchPointChildBranchIndices_FusedPoints(
-					BranchChildrenAttribute,
-					BranchParentNumberAttribute,
-					BranchNumberAttribute,
-					BranchPointsAttribute,
-					BranchIndex,
-					PointIndex
-				);
+				for (int32 ChildBranchIndex : ChildBranchIndices)
+				{
+					Children.Emplace(ChildBranchIndex, 1);
+				}
 
-				if (Body(BranchIndex, Index, PointIndex, ChildBranchIndices, PointState) == EForEachResult::Break)
+				if (BranchPointIndex < CurrentBranchPoints.Num() - 1)
 				{
-					return EForEachResult::Break;
+					Children.Emplace(BranchIndex, BranchPointIndex + 1);
+				}
+			}
+			else
+			{
+				if (BranchPointIndex < CurrentBranchPoints.Num() - 1)
+				{
+					Children.Emplace(BranchIndex, BranchPointIndex + 1);
 				}
 
-				for (const int32 ChildBranchIndex : ChildBranchIndices)
+				for (int32 ChildBranchIndex : ChildBranchIndices)
 				{
-					if (WalkPoints(ChildBranchIndex, true, PointState) == EForEachResult::Break)
-					{
-						return EForEachResult::Break;
-					}
+					Children.Emplace(ChildBranchIndex, 1);
+				}
+			}
+
+			for (const auto& [ChildBranchIndex, ChildBranchPointIndex] : Children)
+			{
+				if (WalkPoints(ChildBranchIndex, ChildBranchPointIndex, PointState) == EForEachResult::Break)
+				{
+					return EForEachResult::Break;
 				}
 			}
 
 			return EForEachResult::Continue;
 		};
 
-		return WalkPoints(BranchIndex, false, RootPointState);
+		return WalkPoints(BranchIndex, 0, RootPointState);
 	}
 
+	template<bool bWalkAxillaryFirst = false>
 	inline EForEachResult RecursiveWalkBranchPoints_FusedPoints(
 		PV::FBranchChildrenAttributeConstView BranchChildrenAttribute,
 		PV::FBranchParentNumberAttributeConstView BranchParentNumberAttribute,
@@ -419,7 +481,7 @@ namespace PV::AttributesHelper
 		const TFunction<EForEachResult(int32 BranchIndex, int32 BranchPointIndex, int32 PointIndex, const TArray<int32>& ChildPointIndices)>& Body
 	)
 	{
-		return RecursiveWalkBranchPoints_FusedPoints<int32>(
+		return RecursiveWalkBranchPoints_FusedPoints<int32, bWalkAxillaryFirst>(
 			BranchChildrenAttribute,
 			BranchParentNumberAttribute,
 			BranchNumberAttribute,
@@ -433,7 +495,33 @@ namespace PV::AttributesHelper
 		);
 	}
 
-	template <typename StateType>
+	template<bool bWalkAxillaryFirst = false>
+	inline EForEachResult RecursiveWalkBranchPoints_SplitPoints(
+		PV::FBranchChildrenAttributeConstView BranchChildrenAttribute,
+		PV::FBranchParentNumberAttributeConstView BranchParentNumberAttribute,
+		PV::FBranchNumberAttributeConstView BranchNumberAttribute,
+		PV::FBranchPointsAttributeConstView BranchPointsAttribute,
+		PV::FPointLengthFromRootAttributeConstView PointLengthFromRootAttribute,
+		int32 BranchIndex,
+		const TFunction<EForEachResult(int32 BranchIndex, int32 BranchPointIndex, int32 PointIndex, const TArray<int32>& ChildPointIndices)>& Body
+	)
+	{
+		return RecursiveWalkBranchPoints_SplitPoints<int32, bWalkAxillaryFirst>(
+			BranchChildrenAttribute,
+			BranchParentNumberAttribute,
+			BranchNumberAttribute,
+			BranchPointsAttribute,
+			PointLengthFromRootAttribute,
+			BranchIndex,
+			0,
+			[&](int32 BranchIndex, int32 BranchPointIndex, int32 PointIndex, const TArray<int32>& ChildPointIndices, int32&) -> EForEachResult
+			{
+				return Body(BranchIndex, BranchPointIndex, PointIndex, ChildPointIndices);
+			}
+		);
+	}
+
+	template <typename StateType, bool bWalkAxillaryBranchesFirst = false>
 	void RecursiveWalkPlantPoints_FusedPoints(
 		PV::FBranchChildrenAttributeConstView BranchChildrenAttribute,
 		PV::FBranchParentNumberAttributeConstView BranchParentNumberAttribute,
@@ -445,7 +533,7 @@ namespace PV::AttributesHelper
 	{
 		for (const int32 TrunkIndex : GetTrunkIndices(BranchParentNumberAttribute))
 		{
-			auto Result = RecursiveWalkBranchPoints_FusedPoints(
+			auto Result = RecursiveWalkBranchPoints_FusedPoints<StateType, bWalkAxillaryBranchesFirst>(
 				BranchChildrenAttribute,
 				BranchParentNumberAttribute,
 				BranchNumberAttribute,
