[PVE] Build Health issue fix

#Synced-CL 50393760

diff --git a/Engine/Plugins/Experimental/ProceduralVegetationEditor/Source/ProceduralVegetation/Private/Implementations/PVSeedGenerator.cpp b/Engine/Plugins/Experimental/ProceduralVegetationEditor/Source/ProceduralVegetation/Private/Implementations/PVSeedGenerator.cpp
--- a/Engine/Plugins/Experimental/ProceduralVegetationEditor/Source/ProceduralVegetation/Private/Implementations/PVSeedGenerator.cpp
+++ b/Engine/Plugins/Experimental/ProceduralVegetationEditor/Source/ProceduralVegetation/Private/Implementations/PVSeedGenerator.cpp
@@ -25,9 +25,8 @@ TArray<FTransform> FPVConvertToSeedPointImplementation::FindNeighbours(const FTr
 	return Neighbours;
 }
 
-void FPVConvertToSeedPointImplementation::SetBaseDirection(const UPCGPointData* PointData, const FPVConvertToSeedPointParams& Params, TArray<TArray<FVector3f>>& DirectionArray)
+void FPVConvertToSeedPointImplementation::SetBaseDirection(const UPCGPointData* PointData, const FPVConvertToSeedPointParams& Params, TArray<TArray<FVector3f>>& DirectionArray, int PrimaryAxis /*= 0*/)
 {
-	int PrimaryAxis = 0;
 	TConstPCGValueRange<FTransform> PointTransforms = PointData->GetConstTransformValueRange();
 
 	float MinDistance = FLT_MAX;
diff --git a/Engine/Plugins/Experimental/ProceduralVegetationEditor/Source/ProceduralVegetation/Private/Implementations/PVSeedGenerator.h b/Engine/Plugins/Experimental/ProceduralVegetationEditor/Source/ProceduralVegetation/Private/Implementations/PVSeedGenerator.h
--- a/Engine/Plugins/Experimental/ProceduralVegetationEditor/Source/ProceduralVegetation/Private/Implementations/PVSeedGenerator.h
+++ b/Engine/Plugins/Experimental/ProceduralVegetationEditor/Source/ProceduralVegetation/Private/Implementations/PVSeedGenerator.h
@@ -100,7 +100,7 @@ struct FPVConvertToSeedPointImplementation
 {
 	static TArray<FTransform> FindNeighbours(const FTransform& InTransform, const UPCGPointData* PointData, float MaxDist, int MaxPoints);
 	
-	static void SetBaseDirection(const UPCGPointData* PointData, const FPVConvertToSeedPointParams& Params, TArray<TArray<FVector3f>>& DirectionArray);
+	static void SetBaseDirection(const UPCGPointData* PointData, const FPVConvertToSeedPointParams& Params, TArray<TArray<FVector3f>>& DirectionArray, int PrimaryAxis = 0);
 	
 	static void FillCollection(FManagedArrayCollection& Collection, const FPVConvertToSeedPointParams& Params, const UPCGPointData* SeedPoints);
 };
