Move PVImporter logic to ProceduralVegetation module from ProceduralVegetationEditor module

- Delete PVEditorInterface (only usecase for this class was to allow calling importer functions from the non-editor module)
- Move all import files to ProceduralVegetation module (PVImporter_StaticMesh, PVImporter_Texture2D, PVImportHelpers, PVDynamicMeshVertexAttribute, PVImportTests)
- Move Params files to "Params/..." directory (PVExportParams, PVImportStaticMeshParams, PVImportTexture2DParams)
- Move type/utils files to "Utils/..." directory (PVAttributes.h, PVFloatRamp.h, PVDynamicMeshVertexAttributes.h)
- Move files that no longer need to be public to private directory (PVImportStaticMeshParams.h, PVImportTexture2DParams.h)

#Synced-CL 50813759

diff --git a/Engine/Plugins/Experimental/ProceduralVegetationEditor/Source/ProceduralVegetation/Private/DataTypes/PVGrowerParams.h b/Engine/Plugins/Experimental/ProceduralVegetationEditor/Source/ProceduralVegetation/Private/DataTypes/PVGrowerParams.h
--- a/Engine/Plugins/Experimental/ProceduralVegetationEditor/Source/ProceduralVegetation/Private/DataTypes/PVGrowerParams.h
+++ b/Engine/Plugins/Experimental/ProceduralVegetationEditor/Source/ProceduralVegetation/Private/DataTypes/PVGrowerParams.h
@@ -3,7 +3,7 @@
 #pragma once
 
 #include "CoreMinimal.h"
-#include "PVFloatRamp.h"
+#include "Utils/PVFloatRamp.h"
 #include "PVGrowerParams.generated.h"
 
 struct FPVColliderParams;
diff --git a/Engine/Plugins/Experimental/ProceduralVegetationEditor/Source/ProceduralVegetation/Private/Helpers/PVAnalyticsHelper.cpp b/Engine/Plugins/Experimental/ProceduralVegetationEditor/Source/ProceduralVegetation/Private/Helpers/PVAnalyticsHelper.cpp
--- a/Engine/Plugins/Experimental/ProceduralVegetationEditor/Source/ProceduralVegetation/Private/Helpers/PVAnalyticsHelper.cpp
+++ b/Engine/Plugins/Experimental/ProceduralVegetationEditor/Source/ProceduralVegetation/Private/Helpers/PVAnalyticsHelper.cpp
@@ -4,7 +4,7 @@
  
 #include "EngineAnalytics.h"
 #include "ProceduralVegetationModule.h"
-#include "PVExportParams.h"
+#include "Params/PVExportParams.h"
 #include "Engine/RendererSettings.h"
 #include "Facades/PVFoliageFacade.h"
 #include "GeometryCollection/GeometryCollection.h"
diff --git a/Engine/Plugins/Experimental/ProceduralVegetationEditor/Source/ProceduralVegetation/Private/Helpers/PVImportHelpers.cpp b/Engine/Plugins/Experimental/ProceduralVegetationEditor/Source/ProceduralVegetation/Private/Helpers/PVImportHelpers.cpp
new file mode 100644
--- /dev/null
+++ b/Engine/Plugins/Experimental/ProceduralVegetationEditor/Source/ProceduralVegetation/Private/Helpers/PVImportHelpers.cpp
@@ -0,0 +1,1541 @@
+// Copyright Epic Games, Inc. All Rights Reserved.
+#include "Helpers/PVImportHelpers.h"
+
+#include "ProceduralVegetationPreset.h"
+#include "DynamicMesh/DynamicMesh3.h"
+#include "DynamicMesh/DynamicMeshAABBTree3.h"
+#include "GeometryCollectionToDynamicMesh.h"
+#include "GeometryCollection/GeometryCollection.h"
+
+#include "Helpers/PVAttributesHelper.h"
+
+#include "Engine/StaticMesh.h"
+
+namespace PV::DynamicMeshHelper
+{
+	static void ComputeLengthFromRoot(
+		const UE::Geometry::FDynamicMesh3& DynamicMesh,
+		const FBooleanVertexAttribute& RootPointsAttribute,
+		const FBooleanVertexAttribute& EndPointsAttribute,
+		FFloatVertexAttribute& LengthFromRootAttribute,
+		const TFunction<bool(int32)>& VertexFilterPredicate
+	)
+	{
+		TRACE_CPUPROFILER_EVENT_SCOPE(ComputeLengthFromRoot);
+
+		// This algorithm starts at the root point and walks along every vertex, keeping track of its distance to the root point and stores that distance
+		// in the LengthFromRootAttribute.
+
+		for (int32 VertexID : DynamicMesh.VertexIndicesItr())
+		{
+			LengthFromRootAttribute.SetValue(VertexID, -1);
+		}
+
+		TArray<int32> SearchQueue;
+		SearchQueue.Reserve(DynamicMesh.VertexCount());
+		for (int32 VertexID : RootPointsAttribute.FindAllNonZero())
+		{
+			SearchQueue.Add(VertexID);
+			LengthFromRootAttribute.SetValue(VertexID, 0);
+		}
+
+		while (SearchQueue.Num() > 0)
+		{
+			const int32& VertexID = SearchQueue.Pop(EAllowShrinking::No);
+
+			const float PrevLengthFromRoot = LengthFromRootAttribute.GetValue(VertexID);
+			const FVector PrevVertexPosition = DynamicMesh.GetVertex(VertexID);
+
+			const TArray<int32> Neighbours = GetVertexNeighbours(DynamicMesh, VertexID);
+			for (int32 NeighbourVertexID : Neighbours)
+			{
+				if (!EndPointsAttribute.GetValue(NeighbourVertexID) // Ingore vertex filter for end points as that means we've reached the end of the branch
+					&& !VertexFilterPredicate(NeighbourVertexID))
+				{
+					continue;
+				}
+
+				const FVector VertexPosition = DynamicMesh.GetVertex(NeighbourVertexID);
+				const float NewLengthFromRoot = PrevLengthFromRoot + (VertexPosition - PrevVertexPosition).Size();
+				const float ExistingLengthFromRoot = LengthFromRootAttribute.GetValue(NeighbourVertexID);
+
+				if (ExistingLengthFromRoot >= 0 
+					&& ExistingLengthFromRoot <= NewLengthFromRoot)
+				{
+					// If we find a shorter path then we should ideally remove vertices from the queue
+					// which are computing the distance using old values. 
+					// The current algorithm will still produce correct results, but may compute a large number
+					// of unecessary distances which will later be over-written when a closer path is found.
+					continue;
+				}
+
+				LengthFromRootAttribute.SetValue(NeighbourVertexID, NewLengthFromRoot);
+
+				SearchQueue.Insert(NeighbourVertexID, 0);
+			}
+		}
+	}
+
+	static void TraceBranches(
+		const UE::Geometry::FDynamicMesh3& DynamicMesh,
+		const FFloatVertexAttribute& LengthFromRootAttribute,
+		const FBooleanVertexAttribute& RootPointAttribute,
+		const FBooleanVertexAttribute& EndPointAttribute,
+		FIntVertexAttribute& NextBranchPointAttribute,
+		const TFunction<bool(int32)>& VertexFilterPredicate
+	)
+	{
+		using namespace PV::DynamicMeshHelper;
+
+		TRACE_CPUPROFILER_EVENT_SCOPE(TraceBranches);
+
+		for (int32 VertexID : DynamicMesh.VertexIndicesItr())
+		{
+			NextBranchPointAttribute.SetValue(VertexID, INDEX_NONE);
+		}
+
+		const auto FindNextPoint = [&](int32 PointIndex, float MinLengthFromRoot)->int32
+		{
+			int32 OutNextPointIndex = INDEX_NONE;
+			const TArray<int32> Neighbours = GetVertexNeighbours(DynamicMesh, PointIndex);
+			for (int32 Neighbour : Neighbours)
+			{
+				if (RootPointAttribute.GetValue(Neighbour))
+				{
+					OutNextPointIndex = Neighbour;
+					break;
+				}
+
+				if (!VertexFilterPredicate(Neighbour))
+				{
+					continue;
+				}
+
+				const float LengthFromRoot = LengthFromRootAttribute.GetValue(Neighbour);
+				check(LengthFromRoot >= 0);
+
+				if (LengthFromRoot >= MinLengthFromRoot)
+				{
+					continue;
+				}
+				MinLengthFromRoot = LengthFromRoot;
+				OutNextPointIndex = Neighbour;
+			}
+			return OutNextPointIndex;
+		};
+
+		const TArray<int32> EndPoints = EndPointAttribute.FindAllNonZero();
+
+		TArray<bool> bHasProcessedVertex;
+
+		for (int32 i = 0; i < EndPoints.Num(); i++)
+		{
+			int32 CurrentPointIndex = EndPoints[i];
+
+			float MinLengthFromRoot = LengthFromRootAttribute.GetValue(CurrentPointIndex);
+			if (MinLengthFromRoot < 0)
+			{
+				continue; // This end point was unable to trace down to a root point, perhaps log a warning?
+			}
+
+			bHasProcessedVertex.Reset();
+			bHasProcessedVertex.SetNumZeroed(DynamicMesh.MaxVertexID());
+
+			while (true)
+			{
+				const int32 NextPointIndex = FindNextPoint(CurrentPointIndex, MinLengthFromRoot);
+				if (NextPointIndex == INDEX_NONE)
+				{
+					break;
+				}
+				
+				NextBranchPointAttribute.SetValue(CurrentPointIndex, NextPointIndex);
+
+				if (NextBranchPointAttribute.GetValue(NextPointIndex) != INDEX_NONE
+					|| RootPointAttribute.GetValue(NextPointIndex))
+				{
+					break;
+				}
+
+				if (!ensureMsgf(!bHasProcessedVertex[NextPointIndex], TEXT("Infinite loop detected")))
+				{
+					break;
+				}
+
+				bHasProcessedVertex[NextPointIndex] = true;
+
+				CurrentPointIndex = NextPointIndex;
+				MinLengthFromRoot = LengthFromRootAttribute.GetValue(CurrentPointIndex);
+			}
+		}
+	}
+
+	static void ExtractBranchHierarchies(
+		const UE::Geometry::FDynamicMesh3& DynamicMesh,
+		const FBooleanVertexAttribute& RootPointAttribute,
+		const FBooleanVertexAttribute& EndPointAttribute,
+		const FIntVertexAttribute& NextBranchPointAttribute,
+		TArray<FMeshBranchHierarchy>& OutMeshBranchHierarchies
+	)
+	{
+		using namespace PV::DynamicMeshHelper;
+
+		TRACE_CPUPROFILER_EVENT_SCOPE(ExtractBranchHierarchies);
+
+		TArray<bool> bProcessedVertices;
+		bProcessedVertices.SetNumZeroed(DynamicMesh.MaxVertexID());
+
+		const auto FindBestParentPoint = [&](int32 PrevVertexID, int32 VertexID, const TArray<int32>& NextPoints)
+		{
+			if (NextPoints.Num() == 1)
+			{
+				return NextPoints[0];
+			}
+
+			check(NextPoints.Num() > 0);
+
+			const FVector3d& VertexPos = DynamicMesh.GetVertex(VertexID);
+			const FVector3d& PrevVertexPos = DynamicMesh.GetVertex(PrevVertexID);
+			const FVector3d ApicalDirection = (VertexPos - PrevVertexPos).GetSafeNormal();
+
+			float BestDot = -1;
+			int32 BestDotIndex = INDEX_NONE;
+			for (int32 NextVertexID : NextPoints)
+			{
+				const FVector3d& NextVertexPos = DynamicMesh.GetVertex(NextVertexID);
+				const FVector3d NextApicalDirection = (NextVertexPos - VertexPos).GetSafeNormal();
+				const float Dot = FVector3d::DotProduct(ApicalDirection, NextApicalDirection);
+				if (Dot > BestDot)
+				{
+					BestDot = Dot;
+					BestDotIndex = NextVertexID;
+				}
+			}
+
+			check(BestDotIndex != INDEX_NONE);
+			return BestDotIndex;
+		};
+
+		const auto FindNextBranchPoints = [&](int32 VertexID)
+		{
+			TArray<int32> NextBranchPoints;
+			const TArray<int32> Neighbours = GetVertexNeighbours(DynamicMesh, VertexID);
+
+			for (int32 Neighbour : Neighbours)
+			{
+				if (NextBranchPointAttribute.GetValue(Neighbour) == VertexID
+					&& !bProcessedVertices[Neighbour])
+				{
+					NextBranchPoints.Add(Neighbour);
+				}
+			}
+			return NextBranchPoints;
+		};
+
+		const TFunction<void(int32, TArray<FMeshBranchHierarchy::FBranch>&)> TraceBranchesRecursive = [&](
+			int32 VertexID, 
+			TArray<FMeshBranchHierarchy::FBranch>& OutBranches
+		)
+		{
+			const int32 BranchIndex = OutBranches.Num() - 1;
+			OutBranches[BranchIndex].BranchPoints.Add(VertexID);
+			
+			bProcessedVertices[VertexID] = true;
+
+			if (EndPointAttribute.GetValue(VertexID))
+			{
+				return;
+			}
+
+			const TArray<int32> NextBranchPoints = FindNextBranchPoints(VertexID);
+			check(NextBranchPoints.Num() > 0);
+
+			for (int32 NextVertexID : NextBranchPoints)
+			{
+				bProcessedVertices[NextVertexID] = true;
+			}
+
+			const int32 PrevVertexID = OutBranches[BranchIndex].BranchPoints.Last(1);
+			const int32 NextBranchPointsIndex = FindBestParentPoint(PrevVertexID, VertexID, NextBranchPoints);
+			
+			TraceBranchesRecursive(NextBranchPointsIndex, OutBranches);
+
+			for (int32 NextVertexID : NextBranchPoints)
+			{
+				if (NextVertexID != NextBranchPointsIndex)
+				{
+					FMeshBranchHierarchy::FBranch& NextBranch = OutBranches.AddDefaulted_GetRef();
+					NextBranch.BranchPoints.Add(VertexID);
+					NextBranch.ParentBranch = BranchIndex;
+
+					TraceBranchesRecursive(NextVertexID, OutBranches);
+				}
+			}
+		};
+
+		const TArray<int32> RootPoints = RootPointAttribute.FindAllNonZero();
+		OutMeshBranchHierarchies.Reserve(RootPoints.Num());
+
+		for (const int32 RootPointIndex : RootPoints)
+		{
+			TArray<TArray<FMeshBranchHierarchy::FBranch>> BranchHierarchies;
+
+			const TArray<int32> NextBranchPoints = FindNextBranchPoints(RootPointIndex);
+			for (int32 NextVertexID : NextBranchPoints)
+			{
+				bProcessedVertices[NextVertexID] = true;
+			}
+
+			for (int32 NextVertexID : NextBranchPoints)
+			{
+				TArray<FMeshBranchHierarchy::FBranch>& Branches = BranchHierarchies.AddDefaulted_GetRef();
+				FMeshBranchHierarchy::FBranch& RootBranch = Branches.AddDefaulted_GetRef();
+				RootBranch.BranchPoints.Add(RootPointIndex);
+				TraceBranchesRecursive(NextVertexID, Branches);
+			}
+
+			if (BranchHierarchies.Num() > 0)
+			{
+				// TODO: Handle multiple branches starting from root
+				TArray<FMeshBranchHierarchy::FBranch>* BestBranches = Algo::MaxElement(
+					BranchHierarchies, 
+					[](const auto& A, const auto& B) { return A.Num() < B.Num(); }
+				);
+				OutMeshBranchHierarchies.Emplace(MoveTemp(*BestBranches));
+			}
+		}
+	}
+
+	
+
+	const static TSet<FName> GeometryGroups = {
+		FGeometryCollection::VerticesGroup,
+		FGeometryCollection::FacesGroup,
+		FGeometryCollection::GeometryGroup,
+		FGeometryCollection::BreakingGroup,
+		FGeometryCollection::MaterialGroup,
+		FTransformCollection::TransformGroup,
+		FTransformCollection::ConvexGroup
+	};
+
+	const static FName MaterialPathAttributeName("MaterialPath");
+
+	static void AppendGeometryFromManagedArrayCollection(FGeometryCollection& InOutGeometryCollection, const FManagedArrayCollection& InManagedArrayCollection)
+	{
+		TArray<FName> GroupsToSkip;
+		GroupsToSkip.Reserve(InManagedArrayCollection.GroupNames().Num());
+
+		for (FName GroupName : InManagedArrayCollection.GroupNames())
+		{
+			if (!GeometryGroups.Contains(GroupName))
+			{
+				GroupsToSkip.Add(GroupName);
+			}
+		}
+
+		FGeometryCollection GeometryCollection;
+		InManagedArrayCollection.CopyTo(&GeometryCollection, GroupsToSkip);
+
+		const auto* ExistingMaterialPathsAttribute = InOutGeometryCollection.FindAttribute<FString>(MaterialPathAttributeName, FGeometryCollection::MaterialGroup);
+		const int32 NumExistingMaterials = ExistingMaterialPathsAttribute ? ExistingMaterialPathsAttribute->Num() : 1;
+		const bool bHasCopiedGeometry = InOutGeometryCollection.AppendGeometry(GeometryCollection, NumExistingMaterials, false) != INDEX_NONE;
+
+		// Copy materials. Ideally we'd want to re-use materials where possible, but given this is for debugging we just assign 
+		// one material per mesh.
+		if (bHasCopiedGeometry)
+		{
+			if (!InOutGeometryCollection.HasGroup(FGeometryCollection::MaterialGroup))
+			{
+				InOutGeometryCollection.AddGroup(FGeometryCollection::MaterialGroup);
+			}
+
+			auto* MaterialPathsAttributeToCopy = GeometryCollection.FindAttribute<FString>(MaterialPathAttributeName, FGeometryCollection::MaterialGroup);
+			const int32 NumMaterialsToCopy = MaterialPathsAttributeToCopy ? MaterialPathsAttributeToCopy->Num() : 1;
+			InOutGeometryCollection.Resize(NumExistingMaterials + NumMaterialsToCopy, FGeometryCollection::MaterialGroup);
+
+			if (MaterialPathsAttributeToCopy)
+			{
+				auto* MaterialPathsAttribute = InOutGeometryCollection.FindOrAddAttributeTyped<FString>(MaterialPathAttributeName, FGeometryCollection::MaterialGroup);
+				for (int32 i = NumExistingMaterials; i < NumExistingMaterials + NumMaterialsToCopy; ++i)
+				{
+					const int32 MaterialIndexToCopy = i - NumExistingMaterials;
+					if (MaterialPathsAttributeToCopy->IsValidIndex(MaterialIndexToCopy))
+					{
+						(*MaterialPathsAttribute)[i] = (*MaterialPathsAttributeToCopy)[MaterialIndexToCopy];
+					}
+				}
+			}
+		}
+	}
+
+	static void SetManagedArrayCollectionGeometry(const FGeometryCollection& InGeometryCollection, FManagedArrayCollection& InOutManagedArrayCollection)
+	{
+		for (const FName GroupName : GeometryGroups)
+		{
+			InOutManagedArrayCollection.RemoveGroup(GroupName);
+		}
+
+		InGeometryCollection.CopyTo(&InOutManagedArrayCollection);
+	}
+}
+
+const PV::TDynamicMeshVertexAttributeDefinition<float> PV::DynamicMeshHelper::VertexScaleAttributeDefinition(TEXT("VertexScale"));
+const PV::TDynamicMeshVertexAttributeDefinition<bool> PV::DynamicMeshHelper::EndPointAttributeDefinition(TEXT("EndPoint"));
+const PV::TDynamicMeshVertexAttributeDefinition<bool> PV::DynamicMeshHelper::RootPointAttributeDefinition(TEXT("RootPoint"));
+const PV::TDynamicMeshVertexAttributeDefinition<float> PV::DynamicMeshHelper::LengthFromRootAttributeDefinition(TEXT("LengthFromRoot"));
+const PV::TDynamicMeshVertexAttributeDefinition<int32> PV::DynamicMeshHelper::NextBranchPointAttributeDefinition(TEXT("NextBranchPoint"));
+
+void PV::DynamicMeshHelper::AddDynamicMeshToCollection(const UE::Geometry::FDynamicMesh3& DynamicMesh, FManagedArrayCollection& InOutCollection, const FTransform& Offset)
+{
+	FGeometryCollection MeshGeomCollection;
+
+	UE::Geometry::FGeometryCollectionToDynamicMeshes::FToCollectionOptions Options2;
+	UE::Geometry::FGeometryCollectionToDynamicMeshes::AppendMeshToCollection(MeshGeomCollection, DynamicMesh, Offset, Options2);
+
+	AppendGeometryFromManagedArrayCollection(MeshGeomCollection, InOutCollection);
+	SetManagedArrayCollectionGeometry(MeshGeomCollection, InOutCollection);
+}
+
+void PV::DynamicMeshHelper::AddBranchHierarchiesToCollection(const TArray<FPVBranchHierarchyDescription>& BranchHierarchyDescriptions, FManagedArrayCollection& OutCollection)
+{
+	const static FName CurveAttributeName = TEXT("Curve");
+	const static FName BranchCurvesGroupName = TEXT("BranchCurves");
+	OutCollection.AddGroup(BranchCurvesGroupName);
+	int32 AttributeIndex = 0;
+	for (const FPVBranchHierarchyDescription& BranchHierarchy : BranchHierarchyDescriptions)
+	{
+		OutCollection.AddElements(BranchHierarchy.Branches.Num(), BranchCurvesGroupName);
+
+		TManagedArray<TArray<FVector3f>>& CurveAttribute = *OutCollection.FindOrAddAttributeTyped<TArray<FVector3f>>(CurveAttributeName, BranchCurvesGroupName);
+		for (int32 i = 0; i < BranchHierarchy.Branches.Num(); ++i)
+		{
+			const FPVBranchDescription& Branch = BranchHierarchy.Branches[i];
+			CurveAttribute[AttributeIndex].Reserve(Branch.PointIndices.Num());
+			for (int32 VertexIndex : Branch.PointIndices)
+			{
+				CurveAttribute[AttributeIndex].Add(BranchHierarchy.Points[VertexIndex]);
+			}
+			AttributeIndex++;
+		}
+	}
+}
+
+TArray<int32> PV::DynamicMeshHelper::GetUniqueVertexIndices(const UE::Geometry::FDynamicMesh3& DynamicMesh, const TArray<int32>& TriangleIds)
+{
+	TSet<int32> UniqueVertexIndices;
+	UniqueVertexIndices.Reserve(TriangleIds.Num() * 3);
+	for (int32 tid : TriangleIds)
+	{
+		const UE::Geometry::FIndex3i VertexIndices = DynamicMesh.GetTriangle(tid);
+		UniqueVertexIndices.Add(VertexIndices[0]);
+		UniqueVertexIndices.Add(VertexIndices[1]);
+		UniqueVertexIndices.Add(VertexIndices[2]);
+	}
+
+	return UniqueVertexIndices.Array();
+}
+
+TArray<int32> PV::DynamicMeshHelper::GetVertexNeighbours(const UE::Geometry::FDynamicMesh3& DynamicMesh, int32 VertexID)
+{
+	TArray<int32> OutNeighbours;
+	OutNeighbours.Reserve(DynamicMesh.GetVtxEdgeCount(VertexID));
+	for (int32 EdgeID : DynamicMesh.VtxEdgesItr(VertexID))
+	{
+		const UE::Geometry::FIndex2i VertexIDs = DynamicMesh.GetEdgeV(EdgeID);
+		const int32 OtherVertexID = VertexIDs.A == VertexID ? VertexIDs.B : VertexIDs.A;
+		OutNeighbours.Add(OtherVertexID);
+	}
+	return OutNeighbours;
+}
+
+TArray<int32> PV::DynamicMeshHelper::FindVerticesInRadius(const UE::Geometry::FDynamicMeshAABBTree3& DynamicMeshAABBTree, const FVector& Pos, float Radius)
+{
+	TArray<int32> OutVertexIndices;
+
+	const UE::Geometry::FDynamicMesh3* DynamicMesh = DynamicMeshAABBTree.GetMesh();
+	if (DynamicMesh == nullptr)
+	{
+		return OutVertexIndices;
+	}
+
+	OutVertexIndices.Reserve(DynamicMesh->VertexCount());
+
+	const double RadiusSqr = Radius * Radius;
+
+	UE::Geometry::FDynamicMeshAABBTree3::FTreeTraversal TreeTraversal;
+
+	TreeTraversal.NextBoxF = [&](const UE::Geometry::FAxisAlignedBox3d& Box, int32 Depth)->bool
+	{
+		return Box.DistanceSquared(Pos) <= RadiusSqr;
+	};
+
+	TreeTraversal.NextTriangleF = [&](int32 TID)
+	{
+		const UE::Geometry::FIndex3i Triangle = DynamicMesh->GetTriangle(TID);
+		for (int32 j = 0; j < 3; ++j)
+		{
+			const double VertexDistSqr = (DynamicMesh->GetVertex(Triangle[j]) - Pos).SizeSquared();
+			if (VertexDistSqr <= RadiusSqr)
+			{
+				OutVertexIndices.AddUnique(Triangle[j]);
+			}
+		}
+	};
+
+	DynamicMeshAABBTree.DoTraversal(TreeTraversal);
+
+	return OutVertexIndices;
+}
+
+void PV::DynamicMeshHelper::ComputeDynamicMeshBranchHierarchies(
+	const UE::Geometry::FDynamicMesh3& DynamicMesh,
+	const FBooleanVertexAttribute& RootPointsAttribute,
+	const FBooleanVertexAttribute& EndPointsAttribute,
+	const FFloatVertexAttribute& VertexScaleAttribute,
+	FFloatVertexAttribute& LengthFromRootAttribute,
+	FIntVertexAttribute& NextBranchPointAttribute,
+	const TFunction<bool(int32)>& VertexFilterPredicate,
+	TArray<FMeshBranchHierarchy>& OutDynamicMeshBranchHierarchies
+)
+{
+	ComputeLengthFromRoot(
+		DynamicMesh,
+		RootPointsAttribute,
+		EndPointsAttribute,
+		LengthFromRootAttribute,
+		VertexFilterPredicate
+	);
+
+	TraceBranches(
+		DynamicMesh,
+		LengthFromRootAttribute,
+		RootPointsAttribute,
+		EndPointsAttribute,
+		NextBranchPointAttribute,
+		VertexFilterPredicate
+	);
+
+	ExtractBranchHierarchies(
+		DynamicMesh,
+		RootPointsAttribute,
+		EndPointsAttribute,
+		NextBranchPointAttribute,
+		OutDynamicMeshBranchHierarchies
+	);
+}
+
+void PV::DynamicMeshHelper::ConvertDynamicMeshBranchHierarchiesToBranchDescription(
+	const UE::Geometry::FDynamicMesh3& DynamicMesh,
+	const FFloatVertexAttribute& VertexScaleAttribute,
+	const FBooleanVertexAttribute& EndPointAttribute,
+	const TArray<FMeshBranchHierarchy>& InMeshBranchHierarchies,
+	TArray<FPVBranchHierarchyDescription>& OutBranchHierarchyDescriptions
+)
+{
+	TRACE_CPUPROFILER_EVENT_SCOPE(ConvertMeshBranchHierarchiesToBranchDescription);
+
+	for (const FMeshBranchHierarchy& BranchHierarchy : InMeshBranchHierarchies)
+	{
+		int32 MaxNrOfPoints = 0;
+		for (const auto& Branch : BranchHierarchy.Branches)
+		{
+			MaxNrOfPoints += Branch.BranchPoints.Num();
+		}
+
+		FPVBranchHierarchyDescription& OutBranchHierarchy = OutBranchHierarchyDescriptions.AddDefaulted_GetRef();
+		OutBranchHierarchy.Branches.Reserve(BranchHierarchy.Branches.Num());
+		OutBranchHierarchy.Points.Reserve(MaxNrOfPoints);
+		OutBranchHierarchy.PointsRadii.Reserve(MaxNrOfPoints);
+
+		TMap<int32, int32> VertexIndexToPointIndex;
+		VertexIndexToPointIndex.Reserve(MaxNrOfPoints);
+
+		for (const FMeshBranchHierarchy::FBranch& Branch : BranchHierarchy.Branches)
+		{
+			FPVBranchDescription& OutBranchDescription = OutBranchHierarchy.Branches.AddDefaulted_GetRef();
+			OutBranchDescription.ParentBranchIndex = Branch.ParentBranch;
+			OutBranchDescription.BranchIndex = OutBranchHierarchy.Branches.Num() - 1;
+			OutBranchDescription.PointIndices.Reserve(Branch.BranchPoints.Num());
+
+			for (int32 i = 0; i < Branch.BranchPoints.Num(); ++i)
+			{
+				const int32 VertexIndex = Branch.BranchPoints[i];
+				int32 PointIndex = INDEX_NONE;
+
+				int32* ExistingPointIndex = VertexIndexToPointIndex.Find(VertexIndex);
+				if (!ExistingPointIndex)
+				{
+					PointIndex = OutBranchHierarchy.Points.Add(FVector3f(DynamicMesh.GetVertex(VertexIndex)));
+
+					const bool bIsEndPoint = EndPointAttribute.GetValue(VertexIndex);
+					if (bIsEndPoint)
+					{
+						OutBranchHierarchy.PointsRadii.Add(0.01f);
+					}
+					else
+					{
+						const float VertexRadius = VertexScaleAttribute.GetValue(VertexIndex);
+						OutBranchHierarchy.PointsRadii.Add(VertexRadius);
+					}
+
+					VertexIndexToPointIndex.Add(VertexIndex, PointIndex);
+				}
+				else
+				{
+					PointIndex = *ExistingPointIndex;
+				}
+
+				OutBranchDescription.PointIndices.Add(PointIndex);
+			}
+		}
+
+		for (int32 i = 0; i < OutBranchHierarchy.Branches.Num(); ++i)
+		{
+			const FPVBranchDescription& BranchDescription = OutBranchHierarchy.Branches[i];
+			if (BranchDescription.ParentBranchIndex != INDEX_NONE)
+			{
+				OutBranchHierarchy.Branches[BranchDescription.ParentBranchIndex].ChildBranchIndices.Add(i);
+			}
+			else
+			{
+				check(OutBranchHierarchy.RootBranchIndex == INDEX_NONE);
+				OutBranchHierarchy.RootBranchIndex = i;
+			}
+		}
+	}
+}
+
+
+FPVBranchHierarchyDescription::EValidateHierarchyResult FPVBranchHierarchyDescription::ValidateHierarchy() const
+{
+	if (!Branches.IsValidIndex(RootBranchIndex))
+	{
+		return EValidateHierarchyResult::InvalidRootBranchIndex;
+	}
+	
+	if (Branches[RootBranchIndex].ParentBranchIndex != INDEX_NONE)
+	{
+		return EValidateHierarchyResult::InvalidRootBranchParentIndex;
+	}
+	
+	for (int32 BranchIndex = 0; BranchIndex < Branches.Num(); BranchIndex++)
+	{
+		const FPVBranchDescription& Branch = Branches[BranchIndex];
+		if (Branch.BranchIndex != BranchIndex)
+		{
+			return EValidateHierarchyResult::BranchIndexMismatch;
+		}
+		
+		if (Branch.ParentBranchIndex != INDEX_NONE
+			&& !Branches.IsValidIndex(Branch.ParentBranchIndex))
+		{
+			return EValidateHierarchyResult::InvalidParentBranchIndex;
+		}
+		
+		if (Branch.ParentBranchIndex == INDEX_NONE
+			&& RootBranchIndex != Branch.BranchIndex)
+		{
+			return EValidateHierarchyResult::InvalidRootBranchIndex;
+		}
+		
+		for (int32 ChildBranchIndex : Branch.ChildBranchIndices)
+		{
+			if (!Branches.IsValidIndex(ChildBranchIndex))
+			{
+				return EValidateHierarchyResult::InvalidChildBranchIndex;
+			}
+		}
+		
+		for (int32 PointIndex : Branch.PointIndices)
+		{
+			if (!Points.IsValidIndex(PointIndex))
+			{
+				return EValidateHierarchyResult::InvalidPointIndex; 
+			}
+		}
+	}
+
+	TSet<int32> FoundBranches;
+	FoundBranches.Reserve(Branches.Num());
+
+	const TFunction<bool(int32)> WalkHierarchyRecursive = [&](int32 BranchIndex)->bool
+	{
+		bool bAlreadyAdded = false;
+		FoundBranches.Add(BranchIndex, &bAlreadyAdded);
+
+		if (bAlreadyAdded)
+		{
+			return false;
+		}
+
+		for (int32 ChildIndex : Branches[BranchIndex].ChildBranchIndices)
+		{
+			if (!WalkHierarchyRecursive(ChildIndex))
+			{
+				return false;
+			}
+		}
+
+		return true;
+	};
+
+	if (!WalkHierarchyRecursive(RootBranchIndex))
+	{
+		return EValidateHierarchyResult::CircularHierarchy;
+	}
+
+	if (FoundBranches.Num() != Branches.Num())
+	{
+		return EValidateHierarchyResult::DisjointedHierarchy;
+	}
+
+	return EValidateHierarchyResult::Success;
+}
+
+void PV::ImportHelper::CreateEmptyGrowthData(FManagedArrayCollection& OutCollection, int32 NumBranchPoints, int32 NumFoliagePoints, int32 NumBranches, int32 NumFoliageMeshes)
+{
+	// PointGroup
+	OutCollection.AddGroup(PV::GroupNames::PointGroup);
+	OutCollection.AddElements(NumBranchPoints, PV::GroupNames::PointGroup);
+	FPointPositionAttribute::AddAttribute(OutCollection, FVector3f::ZeroVector);
+	FPointLengthFromRootAttribute::AddAttribute(OutCollection, 0);
+	FPointLengthFromSeedAttribute::AddAttribute(OutCollection, 0);
+	FPointScaleGradientAttribute::AddAttribute(OutCollection, 0);
+	FPointHullGradientAttribute::AddAttribute(OutCollection, 0);
+	FPointMainTrunkGradientAttribute::AddAttribute(OutCollection, 0);
+	FPointGroundGradientAttribute::AddAttribute(OutCollection, 0);
+	FPointScaleAttribute::AddAttribute(OutCollection, 0.f);
+	FBudDirectionAttribute::AddAttribute(OutCollection);
+	FBudStatusAttribute::AddAttribute(OutCollection);
+	FBudLateralMeristemAttribute::AddAttribute(OutCollection);
+	FPointBranchGradientAttribute::AddAttribute(OutCollection, 0);
+	FBudHormoneLevelsAttribute::AddAttribute(OutCollection);
+	FPointPlantGradientAttribute::AddAttribute(OutCollection, 0);
+	FPointBudNumberAttribute::AddAttribute(OutCollection, 0);
+	FPointNjordPixelIndexAttribute::AddAttribute(OutCollection, 0);
+	FBudLightDetectedAttribute::AddAttribute(OutCollection);
+	FBudDevelopmentAttribute::AddAttribute(OutCollection);
+
+	// BranchGroup
+	OutCollection.AddGroup(PV::GroupNames::BranchGroup);
+	OutCollection.AddElements(NumBranches, PV::GroupNames::BranchGroup);
+	FBranchParentsAttribute::AddAttribute(OutCollection);
+	FBranchChildrenAttribute::AddAttribute(OutCollection);
+	FBranchPointsAttribute::AddAttribute(OutCollection);
+	FBranchNumberAttribute::AddAttribute(OutCollection, 0);
+	FBranchSourceBudNumberAttribute::AddAttribute(OutCollection, 0);
+	FBranchHierarchyNumberAttribute::AddAttribute(OutCollection, 0);
+	FBranchParentNumberAttribute::AddAttribute(OutCollection, 0);
+	FBranchPlantNumberAttribute::AddAttribute(OutCollection, 0);
+	FBranchFoliageIDsAttribute::AddAttribute(OutCollection);
+
+	// FoliageNamesGroup
+	OutCollection.AddGroup(PV::GroupNames::FoliageNamesGroup);
+	OutCollection.AddElements(NumFoliageMeshes, PV::GroupNames::FoliageNamesGroup);
+	FFoliageNameAttribute::AddAttribute(OutCollection, TEXT(""));
+	
+	// FoliageGroup
+	OutCollection.AddGroup(PV::GroupNames::FoliageGroup);
+	OutCollection.AddElements(NumFoliagePoints, PV::GroupNames::FoliageGroup);
+	FFoliageNameIDAttribute::AddAttribute(OutCollection, 0);
+	FFoliageBranchIDAttribute::AddAttribute(OutCollection, 0);
+	FFoliagePivotPointAttribute::AddAttribute(OutCollection, FVector3f::ZeroVector);
+	FFoliageUpVectorAttribute::AddAttribute(OutCollection, FVector3f::ZeroVector);
+	FFoliageNormalVectorAttribute::AddAttribute(OutCollection, FVector3f::ZeroVector);
+	FFoliageScaleAttribute::AddAttribute(OutCollection, 0);
+	FFoliageLengthFromRootAttribute::AddAttribute(OutCollection, 0);
+	FFoliageParentBoneIDAttribute::AddAttribute(OutCollection, 0);
+
+	// DetailsGroup
+	OutCollection.AddGroup(PV::GroupNames::DetailsGroup);
+	OutCollection.AddElements(1, PV::GroupNames::DetailsGroup);
+	FDetailFoliagePathAttribute::AddAttribute(OutCollection, TEXT(""));
+	FDetailTrunkMaterialPathAttribute::AddAttribute(OutCollection, TEXT(""));
+	FDetailTrunkURangeAttribute::AddAttribute(OutCollection, { });
+	FDetailLeafPhyllotaxyAttribute::AddAttribute(OutCollection, { 0, 200, 50, 1, 1, 0, 0, 0 });
+}
+
+void PV::ImportHelper::FillAttributesFromBranchHierarchy(
+	FManagedArrayCollection& OutCollection,
+	TArrayView<const FPVBranchHierarchyDescription> InBranchHierarchies
+)
+{
+	FillAttributesFromBranchHierarchy(
+		FPointPositionAttribute::GetAttribute(OutCollection),
+		FPointScaleAttribute::GetAttribute(OutCollection),
+		FBranchPointsAttribute::GetAttribute(OutCollection),
+		FBranchParentsAttribute::GetAttribute(OutCollection),
+		FBranchChildrenAttribute::GetAttribute(OutCollection),
+		FBranchParentNumberAttribute::GetAttribute(OutCollection),
+		FBranchNumberAttribute::GetAttribute(OutCollection),
+		FBranchHierarchyNumberAttribute::GetAttribute(OutCollection),
+		FBranchPlantNumberAttribute::GetAttribute(OutCollection),
+		InBranchHierarchies
+	);
+}
+
+void PV::ImportHelper::FillAttributesFromBranchHierarchy(
+	PV::FPointPositionAttributeView PointPositionAttribute,
+	PV::FPointScaleAttributeView PointScaleAttribute,
+	PV::FBranchPointsAttributeView BranchPointsAttribute,
+	PV::FBranchParentsAttributeView BranchParentsAttribute,
+	PV::FBranchChildrenAttributeView BranchChildrenAttribute,
+	PV::FBranchParentNumberAttributeView BranchParentNumberAttribute,
+	PV::FBranchNumberAttributeView BranchNumberAttribute,
+	PV::FBranchHierarchyNumberAttributeView BranchHierarchyNumberAttribute,
+	PV::FBranchPlantNumberAttributeView BranchPlantNumberAttribute,
+	TArrayView<const FPVBranchHierarchyDescription> InBranchHierarchies
+)
+{
+	const static auto BranchIndexToBranchNumber = [](int32 BranchIndex) 
+	{ 
+		// Branch Number is BranchIndex + 1 since branch with parent number 0 is a root in PVE and our root has parent index -1.
+		return BranchIndex + 1; 
+	}; 
+
+	using FFillBranchesRecursive = TFunction<void(const FPVBranchHierarchyDescription&, int32, int32, int32, int32, int32)>;
+
+	const FFillBranchesRecursive FillBranchesRecursive = [&](
+		const FPVBranchHierarchyDescription& BranchHierarchy, 
+		int32 PlantNumber,
+		int32 BranchIndex,
+		int32 AttributeBranchIndexOffset,
+		int32 AttributePointIndexOffset,
+		int32 Depth
+	)
+	{
+		const FPVBranchDescription& Branch = BranchHierarchy.Branches[BranchIndex];
+
+		const int32 AttributeRootBranchIndex = BranchHierarchy.RootBranchIndex + AttributeBranchIndexOffset;
+		const int32 AttributeBranchIndex = BranchIndex + AttributeBranchIndexOffset;
+		const int32 AttributeParentBranchIndex = Branch.ParentBranchIndex != INDEX_NONE
+			? Branch.ParentBranchIndex + AttributeBranchIndexOffset
+			: INDEX_NONE;
+
+		TArray<int32> AttributeBranchPoints;
+		AttributeBranchPoints.Reserve(Branch.PointIndices.Num());
+
+		for (int32 i = 0; i < Branch.PointIndices.Num(); ++i)
+		{
+			const int32 AttributePointIndex = Branch.PointIndices[i] + AttributePointIndexOffset;
+			AttributeBranchPoints.Add(AttributePointIndex);
+		}
+		BranchPointsAttribute[AttributeBranchIndex] = MoveTemp(AttributeBranchPoints);
+
+		if (AttributeParentBranchIndex != INDEX_NONE)
+		{
+			TArray<int32> Parents = BranchParentsAttribute[AttributeParentBranchIndex];
+			Parents.Add(BranchIndexToBranchNumber(AttributeParentBranchIndex));
+			BranchParentsAttribute[AttributeBranchIndex] = MoveTemp(Parents);
+		}
+
+		TArray<int32> AttributeBranchChildren;
+		AttributeBranchChildren.SetNumUninitialized(Branch.ChildBranchIndices.Num());
+		for (int32 i = 0; i < Branch.ChildBranchIndices.Num(); ++i)
+		{
+			const int32 AttributeChildBranchIndex = Branch.ChildBranchIndices[i] + AttributeBranchIndexOffset;
+
+			check(AttributeChildBranchIndex != AttributeRootBranchIndex);
+			AttributeBranchChildren[i] = BranchIndexToBranchNumber(AttributeChildBranchIndex);
+		}
+
+		BranchChildrenAttribute[AttributeBranchIndex] = MoveTemp(AttributeBranchChildren);
+
+		BranchParentNumberAttribute[AttributeBranchIndex] = BranchIndexToBranchNumber(AttributeParentBranchIndex);
+		BranchNumberAttribute[AttributeBranchIndex] = BranchIndexToBranchNumber(AttributeBranchIndex);
+
+		BranchHierarchyNumberAttribute[AttributeBranchIndex] = Depth;
+		BranchPlantNumberAttribute[AttributeBranchIndex] = PlantNumber;
+
+		for (int32 ChildIndex : Branch.ChildBranchIndices)
+		{
+			FillBranchesRecursive(
+				BranchHierarchy, 
+				PlantNumber,
+				ChildIndex, 
+				AttributeBranchIndexOffset,
+				AttributePointIndexOffset,
+				Depth + 1
+			);
+		}
+	};
+
+	int32 AttributeBranchIndexOffset = 0;
+	int32 AttributePointIndexOffset = 0;
+	int32 PlantNumber = 0;
+	for (const FPVBranchHierarchyDescription& BranchHierarchy : InBranchHierarchies)
+	{
+		FillBranchesRecursive(
+			BranchHierarchy, 
+			PlantNumber,
+			BranchHierarchy.RootBranchIndex, 
+			AttributeBranchIndexOffset,
+			AttributePointIndexOffset,
+			0
+		);
+
+		for (int32 i = 0; i < BranchHierarchy.Points.Num(); ++i)
+		{
+			PointPositionAttribute[i + AttributePointIndexOffset] = BranchHierarchy.Points[i];
+			PointScaleAttribute[i + AttributePointIndexOffset] = BranchHierarchy.PointsRadii[i] > 0 ? BranchHierarchy.PointsRadii[i] : 0;
+		}
+
+		AttributeBranchIndexOffset += BranchHierarchy.Branches.Num();
+		AttributePointIndexOffset += BranchHierarchy.Points.Num();
+		PlantNumber += 1;
+	}
+}
+
+void PV::ImportHelper::ConnectDisconnectedBranches(FPVBranchHierarchyDescription& InOutBranchHierarchy)
+{
+	struct FNewBranchRootIndex
+	{
+		int32 BranchIndex;
+		int32 NewRootPointIndex;
+	};
+	TArray<FNewBranchRootIndex> BranchRootsToFuse;
+	BranchRootsToFuse.Reserve(InOutBranchHierarchy.Branches.Num()); // We expect to fuse at most one point per branch
+
+	TArray<FNewBranchRootIndex> BranchRootsToConnect;
+	BranchRootsToConnect.Reserve(InOutBranchHierarchy.Branches.Num()); // We expect to connect at most one point per branch
+
+	TArray<float> ApicalDots;
+	TArray<float> PointDistances;
+
+	for (const FPVBranchDescription& Branch : InOutBranchHierarchy.Branches)
+	{
+		if (Branch.ParentBranchIndex == INDEX_NONE)
+		{
+			continue;
+		}
+
+		check(Branch.PointIndices.Num() > 1);
+		
+		const FPVBranchDescription& ParentBranch = InOutBranchHierarchy.Branches[Branch.ParentBranchIndex];
+		
+		const int32 FirstPointIndex = Branch.PointIndices[0];
+		const FVector3f& FirstPointLocation = InOutBranchHierarchy.Points[FirstPointIndex];
+
+		const int32 LastParentPointIndex = ParentBranch.PointIndices.Last();
+		const FVector3f& LastParentPointLocation = InOutBranchHierarchy.Points[LastParentPointIndex];
+
+		const float ParentPointDistSquared = (FirstPointLocation - LastParentPointLocation).SizeSquared();
+
+		constexpr float MaxFuseDistanceSquared = 1.f; // 1cm
+		if (ParentPointDistSquared <= MaxFuseDistanceSquared)
+		{
+			BranchRootsToFuse.Add({
+				.BranchIndex = Branch.BranchIndex,
+				.NewRootPointIndex = LastParentPointIndex,
+			});
+			continue;
+		}
+
+		const int32 SecondPointIndex = Branch.PointIndices[1];
+		const FVector3f& SecondPointLocation = InOutBranchHierarchy.Points[SecondPointIndex];
+		const FVector3f RootApicalDirection = (SecondPointLocation - FirstPointLocation).GetSafeNormal();
+
+		ApicalDots.SetNumUninitialized(ParentBranch.PointIndices.Num(), EAllowShrinking::No);
+		PointDistances.SetNumUninitialized(ParentBranch.PointIndices.Num(), EAllowShrinking::No);
+
+		int32 ClosestPointIndex = INDEX_NONE;
+		float ClosestPointDistance = MAX_flt;
+		for (int32 i = 0; i < ParentBranch.PointIndices.Num(); ++i)
+		{
+			const int32 ParentPointIndex = ParentBranch.PointIndices[i];
+			const FVector3f& ParentPointLocation = InOutBranchHierarchy.Points[ParentPointIndex];
+
+			const FVector3f DeltaLocation = (FirstPointLocation - ParentPointLocation);
+				
+			FVector3f ApicalDirection = FVector3f::ZeroVector;
+			float PointDist = 0;
+
+			const float PointDistSquared = DeltaLocation.SizeSquared();
+
+			if (PointDistSquared == 1.f)
+			{
+				ApicalDirection = DeltaLocation;
+			}
+			else if (PointDistSquared < UE_SMALL_NUMBER)
+			{
+				ApicalDirection = FVector3f::ZeroVector;
+			}
+			else
+			{
+				const float InvPointDist = FMath::InvSqrt(PointDistSquared);
+				ApicalDirection = FVector3f(DeltaLocation.X * InvPointDist, DeltaLocation.Y * InvPointDist, DeltaLocation.Z * InvPointDist);
+				PointDist = PointDistSquared * InvPointDist;
+			}
+
+			const float ApicalDotProduct = RootApicalDirection.Dot(ApicalDirection);
+			ApicalDots[i] = ApicalDotProduct;
+			PointDistances[i] = PointDist;
+
+			if (PointDist < ClosestPointDistance)
+			{
+				ClosestPointIndex = i;
+				ClosestPointDistance = PointDist;
+			}
+		}
+
+		int32 ClosestApicalIndex = ClosestPointIndex;
+		for (int32 i = 0; i < ParentBranch.PointIndices.Num(); ++i)
+		{
+			const float PercentageDiff = FMath::IsNearlyZero(PointDistances[i])
+				? PointDistances[ClosestPointIndex]
+				: FMath::Abs(PointDistances[ClosestPointIndex] - PointDistances[i]) / PointDistances[i];
+
+			if (PercentageDiff < 0.5)
+			{
+				if (ApicalDots[i] > ApicalDots[ClosestApicalIndex])
+				{
+					ClosestApicalIndex = i;
+				}
+			}
+		}
+
+		BranchRootsToConnect.Add({
+			.BranchIndex = Branch.BranchIndex,
+			.NewRootPointIndex = ParentBranch.PointIndices[ClosestApicalIndex],
+		});
+	}
+
+	// Connect disconnected branches
+	for (const FNewBranchRootIndex& NewBranchRootIndex : BranchRootsToConnect)
+	{
+		InOutBranchHierarchy.Branches[NewBranchRootIndex.BranchIndex].PointIndices.Insert(NewBranchRootIndex.NewRootPointIndex, 0);
+	}
+
+	// Fuse branch roots
+	{
+		TArray<int32> PointsToRemove;
+		PointsToRemove.Reserve(BranchRootsToFuse.Num());
+		for (const FNewBranchRootIndex& NewBranchRootIndex : BranchRootsToFuse)
+		{
+			const int32 OldBranchRootPointIndex = InOutBranchHierarchy.Branches[NewBranchRootIndex.BranchIndex].PointIndices[0];
+			PointsToRemove.Add(OldBranchRootPointIndex);
+
+			InOutBranchHierarchy.Branches[NewBranchRootIndex.BranchIndex].PointIndices[0] = NewBranchRootIndex.NewRootPointIndex;
+		}
+
+		struct FPointRef
+		{
+			int32 BranchIndex;
+			int32 PointIndicesIndex;
+		};
+		TArray<TArray<FPointRef>> PointsRefs;
+		PointsRefs.SetNum(InOutBranchHierarchy.Points.Num());
+
+		for (int32 i = 0; i < InOutBranchHierarchy.Branches.Num(); ++i)
+		{
+			const TArray<int32>& PointIndices = InOutBranchHierarchy.Branches[i].PointIndices;
+			for (int32 j = 0; j < PointIndices.Num(); ++j)
+			{
+				const int32 PointIndex = PointIndices[j];
+				PointsRefs[PointIndex].Add({ i, j });
+			}
+		}
+
+		PointsToRemove.Sort(); // Remove points from back to front to avoid having to continuously re-map the PointsToRemove array
+		for (int32 i = PointsToRemove.Num() - 1; i >= 0; --i)
+		{
+			const int32 PointIndex = PointsToRemove[i];
+
+			check(PointsRefs[PointIndex].Num() == 0);
+			const int32 SwappedPointIndex = InOutBranchHierarchy.Points.Num() - 1;
+			InOutBranchHierarchy.Points.RemoveAtSwap(PointIndex);
+			if (InOutBranchHierarchy.PointsRadii.IsValidIndex(PointIndex))
+			{
+				InOutBranchHierarchy.PointsRadii.RemoveAtSwap(PointIndex);
+			}
+
+			for (const FPointRef& PointRef : PointsRefs[SwappedPointIndex])
+			{
+				InOutBranchHierarchy.Branches[PointRef.BranchIndex].PointIndices[PointRef.PointIndicesIndex] = PointIndex;
+			}
+
+			PointsRefs[PointIndex] = MoveTemp(PointsRefs[SwappedPointIndex]);
+		}
+	}
+}
+
+void PV::ImportHelper::FuseCollinearBranches(FPVBranchHierarchyDescription& InOutBranchHierarchy)
+{
+	TArray<int32> BranchesToFuseWithParent;
+	BranchesToFuseWithParent.Reserve(InOutBranchHierarchy.Branches.Num());
+
+	const TFunction<void(int32)> SearchForBranchesToFuse_Recursive = [&](int32 BranchIndex)
+	{
+		const FPVBranchDescription& Branch = InOutBranchHierarchy.Branches[BranchIndex];
+
+		if (Branch.ChildBranchIndices.Num() == 0)
+		{
+			return;
+		}
+
+		const int32 LastParentPointIndex = Branch.PointIndices[Branch.PointIndices.Num() - 1];
+		const int32 SecondToLastParentPointIndex = Branch.PointIndices[Branch.PointIndices.Num() - 2];
+
+		const FVector3f& LastParentPointLocation = InOutBranchHierarchy.Points[LastParentPointIndex];
+		const FVector3f& SecondToLastParentPointLocation = InOutBranchHierarchy.Points[SecondToLastParentPointIndex];
+		const FVector3f ParentApicalDirection = (LastParentPointLocation - SecondToLastParentPointLocation).GetSafeNormal();
+
+		int32 ChildBranchToFuse = INDEX_NONE;
+		float BestApicalDotProduct = -1;
+		for (int32 ChildBranchIndex : Branch.ChildBranchIndices)
+		{
+			const FPVBranchDescription& ChildBranch = InOutBranchHierarchy.Branches[ChildBranchIndex];
+
+			const int32 FistChildBranchPointIndex = ChildBranch.PointIndices[0];
+			if (FistChildBranchPointIndex != LastParentPointIndex)
+			{
+				continue;
+			}
+
+			const FVector3f& FirstChildPointLocation = InOutBranchHierarchy.Points[FistChildBranchPointIndex];
+			const int32 SecondChildBranchPointIndex = ChildBranch.PointIndices[1];
+			const FVector3f& SecondChildPointLocation = InOutBranchHierarchy.Points[SecondChildBranchPointIndex];
+
+			const FVector3f ChildApicalDirection = (SecondChildPointLocation - FirstChildPointLocation).GetSafeNormal();
+			const float ApicalDotProduct = ChildApicalDirection.Dot(ParentApicalDirection);
+
+			constexpr float MinDotForFusingChild = 0.90f;
+			if (ApicalDotProduct < MinDotForFusingChild)
+			{
+				continue;
+			}
+					
+			if (ApicalDotProduct < BestApicalDotProduct)
+			{
+				continue;
+			}
+
+			ChildBranchToFuse = ChildBranchIndex;
+			BestApicalDotProduct = ApicalDotProduct;
+		}
+
+		if (ChildBranchToFuse != INDEX_NONE)
+		{
+			BranchesToFuseWithParent.Add(ChildBranchToFuse);
+		}
+
+		for (int32 ChildBranchIndex : Branch.ChildBranchIndices)
+		{
+			SearchForBranchesToFuse_Recursive(ChildBranchIndex);
+		}
+	};
+
+	SearchForBranchesToFuse_Recursive(InOutBranchHierarchy.RootBranchIndex);
+
+	TArray<TArray<int32>> ParentBranchRefs;
+	ParentBranchRefs.SetNum(InOutBranchHierarchy.Branches.Num());
+
+	BranchesToFuseWithParent.Sort(); // Remove points from back to front to avoid having to continuously re-map the BranchesToFuseWithParent array
+	for (int32 i = BranchesToFuseWithParent.Num() - 1; i >= 0; --i)
+	{
+		const int32 OrphanedBranchIndex = BranchesToFuseWithParent[i];
+
+		// Fuse branch with parent
+		{
+			const FPVBranchDescription& OrphanedBranch = InOutBranchHierarchy.Branches[OrphanedBranchIndex];
+			FPVBranchDescription& BranchToFuseWith = InOutBranchHierarchy.Branches[OrphanedBranch.ParentBranchIndex];
+
+			check(BranchToFuseWith.PointIndices.Last() == OrphanedBranch.PointIndices[0]);
+
+			BranchToFuseWith.PointIndices.Reserve(BranchToFuseWith.PointIndices.Num() + OrphanedBranch.PointIndices.Num() - 1);
+			for (int32 j = 1; j < OrphanedBranch.PointIndices.Num(); ++j)
+			{
+				BranchToFuseWith.PointIndices.Add(OrphanedBranch.PointIndices[j]);
+			}
+
+			check(BranchToFuseWith.ChildBranchIndices.Contains(OrphanedBranchIndex));
+
+			BranchToFuseWith.ChildBranchIndices.Remove(OrphanedBranchIndex);
+			BranchToFuseWith.ChildBranchIndices.Reserve(BranchToFuseWith.ChildBranchIndices.Num() + OrphanedBranch.ChildBranchIndices.Num());
+			for (int32 ChildBranchIndex : OrphanedBranch.ChildBranchIndices)
+			{
+				FPVBranchDescription& ChildBranch = InOutBranchHierarchy.Branches[ChildBranchIndex];
+				BranchToFuseWith.ChildBranchIndices.Add(ChildBranch.BranchIndex);
+				ChildBranch.ParentBranchIndex = BranchToFuseWith.BranchIndex;
+			}
+		}
+			
+		// Remove orphaned branch + fixup branch references
+		{
+			const int32 NewBranchIndex = OrphanedBranchIndex;
+			const int32 OldBranchIndex = InOutBranchHierarchy.Branches.Num() - 1;
+			InOutBranchHierarchy.Branches.RemoveAtSwap(OrphanedBranchIndex);
+
+			if (OldBranchIndex == NewBranchIndex)
+			{
+				continue;
+			}
+
+			FPVBranchDescription& SwappedBranch = InOutBranchHierarchy.Branches[NewBranchIndex];
+			SwappedBranch.BranchIndex = NewBranchIndex;
+
+			for (int32 ChildBranchIndex : SwappedBranch.ChildBranchIndices)
+			{
+				InOutBranchHierarchy.Branches[ChildBranchIndex].ParentBranchIndex = NewBranchIndex;
+			}
+
+			FPVBranchDescription& SwappedBranchParentBranch = InOutBranchHierarchy.Branches[SwappedBranch.ParentBranchIndex];
+			const int32 OldParentChildIndicesIndex = SwappedBranchParentBranch.ChildBranchIndices.Find(OldBranchIndex);
+			check(OldParentChildIndicesIndex != INDEX_NONE);
+
+			SwappedBranchParentBranch.ChildBranchIndices[OldParentChildIndicesIndex] = NewBranchIndex;
+		}
+	}
+}
+
+bool PV::ImportHelper::GenerateGrowthDataFromBranchHierarchies(FManagedArrayCollection& OutCollection, TArrayView<const FPVBranchHierarchyDescription> InBranchHierarchies, bool bSplitPoints)
+{
+	using namespace PV::AttributesHelper;
+
+	if (InBranchHierarchies.Num() < 1)
+	{
+		return false;
+	}
+
+	int32 NumBranchPoints = 0;
+	int32 NumBranches = 0;
+
+	for (const FPVBranchHierarchyDescription& BranchHierarchy : InBranchHierarchies)
+	{
+		NumBranchPoints += BranchHierarchy.GetNumBranchPoints();
+		NumBranches += BranchHierarchy.GetNumBranches();
+
+		if (BranchHierarchy.ValidateHierarchy() != FPVBranchHierarchyDescription::EValidateHierarchyResult::Success)
+		{
+			return false;
+		}
+	}
+
+	CreateEmptyGrowthData(OutCollection, NumBranchPoints, 0, NumBranches, 0);
+	
+	auto PointPositionAttribute = FPointPositionAttribute::GetAttribute(OutCollection);
+	auto PointScaleAttribute = FPointScaleAttribute::GetAttribute(OutCollection);
+	auto BranchPointsAttribute = FBranchPointsAttribute::GetAttribute(OutCollection);
+	auto BranchParentNumberAttribute = FBranchParentNumberAttribute::GetAttribute(OutCollection);
+	auto BranchParentsAttribute = FBranchParentsAttribute::GetAttribute(OutCollection);
+	auto BranchNumberAttribute = FBranchNumberAttribute::GetAttribute(OutCollection);
+	auto BranchChildrenAttribute = FBranchChildrenAttribute::GetAttribute(OutCollection);
+	auto BranchHierarchyNumberAttribute = FBranchHierarchyNumberAttribute::GetAttribute(OutCollection);
+	auto BranchPlantNumberAttribute = FBranchPlantNumberAttribute::GetAttribute(OutCollection);
+	FillAttributesFromBranchHierarchy(
+		PointPositionAttribute,
+		PointScaleAttribute,
+		BranchPointsAttribute,
+		BranchParentsAttribute,
+		BranchChildrenAttribute,
+		BranchParentNumberAttribute,
+		BranchNumberAttribute,
+		BranchHierarchyNumberAttribute,
+		BranchPlantNumberAttribute,
+		InBranchHierarchies
+	);
+
+	auto PointBudNumberAttribute = FPointBudNumberAttribute::GetAttribute(OutCollection);
+	auto BranchSourceBudNumberAttribute = FBranchSourceBudNumberAttribute::GetAttribute(OutCollection);
+	ComputeBudNumbers(
+		PointBudNumberAttribute,
+		BranchSourceBudNumberAttribute, 
+		BranchPointsAttribute
+	);
+
+	auto PointNjordPixelIndexAttribute = FPointNjordPixelIndexAttribute::GetAttribute(OutCollection);
+	ComputeNjordPixelIndex(
+		PointNjordPixelIndexAttribute,
+		BranchPointsAttribute,
+		BranchParentNumberAttribute, 
+		BranchNumberAttribute,
+		BranchChildrenAttribute
+	);
+
+	auto PointLengthFromRootAttribute = FPointLengthFromRootAttribute::GetAttribute(OutCollection);
+	auto PointLengthFromSeedAttribute = FPointLengthFromSeedAttribute::GetAttribute(OutCollection);
+	ComputeLengthFromRoot(
+		PointLengthFromRootAttribute,
+		PointLengthFromSeedAttribute,
+		PointPositionAttribute,
+		BranchPointsAttribute,
+		BranchParentNumberAttribute,
+		BranchNumberAttribute,
+		BranchChildrenAttribute
+	);
+
+	auto BudDirectionAttribute = FBudDirectionAttribute::GetAttribute(OutCollection);
+	ComputeBudDirections(
+		BudDirectionAttribute,
+		PointPositionAttribute,
+		BranchChildrenAttribute,
+		BranchParentNumberAttribute,
+		BranchNumberAttribute,
+		BranchPointsAttribute
+	);
+
+	auto BudLightDetectedAttribute = FBudLightDetectedAttribute::GetAttribute(OutCollection);
+	ComputeBudLightDetected(BudLightDetectedAttribute);
+
+	InterpolatePointsScales(
+		PointScaleAttribute,
+		PointLengthFromRootAttribute,
+		BranchPointsAttribute,
+		BranchParentNumberAttribute,
+		BranchChildrenAttribute,
+		BranchNumberAttribute
+	);
+
+	auto BudLateralMeristemAttribute = FBudLateralMeristemAttribute::GetAttribute(OutCollection);
+	ComputeBudLateralMeristem(
+		BudLateralMeristemAttribute,
+		PointPositionAttribute,
+		PointScaleAttribute,
+		PointLengthFromRootAttribute,
+		BudDirectionAttribute,
+		BranchPointsAttribute,
+		BranchParentNumberAttribute,
+		BranchChildrenAttribute,
+		BranchNumberAttribute
+	);
+
+	auto GroundGradientAttribute = PV::FPointGroundGradientAttribute::GetAttribute(OutCollection);
+	auto PointScaleGradientAttribute = PV::FPointScaleGradientAttribute::GetAttribute(OutCollection);
+	auto PointPlantGradientAttribute = PV::FPointPlantGradientAttribute::GetAttribute(OutCollection);
+	auto PointBranchGradientAttribute = PV::FPointBranchGradientAttribute::GetAttribute(OutCollection);
+	ComputeGradients(
+		GroundGradientAttribute,
+		PointScaleGradientAttribute,
+		PointPlantGradientAttribute,
+		PointBranchGradientAttribute,
+		PointLengthFromRootAttribute,
+		PointScaleAttribute,
+		BranchPointsAttribute,
+		BranchParentNumberAttribute,
+		BranchNumberAttribute,
+		BranchChildrenAttribute
+	);
+
+	auto BudDevelopmentAttribute = FBudDevelopmentAttribute::GetAttribute(OutCollection);
+	ComputeBudDevelopment(
+		BudDevelopmentAttribute,
+		PointPlantGradientAttribute,
+		BranchPointsAttribute,
+		BranchParentNumberAttribute,
+		BranchParentsAttribute,
+		BranchNumberAttribute
+	);
+
+	auto BudHormoneLevelsAttribute = FBudHormoneLevelsAttribute::GetAttribute(OutCollection);
+	ComputeBudHormoneLevels(
+		BudHormoneLevelsAttribute,
+		PointPlantGradientAttribute
+	);
+
+	auto BudStatusAttribute = FBudStatusAttribute::GetAttribute(OutCollection);
+	ComputeBudStatus(
+		BudStatusAttribute,
+		BranchPointsAttribute,
+		BranchParentNumberAttribute,
+		BranchNumberAttribute,
+		BranchChildrenAttribute
+	);
+
+	if (bSplitPoints)
+	{
+		SplitFusedPoints(OutCollection); // NOTE: Important that this is done last as the compute steps assumes fused points
+	}
+
+	return true;
+}
+
+bool PV::ImportHelper::GenerateGrowthDataFromBranchHierarchy(FManagedArrayCollection& OutCollection, const FPVBranchHierarchyDescription& InBranchHierarchy, bool bSplitPoints)
+{
+	TArrayView<const FPVBranchHierarchyDescription> BranchHierarchyArrayView(&InBranchHierarchy, 1);
+	return GenerateGrowthDataFromBranchHierarchies(OutCollection, BranchHierarchyArrayView, bSplitPoints);
+}
+
+void PV::ImportHelper::DistToBranch(const FVector3f& Point, const TArray<FVector3f>& Points, const TArray<int32>& BranchPointIndices, FVector3f& OutClosestPoint, float& OutDistSquared, int32& OutPointIndex)
+{
+	OutClosestPoint = FVector3f::ZeroVector;
+	OutDistSquared = MAX_FLT;
+	OutPointIndex = INDEX_NONE;
+
+	for (int32 i = 0; i < BranchPointIndices.Num() - 1; ++i)
+	{
+		const int32 PointIndex0 = BranchPointIndices[i];
+		const int32 PointIndex1 = BranchPointIndices[i + 1];
+		if (!Points.IsValidIndex(PointIndex0) || !Points.IsValidIndex(PointIndex1))
+		{
+			continue;
+		}
+
+		const FVector3f& BranchPoint1 = Points[PointIndex0];
+		const FVector3f& BranchPoint2 = Points[PointIndex1];
+
+		const FVector3f ClosestPointOnLine = FMath::ClosestPointOnSegment(Point, BranchPoint1, BranchPoint2);
+		const float DistToLineSquared = (Point - ClosestPointOnLine).SizeSquared();
+		if (DistToLineSquared < OutDistSquared)
+		{
+			OutDistSquared = DistToLineSquared;
+			OutClosestPoint = ClosestPointOnLine;
+			OutPointIndex = PointIndex0;
+		}
+	}
+}
+
+float PV::ImportHelper::DistToBranch(const FVector3f& Point, const TArray<FVector3f>& Points, const TArray<int32>& BranchPointIndices)
+{
+	FVector3f ClosestPoint;
+	float ClosestDistSquared = MAX_FLT;
+	int32 PointIndex;
+	DistToBranch(Point, Points, BranchPointIndices, ClosestPoint, ClosestDistSquared, PointIndex);
+
+	return ClosestDistSquared > 0 ? FMath::Sqrt(ClosestDistSquared) : 0;
+}
+
+TArray<int32> PV::ImportHelper::EstimateBranchHierarchy(const TArray<FVector3f>& Points, const TArray<TArray<int32>>& Branches, int32 RootBranchIndex)
+{
+	TRACE_CPUPROFILER_EVENT_SCOPE(PV::ImportHelper::EstimateBranchHierarchy);
+
+	TArray<int32> OutBranchParents;
+	OutBranchParents.SetNum(Branches.Num());
+
+	for (const TArray<int32>& Branch : Branches)
+	{
+		for (int32 PointIndex : Branch)
+		{
+			if (!Points.IsValidIndex(PointIndex))
+			{
+				return TArray<int32>();
+			}
+		}
+	}
+
+	for (int32 i = 0; i < Branches.Num(); ++i)
+	{
+		if (i == RootBranchIndex)
+		{
+			OutBranchParents[i] = INDEX_NONE;
+			continue;
+		}
+
+		if (Branches[i].Num() == 0)
+		{
+			continue;
+		}
+
+		const int32 BranchRootPointIndex = Branches[i][0];
+		const FVector3f& BranchRootPoint = Points[BranchRootPointIndex];
+		float ClosestBranchDistSquared = MAX_FLT;
+		int32 ClosestBranchIndex = INDEX_NONE;
+		bool ClosestBranchPointIsRoot = true;
+
+		for (int32 j = 0; j < Branches.Num(); ++j)
+		{
+			if (i == j)
+			{
+				continue;
+			}
+
+			FVector3f ClosestPointOnBranch;
+			float BranchDistSquared = MAX_FLT;
+			int32 ClosestPointIndex;
+			ImportHelper::DistToBranch(BranchRootPoint, Points, Branches[j], ClosestPointOnBranch, BranchDistSquared, ClosestPointIndex);
+
+			const bool bBranchPointIsRoot = ClosestPointIndex == BranchRootPointIndex;
+
+			// If there are multiple branches growing from the same point then we may select the wrong point as the distance to the branch will be the same.
+			// We therefore prefer the branch where the closest point is not root.
+			if ((FMath::IsNearlyEqual(BranchDistSquared, ClosestBranchDistSquared) && ClosestBranchPointIsRoot && !bBranchPointIsRoot)
+				|| BranchDistSquared < ClosestBranchDistSquared)
+			{
+				ClosestBranchIndex = j;
+				ClosestBranchDistSquared = BranchDistSquared;
+				ClosestBranchPointIsRoot = bBranchPointIsRoot;
+			}
+		}
+
+		OutBranchParents[i] = ClosestBranchIndex;
+	}
+
+	return OutBranchParents;
+}
+
+int32 PV::ImportHelper::FindRootBranch(const TArray<FVector3f>& InPoints, const TArray<TArray<int32>>& BranchIndices)
+{
+	float LowestRootPoint = MAX_FLT;
+	int32 LowestRootPointIdx = INDEX_NONE;
+	for (int32 i = 0; i < BranchIndices.Num(); ++i)
+	{
+		const TArray<int32>& Branch = BranchIndices[i];
+		if (Branch.Num() > 0 && InPoints[Branch[0]].Z < LowestRootPoint)
+		{
+			LowestRootPoint = InPoints[Branch[0]].Z;
+			LowestRootPointIdx = i;
+		}
+	}
+	return LowestRootPointIdx;
+}
+
+FPVBranchHierarchyDescription PV::ImportHelper::CreateBranchHierarchyFromPoints(const TArray<FVector3f>& InPoints, const TArray<float>& InPointsRadii, const TArray<TArray<int32>>& InBranchPointIndices)
+{
+	TRACE_CPUPROFILER_EVENT_SCOPE(PV::ImportHelper::CreateBranchHierarchyFromPoints);
+
+	if (InBranchPointIndices.Num() == 0 || InPoints.Num() == 0)
+	{
+		return FPVBranchHierarchyDescription();
+	}
+	
+	if (InPointsRadii.Num() != InPoints.Num())
+	{
+		return FPVBranchHierarchyDescription();
+	}
+
+	if (InBranchPointIndices.Num() == 1)
+	{
+		FPVBranchHierarchyDescription OutBranchHierarchy;
+		OutBranchHierarchy.RootBranchIndex = 0;
+		OutBranchHierarchy.Points = InPoints;
+		OutBranchHierarchy.PointsRadii = InPointsRadii;
+		FPVBranchDescription& Branch = OutBranchHierarchy.Branches.AddDefaulted_GetRef();
+		Branch.PointIndices = InBranchPointIndices[0];
+		Branch.ParentBranchIndex = INDEX_NONE;
+		Branch.BranchIndex = 0;
+		return OutBranchHierarchy;
+	}
+
+	const int32 RootBranchIndex = FindRootBranch(InPoints, InBranchPointIndices);
+
+	FPVBranchHierarchyDescription OutBranchHierarchy;
+	OutBranchHierarchy.RootBranchIndex = RootBranchIndex;
+	OutBranchHierarchy.Points = InPoints;
+	OutBranchHierarchy.PointsRadii = InPointsRadii;
+
+	// Fill branch points
+	for (const TArray<int32>& BranchPointIndices : InBranchPointIndices)
+	{
+		FPVBranchDescription& Branch = OutBranchHierarchy.Branches.AddDefaulted_GetRef();
+		Branch.PointIndices = BranchPointIndices;
+	}
+
+	// Set up branch hierarchy
+	const TArray<int32> BranchParents = EstimateBranchHierarchy(InPoints, InBranchPointIndices, RootBranchIndex);
+	for (int32 i = 0; i < BranchParents.Num(); ++i)
+	{
+		OutBranchHierarchy.Branches[i].BranchIndex = i;
+		OutBranchHierarchy.Branches[i].ParentBranchIndex = BranchParents[i];
+
+		if (BranchParents[i] != INDEX_NONE)
+		{
+			OutBranchHierarchy.Branches[BranchParents[i]].ChildBranchIndices.Add(i);
+		}
+	}
+
+	if (OutBranchHierarchy.ValidateHierarchy() != FPVBranchHierarchyDescription::EValidateHierarchyResult::Success)
+	{
+		return FPVBranchHierarchyDescription();
+	}
+
+	return OutBranchHierarchy;
+}
diff --git a/Engine/Plugins/Experimental/ProceduralVegetationEditor/Source/ProceduralVegetation/Private/Helpers/PVImportHelpers.h b/Engine/Plugins/Experimental/ProceduralVegetationEditor/Source/ProceduralVegetation/Private/Helpers/PVImportHelpers.h
new file mode 100644
--- /dev/null
+++ b/Engine/Plugins/Experimental/ProceduralVegetationEditor/Source/ProceduralVegetation/Private/Helpers/PVImportHelpers.h
@@ -0,0 +1,170 @@
+// Copyright Epic Games, Inc. All Rights Reserved.
+#pragma once
+
+#include "CoreMinimal.h"
+
+#include "Utils/PVAttributes.h"
+#include "Utils/PVDynamicMeshVertexAttribute.h"
+
+namespace PV::Facades
+{
+	class FPointFacade;
+	class FBranchFacade;
+	class FPlantFacade;
+};
+
+struct FManagedArrayCollection;
+class UProceduralVegetationPreset;
+class UPackage;
+
+namespace PV
+{
+	template<typename T> class TDynamicMeshVertexAttributeDefinition;
+	template<typename AttribValueType> class TDynamicMeshVertexAttributeExt;
+	template<typename AttribValueType> class TVertexAttributeView;
+}
+
+namespace UE::Geometry
+{
+	class FDynamicMesh3;
+	template<class T> class TMeshAABBTree3;
+	using FDynamicMeshAABBTree3 = TMeshAABBTree3<FDynamicMesh3>;
+}
+
+struct FPVBranchDescription
+{
+	int32 ParentBranchIndex = INDEX_NONE;
+	int32 BranchIndex = INDEX_NONE;
+	TArray<int32> ChildBranchIndices;
+	TArray<int32> PointIndices;
+};
+
+struct FPVBranchHierarchyDescription
+{
+	int32 RootBranchIndex = INDEX_NONE;
+	TArray<FPVBranchDescription> Branches;
+	TArray<FVector3f> Points;
+	TArray<float> PointsRadii;
+
+	int32 GetNumBranches() const { return Branches.Num(); }
+	int32 GetNumBranchPoints() const { return Points.Num(); }
+
+	enum class EValidateHierarchyResult
+	{
+		InvalidRootBranchIndex,
+		InvalidRootBranchParentIndex,
+		InvalidParentBranchIndex,
+		InvalidBranchIndex,
+		InvalidChildBranchIndex,
+		InvalidPointIndex,
+		BranchIndexMismatch,
+		DisjointedHierarchy,
+		CircularHierarchy,
+		Success
+	};
+	EValidateHierarchyResult ValidateHierarchy() const;
+};
+
+namespace PV::DynamicMeshHelper
+{
+	using FBooleanVertexAttribute = TDynamicMeshVertexAttributeExt<bool>;
+	using FIntVertexAttribute     = TDynamicMeshVertexAttributeExt<int32>;
+	using FFloatVertexAttribute   = TDynamicMeshVertexAttributeExt<float>;
+
+	extern const PV::TDynamicMeshVertexAttributeDefinition<float> VertexScaleAttributeDefinition;
+	extern const PV::TDynamicMeshVertexAttributeDefinition<bool> EndPointAttributeDefinition;
+	extern const PV::TDynamicMeshVertexAttributeDefinition<bool> RootPointAttributeDefinition;
+	extern const PV::TDynamicMeshVertexAttributeDefinition<float> LengthFromRootAttributeDefinition;
+	extern const PV::TDynamicMeshVertexAttributeDefinition<int32> NextBranchPointAttributeDefinition;
+
+	void AddDynamicMeshToCollection(const UE::Geometry::FDynamicMesh3& DynamicMesh, FManagedArrayCollection& InOutCollection, const FTransform& Offset = FTransform::Identity);
+	
+	void AddBranchHierarchiesToCollection(const TArray<FPVBranchHierarchyDescription>& BranchHierarchyDescriptions, FManagedArrayCollection& OutCollection);
+
+	TArray<int32> GetUniqueVertexIndices(const UE::Geometry::FDynamicMesh3& DynamicMesh, const TArray<int32>& TriangleIds);
+
+	TArray<int32> GetVertexNeighbours(const UE::Geometry::FDynamicMesh3& DynamicMesh, int32 VertexID);
+
+	TArray<int32> FindVerticesInRadius(const UE::Geometry::FDynamicMeshAABBTree3& DynamicMeshAABBTree, const FVector& Pos, float Radius);
+	
+	struct FMeshBranchHierarchy
+	{
+		struct FBranch
+		{
+			TArray<int32> BranchPoints;
+			int32 ParentBranch = INDEX_NONE;
+		};
+
+		TArray<FBranch> Branches;
+	};
+
+	void ComputeDynamicMeshBranchHierarchies(
+		const UE::Geometry::FDynamicMesh3& DynamicMesh,
+		const FBooleanVertexAttribute& RootPointsAttribute,
+		const FBooleanVertexAttribute& EndPointsAttribute,
+		const FFloatVertexAttribute& VertexScaleAttribute,
+		FFloatVertexAttribute& LengthFromRootAttribute,
+		FIntVertexAttribute& NextBranchPointAttribute,
+		const TFunction<bool(int32)>& VertexFilterPredicate,
+		TArray<FMeshBranchHierarchy>& OutDynamicMeshBranchHierarchies
+	);
+
+	void ConvertDynamicMeshBranchHierarchiesToBranchDescription(
+		const UE::Geometry::FDynamicMesh3& DynamicMesh,
+		const FFloatVertexAttribute& VertexScaleAttribute,
+		const FBooleanVertexAttribute& EndPointAttribute,
+		const TArray<FMeshBranchHierarchy>& InMeshBranchHierarchies,
+		TArray<FPVBranchHierarchyDescription>& OutBranchHierarchyDescriptions
+	);
+}
+
+namespace PV::ImportHelper
+{
+	// Creates a collection containing all attributes (of the supplied size) required for creating PV growth data.
+	void CreateEmptyGrowthData(
+		FManagedArrayCollection& OutCollection,
+		int32 NumBranchPoints,
+		int32 NumFoliagePoints,
+		int32 NumBranches,
+		int32 NumFoliageMeshes
+	);
+
+	void FillAttributesFromBranchHierarchy(
+		FManagedArrayCollection& OutCollection,
+		TArrayView<const FPVBranchHierarchyDescription> InBranchHierarchies
+	);
+
+	// Copies the branch hierarchies from the input BranchHierarchies to the attributes.
+	// Assumes the attributes are large enough to receive the hierarchies.
+	void FillAttributesFromBranchHierarchy(
+		PV::FPointPositionAttributeView PointPositionAttribute,
+		PV::FPointScaleAttributeView PointScaleAttribute,
+		PV::FBranchPointsAttributeView BranchPointsAttribute,
+		PV::FBranchParentsAttributeView BranchParentsAttribute,
+		PV::FBranchChildrenAttributeView BranchChildrenAttribute,
+		PV::FBranchParentNumberAttributeView BranchParentNumberAttribute,
+		PV::FBranchNumberAttributeView BranchNumberAttribute,
+		PV::FBranchHierarchyNumberAttributeView BranchHierarchyNumberAttribute,
+		PV::FBranchPlantNumberAttributeView BranchPlantNumberAttribute,
+		TArrayView<const FPVBranchHierarchyDescription> InBranchHierarchies
+	);
+
+	void ConnectDisconnectedBranches(FPVBranchHierarchyDescription& InOutBranchHierarchy);
+
+	void FuseCollinearBranches(FPVBranchHierarchyDescription& InOutBranchHierarchy);
+
+	bool GenerateGrowthDataFromBranchHierarchies(FManagedArrayCollection& OutCollection, TArrayView<const FPVBranchHierarchyDescription> InBranchHierarchies, bool bSplitPoints = true);
+
+	bool GenerateGrowthDataFromBranchHierarchy(FManagedArrayCollection& OutCollection, const FPVBranchHierarchyDescription& InBranchHierarchy, bool bSplitPoints = true);
+
+	void DistToBranch(const FVector3f& Point, const TArray<FVector3f>& Points, const TArray<int32>& BranchPointIndices, FVector3f& OutClosestPoint, float& OutDistSquared, int32& OutPointIndex);
+
+	float DistToBranch(const FVector3f& Point, const TArray<FVector3f>& Points, const TArray<int32>& BranchPointIndices);
+
+	// Returns array of parent branch index for each supplied branch. Root branch will have parent index INDEX_NONE.
+	TArray<int32> EstimateBranchHierarchy(const TArray<FVector3f>& Points, const TArray<TArray<int32>>& Branches, int32 RootBranchIndex);
+
+	int32 FindRootBranch(const TArray<FVector3f>& InPoints, const TArray<TArray<int32>>& BranchIndices);
+
+	FPVBranchHierarchyDescription CreateBranchHierarchyFromPoints(const TArray<FVector3f>& InPoints, const TArray<float>& InPointsRadii, const TArray<TArray<int32>>& InBranchPointIndices);
+};
\ No newline at end of file
diff --git a/Engine/Plugins/Experimental/ProceduralVegetationEditor/Source/ProceduralVegetation/Private/Helpers/PVTransformHelper.h b/Engine/Plugins/Experimental/ProceduralVegetationEditor/Source/ProceduralVegetation/Private/Helpers/PVTransformHelper.h
--- a/Engine/Plugins/Experimental/ProceduralVegetationEditor/Source/ProceduralVegetation/Private/Helpers/PVTransformHelper.h
+++ b/Engine/Plugins/Experimental/ProceduralVegetationEditor/Source/ProceduralVegetation/Private/Helpers/PVTransformHelper.h
@@ -1,7 +1,7 @@
 // Copyright Epic Games, Inc. All Rights Reserved.
 #pragma once
 
-#include "Facades/PVAttributes.h"
+#include "Utils/PVAttributes.h"
 
 namespace PV::Transform
 {
diff --git a/Engine/Plugins/Experimental/ProceduralVegetationEditor/Source/ProceduralVegetation/Private/IPVEditorInterface.cpp b/Engine/Plugins/Experimental/ProceduralVegetationEditor/Source/ProceduralVegetation/Private/IPVEditorInterface.cpp
deleted file mode 100644
diff --git a/Engine/Plugins/Experimental/ProceduralVegetationEditor/Source/ProceduralVegetation/Private/Implementations/PVFoliage.cpp b/Engine/Plugins/Experimental/ProceduralVegetationEditor/Source/ProceduralVegetation/Private/Implementations/PVFoliage.cpp
--- a/Engine/Plugins/Experimental/ProceduralVegetationEditor/Source/ProceduralVegetation/Private/Implementations/PVFoliage.cpp
+++ b/Engine/Plugins/Experimental/ProceduralVegetationEditor/Source/ProceduralVegetation/Private/Implementations/PVFoliage.cpp
@@ -3,7 +3,7 @@
 #include "PVFoliage.h"
  
 #include "ProceduralVegetationModule.h"
-#include "PVFloatRamp.h"
+#include "Utils/PVFloatRamp.h"
 #include "Algo/RandomShuffle.h"
  
 #include "Facades/PVAttributesNames.h"
diff --git a/Engine/Plugins/Experimental/ProceduralVegetationEditor/Source/ProceduralVegetation/Private/Implementations/PVGrafter.cpp b/Engine/Plugins/Experimental/ProceduralVegetationEditor/Source/ProceduralVegetation/Private/Implementations/PVGrafter.cpp
--- a/Engine/Plugins/Experimental/ProceduralVegetationEditor/Source/ProceduralVegetation/Private/Implementations/PVGrafter.cpp
+++ b/Engine/Plugins/Experimental/ProceduralVegetationEditor/Source/ProceduralVegetation/Private/Implementations/PVGrafter.cpp
@@ -3,7 +3,7 @@
 #include "PVGrafter.h"
 
 #include "ProceduralVegetationModule.h"
-#include "PVFloatRamp.h"
+#include "Utils/PVFloatRamp.h"
 #include "Algo/RandomShuffle.h"
 
 #include "Facades/PVAttributesNames.h"
@@ -189,10 +189,10 @@ void FPVGrafter::SetGraft(FManagedArrayCollection& OutCollection, const FManaged
 	const TManagedArray<int32>& BranchHierarchyNumbers = BranchFacadeSourceGraft.GetBranchHierarchyNumbers();
 	TArray<int32> TransformedBranchHierarchyNumbers;
 	TransformedBranchHierarchyNumbers.SetNum(BranchHierarchyNumbers.Num());
-	const int32 BranchHeirarchyNumberOfTarget = BranchFacadeOut.GetBranchHierarchyNumber(InAttributes.BranchIndex);
+	const int32 BranchHierarchyNumberOfTarget = BranchFacadeOut.GetBranchHierarchyNumber(InAttributes.BranchIndex);
 	for (int32 i = 0; i < BranchHierarchyNumbers.Num(); ++i)
 	{
-		TransformedBranchHierarchyNumbers[i] = BranchHierarchyNumbers[i] + BranchHeirarchyNumberOfTarget;
+		TransformedBranchHierarchyNumbers[i] = BranchHierarchyNumbers[i] + BranchHierarchyNumberOfTarget;
 	}
 
 	BranchFacadeOut.SetBranchHierarchyNumbersFromIndex(TransformedBranchHierarchyNumbers, BranchFacadeOutStartIndex);
diff --git a/Engine/Plugins/Experimental/ProceduralVegetationEditor/Source/ProceduralVegetation/Private/Implementations/PVImporter_StaticMesh.cpp b/Engine/Plugins/Experimental/ProceduralVegetationEditor/Source/ProceduralVegetation/Private/Implementations/PVImporter_StaticMesh.cpp
new file mode 100644
--- /dev/null
+++ b/Engine/Plugins/Experimental/ProceduralVegetationEditor/Source/ProceduralVegetation/Private/Implementations/PVImporter_StaticMesh.cpp
@@ -0,0 +1,981 @@
+// Copyright Epic Games, Inc. All Rights Reserved.
+
+#include "PVImporter_StaticMesh.h"
+
+#include "Helpers/PVImportHelpers.h"
+#include "Utils/PVDynamicMeshVertexAttribute.h"
+
+#include "Engine/StaticMesh.h"
+#include "DynamicMesh/DynamicMesh3.h"
+#include "DynamicMesh/DynamicMeshAABBTree3.h"
+#include "DynamicMesh/MeshNormals.h"
+#include "ConversionUtils/SceneComponentToDynamicMesh.h"
+#include "Selections/MeshConnectedComponents.h"
+
+using FBooleanVertexAttribute = PV::DynamicMeshHelper::FBooleanVertexAttribute;
+using FIntVertexAttribute = PV::DynamicMeshHelper::FIntVertexAttribute;
+using FFloatVertexAttribute = PV::DynamicMeshHelper::FFloatVertexAttribute;
+
+const static PV::TDynamicMeshVertexAttributeDefinition<bool> FillerVertexAttributeDefinition(TEXT("FillerVertex"));
+const static PV::TDynamicMeshVertexAttributeDefinition<int32> FuseToAttributeDefinition(TEXT("FuseTo"));
+
+static void AddBranchCurvesToCollection(FManagedArrayCollection& InCollection, const TArray<TArray<int32>>& BranchCurves, const TFunction<FVector3f(int32)>& GetVertexPosFunc)
+{
+	const static FName BranchCurveGroupName = TEXT("BranchCurves");
+	const static FName CurveAttributeName = TEXT("Curve");
+
+	if (!InCollection.HasGroup(BranchCurveGroupName))
+	{
+		InCollection.AddGroup(BranchCurveGroupName);
+	}
+
+	InCollection.AddElements(BranchCurves.Num(), BranchCurveGroupName);
+
+	TManagedArray<TArray<FVector3f>>& CurveAttribute = InCollection.AddAttribute<TArray<FVector3f>>(CurveAttributeName, BranchCurveGroupName);
+	for (int32 i = 0; i < BranchCurves.Num(); ++i)
+	{
+		const TArray<int32>& BranchCurve = BranchCurves[i];
+		CurveAttribute[i].Reserve(BranchCurve.Num());
+		for (int32 VertexIndex : BranchCurve)
+		{
+			CurveAttribute[i].Add(GetVertexPosFunc(VertexIndex));
+		}
+	}
+}
+
+static FVector GenerateFillerVertex(const FVector& Pos1, const FVector& Pos2)
+{
+	// Since we don't have edges for the DynamicMesh we add a new vertex at the center (offset by 0.1mm) whenever we need to represent a line
+	return ((Pos1 + Pos2) * 0.5)  + (Pos1 - Pos2).GetSafeNormal().Cross(FVector::ForwardVector).Cross(FVector::UpVector)/* * 0.01*/;
+}
+
+static void RemoveUnusedVertices(UE::Geometry::FDynamicMesh3& DynamicMesh)
+{
+	TRACE_CPUPROFILER_EVENT_SCOPE(OrphanRemoval);
+
+	DynamicMesh.RemoveUnusedVertices();
+	DynamicMesh.CompactInPlace();
+}
+
+static bool FilterTrianglesUsingMaterials(UE::Geometry::FDynamicMesh3& DynamicMesh, const UStaticMesh* StaticMesh, const TArray<FName>& MaterialsToKeep)
+{
+	TRACE_CPUPROFILER_EVENT_SCOPE(KeepMaterial);
+
+	TSet<int32> MaterialIDs;
+	MaterialIDs.Reserve(MaterialsToKeep.Num());
+
+	for (FName MaterialName : MaterialsToKeep)
+	{
+		const int32 MaterialIndexToKeep = StaticMesh->GetMaterialIndex(MaterialName);
+		if (MaterialIndexToKeep != INDEX_NONE)
+		{
+			MaterialIDs.Add(MaterialIndexToKeep);
+		}
+	}
+
+	if (MaterialIDs.Num() == 0)
+	{
+		return false;
+	}
+
+	UE::Geometry::FDynamicMeshAttributeSet* MeshAttributeSet = DynamicMesh.Attributes();
+	UE::Geometry::FDynamicMeshMaterialAttribute* MaterialIDAttribute = MeshAttributeSet ? MeshAttributeSet->GetMaterialID() : nullptr;
+		
+	if (!MaterialIDAttribute)
+	{
+		return false;
+	}
+
+	for (int32 TriangleID : DynamicMesh.TriangleIndicesItr())
+	{
+		const int32 MaterialID = MaterialIDAttribute->GetValue(TriangleID);
+		if (!MaterialIDs.Contains(MaterialID))
+		{
+			DynamicMesh.RemoveTriangle(TriangleID);
+		}
+	}
+
+	RemoveUnusedVertices(DynamicMesh);
+
+	return true;
+}
+
+static void ComputeNormals(UE::Geometry::FDynamicMesh3& DynamicMesh)
+{
+	TRACE_CPUPROFILER_EVENT_SCOPE(ComputeNormals);
+
+	if (!DynamicMesh.HasVertexNormals())
+	{
+		DynamicMesh.EnableVertexNormals(FVector3f::UnitX());
+	}
+
+	if (DynamicMesh.Attributes() && DynamicMesh.Attributes()->NumNormalLayers() > 0)
+	{
+		UE::Geometry::FDynamicMeshNormalOverlay* Overlay = DynamicMesh.Attributes()->GetNormalLayer(0);
+		for (int32 EID : Overlay->ElementIndicesItr())
+		{
+			int32 ParentVID = Overlay->GetParentVertex(EID);
+			if (ParentVID != INDEX_NONE)
+			{
+				DynamicMesh.SetVertexNormal(ParentVID, Overlay->GetElement(EID));
+			}
+		}
+	}
+	else
+	{
+		UE::Geometry::FMeshNormals::QuickComputeVertexNormals(DynamicMesh);
+	}
+}
+
+static void ComputeVerticesCenterAndRadius_Internal_V1(
+	UE::Geometry::FDynamicMesh3& DynamicMesh, 
+	const TArray<int32>& TriangleIds,
+	FFloatVertexAttribute& VertexScaleAttribute
+)
+{	
+	// This function tries to find the center point of the branch/trunk for each vertex, it then moves the vertex to said center point and stores the
+	// branch radius in the VertexScale attribute
+	//
+	// The algorithm will perform a sphere overlap with a radius of 50cm, at 50cm behind the vertex normal. If it overlaps a vertex it will try to move the
+	// center point to lay between the current vertex, and the found vertex. It then performs another overlap, and repeats the process until it no longer overlaps any vertices.
+	// Once an opposing vertex has been found, an extra check is performed to make sure the normal of the closest triangle is facing away from the center point.
+	
+	const TArray<int32> ComponentVertexIndices = PV::DynamicMeshHelper::GetUniqueVertexIndices(DynamicMesh, TriangleIds);
+
+	UE::Geometry::FDynamicMeshAABBTree3 DynamicMeshAABBTree(&DynamicMesh, false);
+	DynamicMeshAABBTree.Build(TriangleIds);
+
+	const static auto ComputeRadius = [](const FVector& Point, const FVector& Normal, const FVector& QueryPosition)->double
+	{
+		// Compute radius of the ball that touches points p and q and whose center falls on the normal n from p
+		const FVector Delta = Point - QueryPosition;
+		const double Dist = Delta.Size();
+		const double CosTheta = Dist > 0 ? Normal.Dot(Delta) / Dist : 0;
+		return CosTheta > 0 ? Dist / (2 * CosTheta) : 0;
+	};
+
+	double NearestDistSqr = 0.f;
+
+	TArray<FVector> NewVertexPositions;
+	NewVertexPositions.SetNum(ComponentVertexIndices.Num());
+
+	TArray<float> Scales;
+	Scales.SetNum(ComponentVertexIndices.Num());
+
+	for (int32 i = 0; i < ComponentVertexIndices.Num(); ++i)
+	{
+		const int32 VertexIndex = ComponentVertexIndices[i];
+		const FVector& VertexPosition = DynamicMesh.GetVertex(VertexIndex);
+		const FVector VertexNormal = FVector(DynamicMesh.GetVertexNormal(VertexIndex));
+
+		NewVertexPositions[i] = VertexPosition;
+		Scales[i] = 0.5;
+
+		float Radius = 50; // 50cm
+		int32 OldPoint = INDEX_NONE;
+		FVector SphereCenter = VertexPosition - (FVector(VertexNormal) * Radius);
+			
+		int32 NearPoint = DynamicMeshAABBTree.FindNearestVertex(SphereCenter, NearestDistSqr, Radius);
+		while (NearPoint != INDEX_NONE)
+		{
+			// Get the position of the nearest point of the previous near point.
+			const FVector point_position = DynamicMesh.GetVertex(NearPoint);
+
+			// Computes the radius of a sphere that would fit spherically between these two positions
+			const float NewRadius = ComputeRadius(VertexPosition, VertexNormal, point_position);
+
+			Radius = FMath::Min(Radius, FMath::Max(NewRadius, 0));
+			// Sets the center of the lookup shere to be at the positions of the orignal position, minus the offset of the new found radius
+			SphereCenter = VertexPosition - VertexNormal * Radius;
+
+			// Stores the current near point as the old point for comparison
+			OldPoint = NearPoint;
+			NearPoint = DynamicMeshAABBTree.FindNearestVertex(SphereCenter, NearestDistSqr, Radius);
+
+			if (NearPoint == VertexIndex || OldPoint == NearPoint)
+			{ 
+				break; 
+			}
+		}
+
+		const int32 NearTriID = DynamicMeshAABBTree.FindNearestTriangle(SphereCenter, NearestDistSqr);
+		if (NearTriID == INDEX_NONE)
+		{
+			continue;
+		}
+
+		const UE::Geometry::FDistPoint3Triangle3d Query = UE::Geometry::TMeshQueries<UE::Geometry::FDynamicMesh3>::TriangleDistance(DynamicMesh, NearTriID, SphereCenter);
+		const FVector& ClosestTrianglePoint = Query.ClosestTrianglePoint;
+		const FVector ClosestTrianglePointNormal = DynamicMesh.GetTriBaryNormal(NearTriID, Query.TriangleBaryCoords[0], Query.TriangleBaryCoords[1], Query.TriangleBaryCoords[2]);
+
+		const FVector CenterDir = (SphereCenter - ClosestTrianglePoint).GetSafeNormal();
+		const double CenterDot = ClosestTrianglePointNormal.Dot(CenterDir);
+		if (CenterDot > 0) 
+		{
+			continue; 
+		}
+
+		Scales[i] = FMath::Max(Radius, 0.01);
+		NewVertexPositions[i] = SphereCenter;
+	}
+
+	for (int32 i = 0; i < ComponentVertexIndices.Num(); ++i)
+	{
+		const int32 VertexIndex = ComponentVertexIndices[i];
+		DynamicMesh.SetVertex(VertexIndex, NewVertexPositions[i]);
+		VertexScaleAttribute.SetValue(VertexIndex, Scales[i]);
+	}
+}
+
+static void ComputeVerticesCenterAndRadius_Internal_V2(
+	UE::Geometry::FDynamicMesh3& DynamicMesh, 
+	const TArray<int32>& TriangleIds,
+	FFloatVertexAttribute& VertexScaleAttribute
+)
+{
+	// This is a second attempt at creating an algorithm that can walk along each vertex on a tree mesh and find the corresponding center point and
+	// radius of the branch at the point of the vertex.
+	//
+	// This algorithm will iterate the vertices and perform a ray-cast in the opposite direction of the vertex normal.
+	// If it hits another triangle then it will check whether the normal of said triangle is opposing to the current vertex (e.g. the dot of the
+	// vertex normal and the triangle normal is < -0.5). If the ray-cast fails to find a triangle with an opposing normal, then a second algorithm
+	// is employed to try and find an opposing vertex. This second algorithm finds the closest vertex to a position 100cm from the vertex in the
+	// opposite direction of the vertex normal, it then checks if the normal of this vertex is opposing. If this check fails, then it moves the
+	// query point closer and repeats the check. This is repeated 3 times before the algorithm gives up and the vertex is left with a scale if 0 which
+	// indicates that at radius/center point was not found.
+
+	// Issues with this algorithm:
+	// 1. This algorithm only works if the branch mesh is fully enclosed. If there are holes then the ray may exit through those holes and hit parts
+	// of the branch we do not desire for it to hit.
+	// 2. This algorithm does not handle intersection geometry well as it cannot differentiate between hitting an intersection branch and the opposite
+	// side of the current branch. This is partially handled by running this algorithm per detached mesh (which for most meshes will mean per branch)
+	// 3. It also not work well if the mesh is built up using only faces (e.g flat geometry instead of a tube for the branches)
+	// 4. The normal of the vertex may not point directly towards the opposite side (see "Remove extreme radius values" section below for more information)
+
+	const TArray<int32> ComponentVertexIndices = PV::DynamicMeshHelper::GetUniqueVertexIndices(DynamicMesh, TriangleIds);
+
+	UE::Geometry::FDynamicMeshAABBTree3 DynamicMeshAABBTree(&DynamicMesh, false);
+	DynamicMeshAABBTree.Build(TriangleIds);
+
+	TArray<FVector> NewVertexPositions;
+	NewVertexPositions.SetNum(ComponentVertexIndices.Num());
+
+	TArray<float> Scales;
+	Scales.SetNumZeroed(ComponentVertexIndices.Num());
+
+	TArray<int32> ThisVertexTriangleIDs;
+	for (int32 i = 0; i < ComponentVertexIndices.Num(); ++i)
+	{
+		const int32 VertexIndex = ComponentVertexIndices[i];
+
+		const FVector VertexNormal = FVector(DynamicMesh.GetVertexNormal(VertexIndex));
+		const FVector& VertexPosition = DynamicMesh.GetVertex(VertexIndex);
+
+		const FRay3d Ray(VertexPosition, -VertexNormal);
+
+		ThisVertexTriangleIDs.Reset();
+		DynamicMesh.GetVtxTriangles(VertexIndex, ThisVertexTriangleIDs);
+		
+		double HitTriangleDistance = 0;
+		int HitTriangleID = INDEX_NONE;
+		FVector3d HitBaryCoords = FVector3d::Zero();
+		UE::Geometry::FDynamicMeshAABBTree3::FQueryOptions QueryOptions;
+		QueryOptions.TriangleFilterF = [&](int32 TriangleID) { return !ThisVertexTriangleIDs.Contains(TriangleID); }; // Make sure to ignore triangles which includes this vertex, otherise we will instantly get a hit on ourselves
+		DynamicMeshAABBTree.FindNearestHitTriangle(Ray, HitTriangleDistance, HitTriangleID, HitBaryCoords, QueryOptions);
+
+		const float CollinearLimit = -0.5;
+
+		if (HitTriangleID != INDEX_NONE)
+		{
+			const FVector HitNormal = DynamicMesh.GetTriBaryNormal(HitTriangleID, HitBaryCoords[0], HitBaryCoords[1], HitBaryCoords[2]);
+			const FVector HitLocation = DynamicMesh.GetTriBaryPoint(HitTriangleID, HitBaryCoords[0], HitBaryCoords[1], HitBaryCoords[2]);
+
+			const float Dot = HitNormal.Dot(VertexNormal);
+			const bool bIsOpposingAndCollinear = Dot < CollinearLimit;
+			if (bIsOpposingAndCollinear)
+			{
+				const FVector Center = (HitLocation + VertexPosition) * 0.5;
+				Scales[i] = (VertexPosition - Center).Size();
+				NewVertexPositions[i] = Center;
+			}
+		}
+
+		if (Scales[i] == 0)
+		{
+			// Could not find opposing point using ray intersection, try using FindNearestVertex instead.
+			const int32 MaxNumSamples = 3;
+			for (int32 j = 1; j <= MaxNumSamples; ++j)
+			{
+				double NearestDistSqr = 0.f;
+				const float MaxSampleDist = 100; // Value is arbitrary, should probably be based on the bounds of the mesh
+				const FVector SamplePosition = VertexPosition - VertexNormal * MaxSampleDist / j;
+				const int32 NearestVertexID = DynamicMeshAABBTree.FindNearestVertex(SamplePosition, NearestDistSqr, MaxSampleDist);
+				if (NearestVertexID == INDEX_NONE)
+				{
+					continue;
+				}
+
+				const FVector OtherVertexPosition = DynamicMesh.GetVertex(NearestVertexID);
+				const FVector OtherVertexNormal = FVector(DynamicMesh.GetVertexNormal(NearestVertexID));
+
+				const float Dot = OtherVertexNormal.Dot(VertexNormal);
+				const bool bIsOpposingAndCollinear = Dot < CollinearLimit;
+				if (!bIsOpposingAndCollinear)
+				{
+					continue;
+				}
+
+				const FVector OpposingPosition = FVector::PointPlaneProject(VertexPosition, OtherVertexPosition, -VertexNormal);
+				const FVector Center = (VertexPosition + OpposingPosition) * 0.5;
+				Scales[i] = (VertexPosition - Center).Size();
+				NewVertexPositions[i] = Center;
+				break;
+			}
+		}
+	}
+
+	// Remove extreme radius values
+	{
+		// In the below example, point "a" may have a normal which causes the ray to hit point "b" which is not the outcome we would expect.
+		// We solve this by taking the average radius of all vertices and removing any radius which lies way outside the average.
+		// This works well as long as the branches we're operating on have a fairly uniform size.
+		//
+		//       /    /
+		//      /  b*/
+		//     /    /
+		//    /    /
+		//    \*a /
+		//   /   /
+		//  /    \
+		// /    /
+		 
+		int32 RadiiNum = 0;
+		float AvgRadius = 0;
+		for (float Scale : Scales)
+		{
+			if (Scale > 0)
+			{
+				RadiiNum++;
+				AvgRadius += Scale;
+			}
+		}
+
+		if (RadiiNum > 0)
+		{
+			AvgRadius /= RadiiNum;
+				
+			for (int32 i = 0; i < ComponentVertexIndices.Num(); ++i)
+			{
+				const float Scale = Scales[i];
+				if (Scale > 0)
+				{
+					const float Ratio = Scale / AvgRadius;
+					if (Ratio > 3)
+					{
+						Scales[i] = 0;
+					}
+				}
+			}
+		}
+	}
+
+	for (int32 i = 0; i < ComponentVertexIndices.Num(); ++i)
+	{
+		if (Scales[i] > 0)
+		{
+			const int32 VertexIndex = ComponentVertexIndices[i];
+			DynamicMesh.SetVertex(VertexIndex, NewVertexPositions[i]);
+			VertexScaleAttribute.SetValue(VertexIndex, Scales[i]);
+		}
+	}
+}
+
+static TAutoConsoleVariable<int32> ComputeVerticesCenterVersion(
+	TEXT("PVE.Import.ComputeVerticesCenter"),
+	2,
+	TEXT("PVE.Import.ComputeVerticesCenter {VERSION}"),
+	ECVF_Default
+);
+
+static void ComputeVerticesCenterAndRadius_Internal(
+	UE::Geometry::FDynamicMesh3& DynamicMesh, 
+	const TArray<int32>& TriangleIds,
+	FFloatVertexAttribute& VertexScaleAttribute
+)
+{
+	TRACE_CPUPROFILER_EVENT_SCOPE(ComputeVerticesCenterAndRadius_Internal);
+
+	const int32 UseVersion = ComputeVerticesCenterVersion.GetValueOnAnyThread();
+	if (UseVersion == 1)
+	{
+		ComputeVerticesCenterAndRadius_Internal_V1(DynamicMesh, TriangleIds, VertexScaleAttribute);
+	}
+	else
+	{
+		ComputeVerticesCenterAndRadius_Internal_V2(DynamicMesh, TriangleIds, VertexScaleAttribute);
+	}
+}
+
+static void ComputeVerticesCenterAndRadius(UE::Geometry::FDynamicMesh3& DynamicMesh, FFloatVertexAttribute& VertexScaleAttribute)
+{
+	TRACE_CPUPROFILER_EVENT_SCOPE(ComputeVerticesCenterAndRadius);
+
+	// We execute this per component (detached mesh) as that often correlates to individual branches which significantly improves
+	// the results as we won't accidentally find an opposing vertex in a different branch from the one we're currently operating on. 
+	UE::Geometry::FMeshConnectedComponents ConnectedComponents(&DynamicMesh);
+	ConnectedComponents.FindConnectedTriangles(); 
+
+	ParallelFor(ConnectedComponents.Num(), [&](int32 Index)
+	{
+		ComputeVerticesCenterAndRadius_Internal(DynamicMesh, ConnectedComponents.GetComponent(Index).Indices, VertexScaleAttribute);
+	});
+}
+
+static void ExtractReductionCloud_Internal(
+	const UE::Geometry::FDynamicMesh3& DynamicMesh, 
+	const TArray<int32>& TriangleIds, 
+	float SearchRadiusInflate,
+	const FFloatVertexAttribute& VertexScaleAttribute,
+	FIntVertexAttribute& FuseToAttribute
+)
+{
+	TRACE_CPUPROFILER_EVENT_SCOPE(ExtractReductionCloud_Internal);
+		
+	UE::Geometry::FDynamicMeshAABBTree3 DynamicMeshAABBTree(&DynamicMesh, false);
+	DynamicMeshAABBTree.Build(TriangleIds);
+
+	TArray<int32> VertexIndices = PV::DynamicMeshHelper::GetUniqueVertexIndices(DynamicMesh, TriangleIds);
+
+	const int32 MaxVertexID = DynamicMesh.MaxVertexID();
+		
+	TArray<bool> Removed; 
+	Removed.SetNumZeroed(MaxVertexID);
+		
+	TArray<bool> Checked;
+	Checked.SetNumZeroed(MaxVertexID);
+		
+	TArray<float> Radii;
+	Radii.SetNumUninitialized(MaxVertexID);
+	
+	for (int32 i = 0; i < VertexIndices.Num(); i++)
+	{
+		const int32 VertexIndex = VertexIndices[i];
+		Radii[VertexIndex] = VertexScaleAttribute.GetValue(VertexIndex);
+	}
+
+	{
+		TRACE_CPUPROFILER_EVENT_SCOPE(SortByRadius);
+		VertexIndices.Sort([&](int32 A, int32 B) { return Radii[A] > Radii[B]; });
+	}
+
+	for (int32 i = 0; i < VertexIndices.Num(); i++)
+	{
+		const int32 VertexIndex = VertexIndices[i];
+
+		if (Removed[VertexIndex])
+		{ 
+			continue; 
+		}
+
+		Checked[VertexIndex] = true;
+
+		const FVector VertexPosition = DynamicMesh.GetVertex(VertexIndex);
+		const float VertexRadius = Radii[VertexIndex];
+		if (VertexRadius <= 0)
+		{
+			FuseToAttribute.SetValue(VertexIndex, VertexIndex);
+			continue;
+		}
+
+		const float SearchRadius = VertexRadius * SearchRadiusInflate;
+		const TArray<int32> OverlappingVertices = PV::DynamicMeshHelper::FindVerticesInRadius(DynamicMeshAABBTree, VertexPosition, SearchRadius);
+		for (int32 j = 0; j < OverlappingVertices.Num(); j++)
+		{
+			const int32 OverlappingVertexIndex = OverlappingVertices[j];
+				
+			if (OverlappingVertexIndex == VertexIndex 
+				|| Removed[OverlappingVertexIndex] 
+				|| Checked[OverlappingVertexIndex])
+			{
+				continue;
+			}
+
+			const float OverlappingVertexRadius = Radii[OverlappingVertexIndex];
+
+			if (OverlappingVertexRadius > 0)
+			{
+				const FVector OverlappingVertexPosition = DynamicMesh.GetVertex(OverlappingVertexIndex);
+				const float CombinedRadius = OverlappingVertexRadius + VertexRadius;
+				const float CombinedRadiusSqr = CombinedRadius * CombinedRadius;
+				const float DistSquared = (OverlappingVertexPosition - VertexPosition).SizeSquared();
+				if (DistSquared > CombinedRadiusSqr)
+				{
+					continue;
+				}
+			}
+
+			Removed[OverlappingVertexIndex] = true;
+			FuseToAttribute.SetValue(OverlappingVertexIndex, VertexIndex);
+		}
+
+		FuseToAttribute.SetValue(VertexIndex, VertexIndex);
+	}
+}
+	
+static void FuseVerticesUsingFuseAttribute(
+	UE::Geometry::FDynamicMesh3& DynamicMesh,
+	const FIntVertexAttribute& FuseToAttribute,
+	FBooleanVertexAttribute& FillerVertexAttribute
+)
+{
+	TRACE_CPUPROFILER_EVENT_SCOPE(FuseVerticesUsingFuseAttribute);
+
+	TArray<int32> TriangleIndices;
+	TriangleIndices.Reserve(DynamicMesh.TriangleCount());
+	for (int32 TID : DynamicMesh.TriangleIndicesItr())
+	{
+		TriangleIndices.Add(TID);
+	}
+
+	for (int32 i = 0; i < TriangleIndices.Num(); i++)
+	{
+		const int32 TriangleIndex = TriangleIndices[i];
+		const UE::Geometry::FIndex3i VertexIndices = DynamicMesh.GetTriangle(TriangleIndex);
+
+		TArray<int32, TInlineAllocator<3>> NewTriangleIndices;
+		for (int32 j = 0; j < 3; j++) 
+		{
+			const int32 FuseTo = FuseToAttribute.GetValue(VertexIndices[j]);
+			NewTriangleIndices.AddUnique(FuseTo);
+
+			FillerVertexAttribute.SetValue(VertexIndices[j], false);
+		}
+
+		DynamicMesh.RemoveTriangle(TriangleIndex, false);
+
+		if (NewTriangleIndices.Num() == 2)
+		{ 
+			const FVector P1 = DynamicMesh.GetVertex(NewTriangleIndices[0]);
+			const FVector P2 = DynamicMesh.GetVertex(NewTriangleIndices[1]);
+
+			const FVector P3 = GenerateFillerVertex(P1, P2);
+			const int32 P3Index = DynamicMesh.AppendVertex(P3);
+			DynamicMesh.AppendTriangle(NewTriangleIndices[0], NewTriangleIndices[1], P3Index);
+
+			FillerVertexAttribute.SetValue(P3Index, true);
+		}
+		else if (NewTriangleIndices.Num() == 3)
+		{
+			DynamicMesh.AppendTriangle(NewTriangleIndices[0], NewTriangleIndices[1], NewTriangleIndices[2]);
+		}
+		else if (NewTriangleIndices.Num() > 3)
+		{
+			check(false); // This should not be able to happen
+		}
+	}
+
+	RemoveUnusedVertices(DynamicMesh);
+}
+
+static void CollapseBranchesToLine(
+	UE::Geometry::FDynamicMesh3& DynamicMesh,
+	const FFloatVertexAttribute& VertexScaleAttribute,
+	FIntVertexAttribute& FuseToAttribute,
+	FBooleanVertexAttribute& FillerVertexAttribute
+)
+{
+	TRACE_CPUPROFILER_EVENT_SCOPE(CollapseBranchesToLine);
+
+	UE::Geometry::FMeshConnectedComponents ConnectedComponents(&DynamicMesh);
+	ConnectedComponents.FindConnectedTriangles();
+
+	ParallelFor(ConnectedComponents.Num(), [&](int32 Index)
+	{
+		ExtractReductionCloud_Internal(DynamicMesh, ConnectedComponents.GetComponent(Index).Indices, 2.f, VertexScaleAttribute, FuseToAttribute);
+	});
+
+	FuseVerticesUsingFuseAttribute(DynamicMesh, FuseToAttribute, FillerVertexAttribute);
+}
+
+static void FindEndPoints(
+	const UE::Geometry::FDynamicMesh3& DynamicMesh,
+	FBooleanVertexAttribute& EndPointAttribute,
+	const FBooleanVertexAttribute& FillerVertexAttribute
+)
+{
+	TRACE_CPUPROFILER_EVENT_SCOPE(FindEndPoints);
+	
+	TArray<double> xComps;
+	TArray<double> yComps;
+	TArray<double> zComps;
+	for (int32 VertexIndex : DynamicMesh.VertexIndicesItr())
+	{
+		const FVector VertexPosition = DynamicMesh.GetVertex(VertexIndex);
+
+		if (FillerVertexAttribute.GetValue(VertexIndex))
+		{
+			continue;
+		}
+
+		const int32 EdgeCount = DynamicMesh.GetVtxEdgeCount(VertexIndex);
+		check(EdgeCount > 0);
+
+		xComps.Reset(EdgeCount);
+		yComps.Reset(EdgeCount);
+		zComps.Reset(EdgeCount);
+
+		for (int32 EdgeID : DynamicMesh.VtxEdgesItr(VertexIndex))
+		{
+			const UE::Geometry::FIndex2i VertexIDs = DynamicMesh.GetEdgeV(EdgeID);
+			const int32 OtherVertexID = VertexIDs.A == VertexIndex ? VertexIDs.B : VertexIDs.A;
+
+			if (FillerVertexAttribute.GetValue(OtherVertexID))
+			{
+				continue;
+			}
+
+			const FVector NeighbourVertexPosition = DynamicMesh.GetVertex(OtherVertexID);
+			const FVector Dir = (NeighbourVertexPosition - VertexPosition).GetSafeNormal();
+			xComps.Add(Dir.X);
+			yComps.Add(Dir.Y);
+			zComps.Add(Dir.Z);
+		}
+		
+		if (!ensure(xComps.Num() > 0))
+		{
+			continue;
+		}
+
+		const double* MaxXComp = Algo::MaxElement(xComps);
+		const double* MinXComp = Algo::MinElement(xComps);
+		const double* MaxYComp = Algo::MaxElement(yComps);
+		const double* MinYComp = Algo::MinElement(yComps);
+		const double* MaxZComp = Algo::MaxElement(zComps);
+		const double* MinZComp = Algo::MinElement(zComps);
+
+		const double xDiff = *MaxXComp - *MinXComp;
+		const double yDiff = *MaxYComp - *MinYComp;
+		const double zDiff = *MaxZComp - *MinZComp;
+		const double MaxDiff = FMath::Max3(xDiff, yDiff, zDiff);
+
+		const bool bIsEndPoint = MaxDiff <= 1.0;
+		EndPointAttribute.SetValue(VertexIndex, bIsEndPoint);
+	}
+}
+
+static void ConnectBranchesToParent(
+	UE::Geometry::FDynamicMesh3& DynamicMesh,
+	const FFloatVertexAttribute& VertexScaleAttribute,
+	FBooleanVertexAttribute& EndPointAttribute,
+	FBooleanVertexAttribute& FillerVertexAttribute
+)
+{
+	TRACE_CPUPROFILER_EVENT_SCOPE(ConnectBranchesToParent);
+
+	// This algorithm takes the tips of all the branches which have been individually collapsed into a line, and tries to connect them to their parent.
+	// This is done by performing an overlap query at the tip of each end point, and connecting that tip to the closest overlapping vertex which
+	// is not part of the same mesh component.
+	//
+	// Issues with this algorithm:
+	// 1. Since we do not know what is the "tip" and the "base" of the branch at this point, we perform the query for each
+	// end point and just hope that the tips do not overlap any other vertex. They could however be overlapping, which would lead to the tips becoming
+	// fused and for loops to be created in the tree.
+	// 2. We only fuse the end point to the parent if the radii of the end point and the nearby vertex overlap. This may lead to endpoints not being
+	// connected if the closest vertex is too far away (even if the end-point is very close to the parent mesh itself). This would lead to the entire
+	// branch being excluded from the final tree.
+
+	const UE::Geometry::FDynamicMeshAABBTree3 DynamicMeshAABBTree(&DynamicMesh);
+	UE::Geometry::FMeshConnectedComponents ConnectedComponents(&DynamicMesh);
+	ConnectedComponents.FindConnectedVertices();
+
+	TMap<int32, int32> VertexComponentMap;
+	VertexComponentMap.Reserve(DynamicMesh.VertexCount());
+
+	for (int32 i = 0; i < ConnectedComponents.Components.Num(); ++i)
+	{
+		for (int32 VertexID : ConnectedComponents.Components[i].Indices)
+		{
+			VertexComponentMap.Add(VertexID, i);
+		}
+	}
+
+	const TArray<int32> EndPoints = EndPointAttribute.FindAllNonZero();
+	TArray<int32> MergeMap;
+	MergeMap.SetNum(EndPoints.Num());
+	for (int32 i = 0; i < EndPoints.Num(); ++i)
+	{
+		const int32 EndPointVertexIndex = EndPoints[i];
+		
+		const FVector VertexPosition = DynamicMesh.GetVertex(EndPointVertexIndex);
+		const float VertexRadius = VertexScaleAttribute.GetValue(EndPointVertexIndex);
+		const float SearchRadius = VertexRadius * 4.f;
+		const TArray<int32> OverlappingVertices = PV::DynamicMeshHelper::FindVerticesInRadius(DynamicMeshAABBTree, VertexPosition, SearchRadius);
+
+		float ClosestVertexDistSqr = MAX_flt;
+		int32 VertexToMergeWith = INDEX_NONE;
+		for (int32 OverlappingVertexIndex : OverlappingVertices)
+		{
+			if (EndPointAttribute.GetValue(OverlappingVertexIndex)
+				|| FillerVertexAttribute.GetValue(OverlappingVertexIndex))
+			{
+				continue;
+			}
+
+			int32* ComponentIndex1 = VertexComponentMap.Find(EndPointVertexIndex);
+			int32* ComponentIndex2 = VertexComponentMap.Find(OverlappingVertexIndex);
+			check(ComponentIndex1 && ComponentIndex2);
+			if (*ComponentIndex1 == *ComponentIndex2)
+			{
+				continue;
+			}
+
+			const float OverlappingVertexRadius = VertexScaleAttribute.GetValue(OverlappingVertexIndex);
+			const FVector OverlappingVertexPosition = DynamicMesh.GetVertex(OverlappingVertexIndex);
+			const float CombinedRadius = OverlappingVertexRadius + VertexRadius;
+			const float CombinedRadiusSqr = CombinedRadius * CombinedRadius;
+			const float DistSquared = (OverlappingVertexPosition - VertexPosition).SizeSquared();
+			if (DistSquared > CombinedRadiusSqr)
+			{
+				continue;
+			}
+
+			const FVector OtherVertexPosition = DynamicMesh.GetVertex(OverlappingVertexIndex);
+			const float DistSqr = (VertexPosition - OtherVertexPosition).SizeSquared();
+			if (DistSqr < ClosestVertexDistSqr)
+			{
+				VertexToMergeWith = OverlappingVertexIndex;
+				ClosestVertexDistSqr = DistSqr;
+			}
+		}
+
+		MergeMap[i] = VertexToMergeWith;
+	}
+
+	for (int32 i = 0; i < EndPoints.Num(); ++i)
+	{
+		const int32 EndPointVertexIndex = EndPoints[i];
+		const int32 VertexToMergeWith = MergeMap[i];
+
+		if (VertexToMergeWith == INDEX_NONE)
+		{
+			continue;
+		}
+		const FVector VertexPosition = DynamicMesh.GetVertex(EndPointVertexIndex);
+		const FVector OtherVertexPosition = DynamicMesh.GetVertex(VertexToMergeWith);
+
+		const FVector Normal = (VertexPosition - OtherVertexPosition).GetSafeNormal();
+		const FVector NewVertexPosition = OtherVertexPosition + (Normal * 0.1); // We want to keep this vertex to not lose the radius information, so just move it really close and connect it to the other vertex
+
+		DynamicMesh.SetVertex(EndPointVertexIndex, NewVertexPosition);
+		
+		const TArray<int32> Neighbours = PV::DynamicMeshHelper::GetVertexNeighbours(DynamicMesh, EndPointVertexIndex);
+		for (int32 Neighbour : Neighbours)
+		{
+			if (FillerVertexAttribute.GetValue(Neighbour))
+			{
+				const TArray<int32> FillerNeighbours = PV::DynamicMeshHelper::GetVertexNeighbours(DynamicMesh, Neighbour);
+				const int32* AdjacentVertex = FillerNeighbours.FindByPredicate([&](const auto& Value) { return Value != EndPointVertexIndex && Value != Neighbour; });
+				check(AdjacentVertex != nullptr);
+				const FVector AdjacentVertexPosition = DynamicMesh.GetVertex(*AdjacentVertex);
+				const FVector NewFillerVertex = GenerateFillerVertex(AdjacentVertexPosition, NewVertexPosition);
+				
+				DynamicMesh.SetVertex(Neighbour, NewFillerVertex);
+			}
+		}
+		
+		const FVector NewFillerVertex = GenerateFillerVertex(NewVertexPosition, OtherVertexPosition);
+		const int32 NewFillerVertexIndex = DynamicMesh.AppendVertex(NewFillerVertex);
+		DynamicMesh.AppendTriangle(EndPointVertexIndex, VertexToMergeWith, NewFillerVertexIndex);
+		EndPointAttribute.SetValue(EndPointVertexIndex, false);
+	}
+}
+
+static void FindRootPoints(
+	const UE::Geometry::FDynamicMesh3& DynamicMesh,
+	FBooleanVertexAttribute& RootPointAttribute,
+	FBooleanVertexAttribute& EndPointAttribute
+)
+{
+	TRACE_CPUPROFILER_EVENT_SCOPE(FindRootPoints);
+
+	// This function currently just finds all the end points that are close to the bottom of the plant bounding box and converts them to root points.
+	// This could result in issues if a plant has multiple points touching the floor. 
+
+	const UE::Geometry::FAxisAlignedBox3d BoxBounds = DynamicMesh.GetBounds();
+
+	for (int32 VertexIndex : EndPointAttribute.FindAllNonZero())
+	{
+		const FVector VertexPosition = DynamicMesh.GetVertex(VertexIndex);
+
+		const float Z = FMath::GetMappedRangeValueClamped(FVector2D(BoxBounds.Min.Z, BoxBounds.Max.Z), FVector2D(0, 1), VertexPosition.Z);
+		const bool bIsRootPoint = Z <= 0.1;
+		if (bIsRootPoint)
+		{
+			RootPointAttribute.SetValue(VertexIndex, true);
+			EndPointAttribute.SetValue(VertexIndex, false);			
+		}
+	}
+}
+
+PV::StaticMeshImport::EImportResult PV::StaticMeshImport::ImportGrowthDataFromStaticMesh(const FPVImportStaticMeshParams& InParams, FManagedArrayCollection& OutCollection)
+{
+	using namespace PV::DynamicMeshHelper;
+	
+	TRACE_CPUPROFILER_EVENT_SCOPE(PV::StaticMeshImport::ImportGrowthDataFromStaticMesh);
+	
+	if (!InParams.StaticMeshAsset)
+	{
+		return EImportResult::InvalidSourceMesh;
+	}
+
+	if (InParams.StaticMeshAsset->IsCompiling())
+	{
+		return EImportResult::MeshNotReady;
+	}
+
+	if (InParams.MaterialsToKeep.Num() == 0)
+	{
+		return EImportResult::InvalidMaterialFilter;
+	}
+
+	UE::Geometry::FDynamicMesh3 DynamicMesh;
+	{
+		TRACE_CPUPROFILER_EVENT_SCOPE(StaticMeshToDynamicMesh);
+
+		FText OutErrorMessage;
+		UE::Conversion::FStaticMeshConversionOptions Options;
+		if (!UE::Conversion::StaticMeshToDynamicMesh(InParams.StaticMeshAsset, DynamicMesh, OutErrorMessage, Options))
+		{
+			return EImportResult::InvalidMeshData;
+		}
+	}
+
+	if (InParams.DebugState == EPVImportStaticMeshDebugState::InitialMesh)
+	{
+		AddDynamicMeshToCollection(DynamicMesh, OutCollection);
+		return EImportResult::Success;
+	}
+	
+	if (!FilterTrianglesUsingMaterials(DynamicMesh, InParams.StaticMeshAsset, InParams.MaterialsToKeep))
+	{
+		return EImportResult::InvalidMaterialFilter;
+	}
+	
+	if (InParams.DebugState == EPVImportStaticMeshDebugState::KeepMaterial)
+	{
+		AddDynamicMeshToCollection(DynamicMesh, OutCollection);
+		return EImportResult::Success;
+	}
+
+	ComputeNormals(DynamicMesh);
+	if (InParams.DebugState == EPVImportStaticMeshDebugState::ComputeNormals)
+	{
+		AddDynamicMeshToCollection(DynamicMesh, OutCollection);
+		return EImportResult::Success;
+	}
+
+	ComputeVerticesCenterAndRadius(
+		DynamicMesh,
+		VertexScaleAttributeDefinition.GetOrAttachAttribute(DynamicMesh)
+	);
+	if (InParams.DebugState == EPVImportStaticMeshDebugState::ComputeVerticesCenterAndRadius)
+	{
+		AddDynamicMeshToCollection(DynamicMesh, OutCollection);
+		VertexScaleAttributeDefinition.CopyAttributeToCollection(DynamicMesh, OutCollection);
+		return EImportResult::Success;
+	}
+
+	CollapseBranchesToLine(
+		DynamicMesh,
+		VertexScaleAttributeDefinition.GetOrAttachAttribute(DynamicMesh),
+		FuseToAttributeDefinition.GetOrAttachAttribute(DynamicMesh),
+		FillerVertexAttributeDefinition.GetOrAttachAttribute(DynamicMesh)
+	);
+	if (InParams.DebugState == EPVImportStaticMeshDebugState::CollapseBranchesToLine)
+	{
+		AddDynamicMeshToCollection(DynamicMesh, OutCollection);
+		FuseToAttributeDefinition.CopyAttributeToCollection(DynamicMesh, OutCollection);
+		VertexScaleAttributeDefinition.CopyAttributeToCollection(DynamicMesh, OutCollection);
+		return EImportResult::Success;
+	}
+
+	FindEndPoints(
+		DynamicMesh,
+		EndPointAttributeDefinition.GetOrAttachAttribute(DynamicMesh),
+		FillerVertexAttributeDefinition.GetOrAttachAttribute(DynamicMesh)
+	);
+	if (InParams.DebugState == EPVImportStaticMeshDebugState::FindEndPoints)
+	{
+		AddDynamicMeshToCollection(DynamicMesh, OutCollection);
+		EndPointAttributeDefinition.CopyAttributeToCollection(DynamicMesh, OutCollection);
+		return EImportResult::Success;
+	}
+
+	ConnectBranchesToParent(
+		DynamicMesh,
+		VertexScaleAttributeDefinition.GetOrAttachAttribute(DynamicMesh),
+		EndPointAttributeDefinition.GetOrAttachAttribute(DynamicMesh),
+		FillerVertexAttributeDefinition.GetOrAttachAttribute(DynamicMesh)
+	);
+	if (InParams.DebugState == EPVImportStaticMeshDebugState::ConnectBranchesToParent)
+	{
+		AddDynamicMeshToCollection(DynamicMesh, OutCollection);
+		EndPointAttributeDefinition.CopyAttributeToCollection(DynamicMesh, OutCollection);
+		return EImportResult::Success;
+	}
+
+	FindRootPoints(
+		DynamicMesh,
+		RootPointAttributeDefinition.GetOrAttachAttribute(DynamicMesh),
+		EndPointAttributeDefinition.GetOrAttachAttribute(DynamicMesh)
+	);
+	if (InParams.DebugState == EPVImportStaticMeshDebugState::FindRootPoints)
+	{
+		AddDynamicMeshToCollection(DynamicMesh, OutCollection);
+		RootPointAttributeDefinition.CopyAttributeToCollection(DynamicMesh, OutCollection);
+		EndPointAttributeDefinition.CopyAttributeToCollection(DynamicMesh, OutCollection);
+		return EImportResult::Success;
+	}
+
+	auto& FillerVertexAttribute = FillerVertexAttributeDefinition.GetOrAttachAttribute(DynamicMesh);
+	TArray<FMeshBranchHierarchy> DynamicMeshBranchHierarchies;
+	ComputeDynamicMeshBranchHierarchies(
+		DynamicMesh,
+		RootPointAttributeDefinition.GetOrAttachAttribute(DynamicMesh),
+		EndPointAttributeDefinition.GetOrAttachAttribute(DynamicMesh),
+		VertexScaleAttributeDefinition.GetOrAttachAttribute(DynamicMesh),
+		LengthFromRootAttributeDefinition.GetOrAttachAttribute(DynamicMesh),
+		NextBranchPointAttributeDefinition.GetOrAttachAttribute(DynamicMesh),
+		[&](int32 VertexID) { return !FillerVertexAttribute.GetValue(VertexID); },
+		DynamicMeshBranchHierarchies
+	);
+
+	TArray<FPVBranchHierarchyDescription> BranchHierarchyDescriptions;
+	ConvertDynamicMeshBranchHierarchiesToBranchDescription(
+		DynamicMesh,
+		VertexScaleAttributeDefinition.GetOrAttachAttribute(DynamicMesh),
+		EndPointAttributeDefinition.GetOrAttachAttribute(DynamicMesh),
+		DynamicMeshBranchHierarchies,
+		BranchHierarchyDescriptions
+	);
+
+	if (InParams.DebugState == EPVImportStaticMeshDebugState::ComputeBranchHierarchies)
+	{
+		AddDynamicMeshToCollection(DynamicMesh, OutCollection);
+		AddBranchHierarchiesToCollection(BranchHierarchyDescriptions, OutCollection);
+		RootPointAttributeDefinition.CopyAttributeToCollection(DynamicMesh, OutCollection);
+		EndPointAttributeDefinition.CopyAttributeToCollection(DynamicMesh, OutCollection);
+		VertexScaleAttributeDefinition.CopyAttributeToCollection(DynamicMesh, OutCollection);
+		LengthFromRootAttributeDefinition.CopyAttributeToCollection(DynamicMesh, OutCollection);
+		NextBranchPointAttributeDefinition.CopyAttributeToCollection(DynamicMesh, OutCollection);
+
+		return EImportResult::Success;
+	}
+
+	if (!PV::ImportHelper::GenerateGrowthDataFromBranchHierarchies(OutCollection, BranchHierarchyDescriptions))
+	{
+		return EImportResult::FailedToGenerateGrowthData;
+	}
+
+	return EImportResult::Success;
+}
\ No newline at end of file
diff --git a/Engine/Plugins/Experimental/ProceduralVegetationEditor/Source/ProceduralVegetation/Private/Implementations/PVImporter_StaticMesh.h b/Engine/Plugins/Experimental/ProceduralVegetationEditor/Source/ProceduralVegetation/Private/Implementations/PVImporter_StaticMesh.h
new file mode 100644
--- /dev/null
+++ b/Engine/Plugins/Experimental/ProceduralVegetationEditor/Source/ProceduralVegetation/Private/Implementations/PVImporter_StaticMesh.h
@@ -0,0 +1,25 @@
+// Copyright Epic Games, Inc. All Rights Reserved.
+#pragma once
+
+#include "CoreMinimal.h"
+
+#include "Params/PVImportStaticMeshParams.h"
+#include "GeometryCollection/ManagedArrayCollection.h"
+
+class UStaticMesh;
+
+namespace PV::StaticMeshImport
+{
+	enum class EImportResult
+	{
+		Success,
+		InvalidSourceMesh,
+		InvalidOutputPath,
+		InvalidMeshData,
+		MeshNotReady,
+		InvalidMaterialFilter,
+		FailedToGenerateGrowthData
+	};
+
+	EImportResult ImportGrowthDataFromStaticMesh(const FPVImportStaticMeshParams& InParams, FManagedArrayCollection& OutCollection);
+};
diff --git a/Engine/Plugins/Experimental/ProceduralVegetationEditor/Source/ProceduralVegetation/Private/Implementations/PVImporter_Texture2D.cpp b/Engine/Plugins/Experimental/ProceduralVegetationEditor/Source/ProceduralVegetation/Private/Implementations/PVImporter_Texture2D.cpp
new file mode 100644
--- /dev/null
+++ b/Engine/Plugins/Experimental/ProceduralVegetationEditor/Source/ProceduralVegetation/Private/Implementations/PVImporter_Texture2D.cpp
@@ -0,0 +1,1203 @@
+// Copyright Epic Games, Inc. All Rights Reserved.
+
+#include "Implementations/PVImporter_Texture2D.h"
+#include "Helpers/PVImportHelpers.h"
+#include "Utils/PVDynamicMeshVertexAttribute.h"
+
+#include "IntVectorTypes.h"
+#include "PCGSettings.h"
+#include "GeometryCollection/ManagedArrayCollection.h"
+#include "Helpers/PVAttributesHelper.h"
+
+#include "Engine/Texture.h"
+#include "Engine/Texture2D.h"
+#include "Polygon2.h"
+
+#include "Materials/MaterialInterface.h"
+
+#include "DynamicMesh/DynamicMesh3.h"
+#include "DynamicMesh/DynamicMeshAABBTree3.h"
+#include "GeometryCollection/GeometryCollection.h"
+#include "GeometryCollection/GeometryCollectionUtility.h"
+#include "SpatialAlgo/PCGMarchingSquares.h"
+#include "CompGeom/Delaunay2.h"
+#include "Selections/MeshConnectedComponents.h"
+
+namespace Internal
+{
+	using FBooleanVertexAttribute = PV::DynamicMeshHelper::FBooleanVertexAttribute;
+	using FFloatVertexAttribute = PV::DynamicMeshHelper::FFloatVertexAttribute; 
+	using FIndex2i = UE::Geometry::FIndex2i;
+
+	using FIndex3i = UE::Geometry::FIndex3i;
+	using FVector2i = UE::Geometry::FVector2i;
+	using FMeshConnectedComponents = UE::Geometry::FMeshConnectedComponents;
+	using FDynamicMeshAABBTree3 = UE::Geometry::FDynamicMeshAABBTree3;
+	using FDynamicMesh3 = UE::Geometry::FDynamicMesh3;
+	using FDynamicMeshQueries3 = UE::Geometry::TMeshQueries<FDynamicMesh3>;
+	using FPolygon2f = UE::Geometry::FPolygon2f;
+	using FDelaunay2 = UE::Geometry::FDelaunay2;
+
+	const static FName PixelsDebugGroupName = TEXT("PixelDebug");
+	const static FName PixelPositionsAttributeName = TEXT("Position");
+	const static FName PixelValueAttributeName = TEXT("Value");
+	const static FName PerimeterCurveGroupName = TEXT("PerimeterCurves");
+	const static FName PerimeterCurveAttributeName = TEXT("PerimeterCurve");
+	const static FName InputRootPointsGroupName = TEXT("InputRootPoints");
+	const static FName InputRootPointsAttributeName = TEXT("Point");
+	const static FName LabelsGroupName = TEXT("Labels");
+	const static FName LabelPositionAttributeName = TEXT("Position");
+	const static FName LabelScaleAttributeName = TEXT("Scale");
+	const static FName LabelTextAttributeName = TEXT("Text");
+
+	const static PV::TDynamicMeshVertexAttributeDefinition<bool> PerimeterAttributeDefinition(TEXT("Perimeter"));
+	
+	static int32 GetPixelIndex(const FVector2i& PixelCoords, int32 ImageWidth)
+	{
+		return PixelCoords.X + PixelCoords.Y * ImageWidth;
+	}
+
+	static FVector2f PixelCoordToUV(const FVector2i& PixelCoord, const FVector2i& ImageSize)
+	{
+		const float U = ImageSize.X > 1 ? PixelCoord.X / static_cast<float>(ImageSize.X - 1) : 0;
+		const float V = ImageSize.Y > 1 ? PixelCoord.Y / static_cast<float>(ImageSize.Y - 1) : 0;
+		return FVector2f(U, V);
+	}
+	
+	template<typename Vector3dType = float, typename Vector2dType = float>
+	UE::Math::TVector<Vector3dType> ToVector3(const UE::Math::TVector2<Vector2dType>& Vec2D, float Scale)
+	{
+		return UE::Math::TVector<Vector3dType>(0, Vec2D.X * Scale, (1.0 - Vec2D.Y) * Scale);
+	}
+
+	static void AddLabelsToCollection(FManagedArrayCollection& InOutCollection, const TArray<TTuple<FVector3f, FString, float>>& Labels)
+	{
+		if (!InOutCollection.HasGroup(LabelsGroupName))
+		{
+			InOutCollection.AddGroup(LabelsGroupName);
+		}
+
+		if (!InOutCollection.HasAttribute(LabelTextAttributeName, LabelsGroupName))
+		{
+			InOutCollection.AddAttribute<FString>(LabelTextAttributeName, LabelsGroupName);
+		}
+
+		if (!InOutCollection.HasAttribute(LabelPositionAttributeName, LabelsGroupName))
+		{
+			InOutCollection.AddAttribute<FVector3f>(LabelPositionAttributeName, LabelsGroupName);
+		}
+
+		if (!InOutCollection.HasAttribute(LabelScaleAttributeName, LabelsGroupName))
+		{
+			InOutCollection.AddAttribute<float>(LabelScaleAttributeName, LabelsGroupName);
+		}
+
+		if (!ensure(InOutCollection.GetAttributeType(LabelTextAttributeName, LabelsGroupName) == FManagedArrayCollection::EArrayType::FStringType))
+		{
+			return;
+		}
+
+		if (!ensure(InOutCollection.GetAttributeType(LabelPositionAttributeName, LabelsGroupName) == FManagedArrayCollection::EArrayType::FVectorType))
+		{
+			return;
+		}
+
+		if (!ensure(InOutCollection.GetAttributeType(LabelScaleAttributeName, LabelsGroupName) == FManagedArrayCollection::EArrayType::FFloatType))
+		{
+			return;
+		}
+
+		const int32 StartIndex = InOutCollection.NumElements(LabelsGroupName);
+
+		InOutCollection.AddElements(Labels.Num(), LabelsGroupName);
+
+		TManagedArray<FString>& LabelTextAttribute = *InOutCollection.FindAttribute<FString>(LabelTextAttributeName, LabelsGroupName);
+		TManagedArray<FVector3f>& LabelPositionAttribute = *InOutCollection.FindAttribute<FVector3f>(LabelPositionAttributeName, LabelsGroupName);
+		TManagedArray<float>& LabelScaleAttribute = *InOutCollection.FindAttribute<float>(LabelScaleAttributeName, LabelsGroupName);
+
+		for (int32 i = 0; i < Labels.Num(); ++i)
+		{
+			LabelPositionAttribute[StartIndex + i] = Labels[i].Get<FVector3f>();
+			LabelTextAttribute[StartIndex + i] = Labels[i].Get<FString>();
+			LabelScaleAttribute[StartIndex + i] = Labels[i].Get<float>();
+		}
+	}
+
+	static void AddPlantSettingsToCollection(FManagedArrayCollection& InOutCollection, const TArray<FPVImportTexture2DPlantSettings>& InPlantSettings, float Scale, const FVector& Offset = FVector::ZeroVector)
+	{
+		InOutCollection.AddGroup(InputRootPointsGroupName);
+		InOutCollection.AddElements(InPlantSettings.Num(), InputRootPointsGroupName);
+		
+		TArray<TTuple<FVector3f, FString, float>> Labels;
+		Labels.Reserve(InPlantSettings.Num());
+
+		auto& InputRootPointAttribute = InOutCollection.AddAttribute<FVector3f>(InputRootPointsAttributeName, InputRootPointsGroupName);
+		for (int32 i = 0; i < InPlantSettings.Num(); ++i)
+		{
+			const FVector3f PointPosition = ToVector3(InPlantSettings[i].RootPosition, Scale) + FVector3f(Offset);
+			InputRootPointAttribute[i] = PointPosition;
+			Labels.Emplace(PointPosition, FString::Printf(TEXT("%d"), i), 1.f);
+		}
+
+		AddLabelsToCollection(InOutCollection, Labels);
+	}
+	
+	static void AddDebugPixelValuesToCollection(FManagedArrayCollection& InOutCollection, const TArray<bool>& ImageData, const FVector2i& ImageSize, float Scale)
+	{
+		check(ImageData.Num() == ImageSize.X * ImageSize.Y);
+		
+		InOutCollection.AddGroup(PixelsDebugGroupName);
+		InOutCollection.AddElements(ImageData.Num(), PixelsDebugGroupName);
+		auto& PixelPositionsAttribute = InOutCollection.AddAttribute<FVector3f>(PixelPositionsAttributeName, PixelsDebugGroupName);
+		auto& PixelValueAttribute = InOutCollection.AddAttribute<bool>(PixelValueAttributeName, PixelsDebugGroupName);
+
+		for (int32 x = 0; x < ImageSize.X; x++)
+		{
+			for (int32 y = 0; y < ImageSize.Y; y++)
+			{
+				const FVector2i PixelCoords = FVector2i(x, y);
+				const FVector2f UV = PixelCoordToUV(PixelCoords, ImageSize);
+				
+				const int32 PixelIndex = GetPixelIndex(PixelCoords, ImageSize.X);
+				PixelPositionsAttribute[PixelIndex] = ToVector3(UV, Scale);
+				PixelValueAttribute[PixelIndex] = ImageData[PixelIndex];
+			}
+		}
+	}
+
+	static void AddPerimeterCurvesToCollection(
+		FManagedArrayCollection& InOutCollection, 
+		const TArray<FPolygon2f>& InPerimeterCurves,
+		float Scale, 
+		const FVector& Offset = FVector::ZeroVector
+	)
+	{
+		InOutCollection.AddGroup(PerimeterCurveGroupName);
+		InOutCollection.AddElements(InPerimeterCurves.Num(), PerimeterCurveGroupName);
+    	
+		auto& PerimeterCurveAttribute = InOutCollection.AddAttribute<TArray<FVector3f>>(PerimeterCurveAttributeName, PerimeterCurveGroupName);
+		for (int32 i = 0; i < InPerimeterCurves.Num(); i++)
+		{
+			const TArray<FVector2f>& Curve = InPerimeterCurves[i].GetVertices();
+			for (const FVector2f& CurvePos : Curve)
+			{
+				PerimeterCurveAttribute[i].Add(ToVector3(CurvePos, Scale) + FVector3f(Offset));
+			}
+		}
+	}
+
+	static void AddDynamicMeshVerticesToCollection(
+		FManagedArrayCollection& InOutCollection, 
+		const FName GroupName,
+		const FName AttributeName,
+		const UE::Geometry::FDynamicMesh3& DynamicMesh, 
+		const FBooleanVertexAttribute& Filter,
+		const FVector& Offset
+	)
+	{
+		const TArray<int32> VertexIndices = Filter.FindAllNonZero();
+		
+		InOutCollection.AddGroup(GroupName);
+		InOutCollection.AddElements(VertexIndices.Num(), GroupName);
+
+		auto& Attribute = InOutCollection.AddAttribute<FVector3f>(AttributeName, GroupName);
+		for (int32 i = 0; i < VertexIndices.Num(); ++i)
+		{
+			Attribute[i] = FVector3f(DynamicMesh.GetVertex(VertexIndices[i]) + Offset);
+		}
+	}
+
+	static void AddTextureQuadToCollection(
+		FManagedArrayCollection& InOutCollection,
+		UTexture2D& Texture2D,
+		const UMaterialInterface* TextureVisualizationMaterial,
+		float Scale,
+		const FVector& Offset = FVector::ZeroVector
+	)
+	{
+		const TArray<FVector> Points = {
+			FVector(0, 0, 1),
+			FVector(0, 1, 1),
+			FVector(0, 1, 0),
+			FVector(0, 0, 0)
+		};
+		const TArray<FVector3f> Normals = {
+			FVector3f(-1, 0, 0),
+			FVector3f(-1, 0, 0),
+			FVector3f(-1, 0, 0),
+			FVector3f(-1, 0, 0)
+		};
+		const TArray<UE::Geometry::FVector3i> Triangles = {
+			{ 0, 1, 2 },
+			{ 2, 3, 0 }
+		};
+		const TArray<FVector2f> UVs = {
+			FVector2f(0, 0),
+			FVector2f(1, 0),
+			FVector2f(1, 1),
+			FVector2f(0, 1),
+		};
+		const TSharedPtr<FGeometryCollection> MeshGeomCollection = GeometryCollection::MakeMeshElement(
+			Points,
+			Normals,
+			Triangles,
+			UVs,
+			FTransform(FRotator::ZeroRotator, Offset, FVector(Scale)),
+			FTransform::Identity,
+			1
+		);
+
+		const static FName MaterialPathAttributeName("MaterialPath");
+		auto& MaterialArray = MeshGeomCollection->AddAttribute<FString>(MaterialPathAttributeName, FGeometryCollection::MaterialGroup);
+		MaterialArray[0] = TextureVisualizationMaterial ? TextureVisualizationMaterial->GetPathName() : TEXT("None");
+
+		MeshGeomCollection->CopyTo(&InOutCollection);
+	}
+
+	static void ArrangePlantsInGrid(
+		FManagedArrayCollection& InCollection, 
+		const TArray<int32>& PlantNumberToPlantIndices,
+		const FVector& InOffset = FVector::ZeroVector
+	)
+	{
+		using namespace PV::AttributesHelper;
+
+		TArray<TTuple<FVector3f, FString, float>> Labels;
+		Labels.Reserve(PlantNumberToPlantIndices.Num());
+
+		struct FPlantOffsetInfo
+		{
+			int32 PlantNumber;
+			FBox3f Bounds;
+		};
+		TArray<FPlantOffsetInfo> PlantsOffsetInfo;
+
+		PV::FPointPositionAttributeView PointPositionAttribute = PV::FPointPositionAttribute::FindAttribute(InCollection);
+		if (!ensure(PointPositionAttribute.IsValid()))
+		{
+			return;
+		}
+
+		ForEachPlant(InCollection, [&](int32 PlantNumber)
+		{
+			FPlantOffsetInfo& PlantOffsetInfo = PlantsOffsetInfo.AddDefaulted_GetRef();
+			PlantOffsetInfo.PlantNumber = PlantNumber;
+
+			ForEachPlantPoint(InCollection, PlantNumber, [&](int32 BranchIndex, int32 PointIndex)
+			{
+				PlantOffsetInfo.Bounds += PointPositionAttribute[PointIndex];
+				return EForEachResult::Continue;
+			});
+			return EForEachResult::Continue;
+		});
+
+		FVector3f Offset = FVector3f(InOffset);
+		for (int32 i = 0; i < PlantsOffsetInfo.Num(); ++i)
+		{
+			Offset.Y -= PlantsOffsetInfo[i].Bounds.Max.Y + 10.f;
+
+			ForEachPlantPoint(InCollection, PlantsOffsetInfo[i].PlantNumber, [&](int32 BranchIndex, int32 PointIndex)
+			{
+				FVector3f& PointPosition = PointPositionAttribute[PointIndex];
+				PointPosition += Offset;
+				return EForEachResult::Continue;
+			});
+
+			const int32 PlantIndex = PlantNumberToPlantIndices.IsValidIndex(PlantsOffsetInfo[i].PlantNumber)
+				? PlantNumberToPlantIndices[PlantsOffsetInfo[i].PlantNumber]
+				: INDEX_NONE;
+
+			Labels.Emplace(Offset + FVector3f(0, 0, -10.f), FString::Printf(TEXT("%d"), PlantIndex), 1.0);
+
+			Offset.Y += PlantsOffsetInfo[i].Bounds.Min.Y;
+		}
+
+		AddLabelsToCollection(InCollection, Labels);
+	}
+
+	static bool SampleImage(UTexture2D& InTexture2D, FVector2i& InOutSampleSize, TArray<bool>& OutBlackAndWhiteImageData, bool bInvertImage, float WhiteLevel)
+	{
+#if WITH_EDITOR
+		FTextureSource& InTextureSource = InTexture2D.Source;
+		if (!InTextureSource.IsValid())
+		{
+			return false;
+		}
+
+		// FTextureSource::FMipLock is only available in editor. For the time being this functionality is not needed outside the editor despite this module being a 
+		// runtime module, so we exclude this in non-editor builds.
+
+		const FTextureSource::FMipLock LockedMip0(FTextureSource::ELockState::ReadOnly, &InTextureSource, 0);
+		check(LockedMip0.IsValid());
+
+		const int32 ImageWith =  LockedMip0.Image.GetWidth();
+		const int32 ImageHeight = LockedMip0.Image.GetHeight();
+		InOutSampleSize.X = FMath::Min(InOutSampleSize.X, ImageWith);
+		InOutSampleSize.Y  = FMath::Min(InOutSampleSize.Y, ImageHeight);
+
+		OutBlackAndWhiteImageData.SetNum(InOutSampleSize.X * InOutSampleSize.Y);
+		for (int32 x = 0; x < InOutSampleSize.X; x++)
+		{
+			for (int32 y = 0; y < InOutSampleSize.Y; y++)
+			{
+				const FVector2i PixelCoords = FVector2i(x, y);
+				const FVector2f UV = PixelCoordToUV(PixelCoords, InOutSampleSize);
+
+				const int32 SamplePosX = FMath::RoundToInt(UV.X * (ImageWith - 1));
+				const int32 SamplePosY = FMath::RoundToInt(UV.Y * (ImageHeight - 1));
+				check(SamplePosX >= 0 && SamplePosX < ImageWith);
+				check(SamplePosY >= 0 && SamplePosY < ImageHeight);
+				
+				const FLinearColor PixelColor = LockedMip0.Image.GetOnePixelLinear(SamplePosX, SamplePosY);
+				const int32 PixelIndex = GetPixelIndex(PixelCoords, InOutSampleSize.X);
+				const float Alpha = bInvertImage ? 1.0 - PixelColor.R : PixelColor.R;
+				OutBlackAndWhiteImageData[PixelIndex] = Alpha > WhiteLevel;
+			}
+		}
+
+		return true;
+#else
+		return false;
+#endif
+	}
+	
+	static void TraceImage(const TArray<bool>& ImageData, const FVector2i& ImageSize, TArray<FPolygon2f>& OutCurves)
+	{
+		TRACE_CPUPROFILER_EVENT_SCOPE(TraceImage);
+		
+		const auto ValueQuery = [&](int32 X, int32 Y)->double 
+		{ 
+			return ImageData[GetPixelIndex(FVector2i(X, Y), ImageSize.X)]; 
+		};
+		const TArray<PCGSpatialAlgo::FPCGMarchingSquareResult> Results = PCGSpatialAlgo::MarchingSquares(
+			ImageSize.X, 
+			ImageSize.Y,
+			1, 
+			ValueQuery, 
+			false
+		);
+		
+		TArray<FVector2f> PerimeterCurve;
+		OutCurves.Reserve(Results.Num());
+		for (const PCGSpatialAlgo::FPCGMarchingSquareResult& Result : Results)
+		{
+			if (!Result.bClosed)
+			{
+				continue;
+			}
+			
+			PerimeterCurve.Reset(Result.LinkedGridCoordinates.Num());
+			for (const FVector2D& Coordinate : Result.LinkedGridCoordinates)
+			{
+				PerimeterCurve.Add(FVector2f(Coordinate.X / (ImageSize.X - 1), Coordinate.Y / (ImageSize.Y - 1)));
+			}
+			OutCurves.Emplace(MoveTemp(PerimeterCurve));
+		}
+	}
+	
+	static void RemoveSmallCurves(TArray<FPolygon2f>& InOutPerimeterCurves, float Threshold)
+	{
+		const float ThresholdSqr = Threshold * Threshold;
+		for (int32 i = InOutPerimeterCurves.Num() - 1; i >= 0; i--)
+		{
+			const TArray<FVector2f>& Curve = InOutPerimeterCurves[i].GetVertices();
+			FVector2f Min = FVector2f(MAX_flt);
+			FVector2f Max = FVector2f(-MAX_flt);
+			
+			for (const FVector2f& Point : Curve)
+			{
+				Min.X = FMath::Min(Min.X, Point.X);
+				Min.Y = FMath::Min(Min.Y, Point.Y);
+				Max.X = FMath::Max(Max.X, Point.X);
+				Max.Y = FMath::Max(Max.Y, Point.Y);
+			}
+			
+			const float Area = (Max.X - Min.X) * (Max.Y - Min.Y); 
+			const float AreaSqr = Area * Area;
+			if (AreaSqr < ThresholdSqr)
+			{
+				InOutPerimeterCurves.RemoveAtSwap(i);
+			}
+		}
+	}
+	
+	static void SmoothCurves(TArray<FPolygon2f>& InOutPerimeterCurves, int32 SmoothingIterations)
+	{
+		while (SmoothingIterations != 0)
+		{
+			for (FPolygon2f& PerimeterCurve : InOutPerimeterCurves)
+			{
+				const TArray<FVector2f>& Vertices = PerimeterCurve.GetVertices();
+
+				TArray<FVector2f> SmoothedCurve;
+				SmoothedCurve.SetNum(Vertices.Num());
+			
+				for (int32 i = 0; i < Vertices.Num(); i++)
+				{
+					const FVector2f& Pos1 = i == 0 ? Vertices.Last() : Vertices[i - 1];
+					const FVector2f& Pos2 = Vertices[i];
+					const FVector2f& Pos3 = i == Vertices.Num() - 1 ? Vertices[0] : Vertices[i + 1];
+				
+					const FVector2f AvgPos = (Pos1 + Pos2 + Pos3) / 3; 
+					SmoothedCurve[i] = AvgPos;
+				}
+			
+				PerimeterCurve = FPolygon2f(MoveTemp(SmoothedCurve));
+			}
+			
+			SmoothingIterations -= 1;
+		}
+	}
+	
+	static void SimplifyCurves(TArray<FPolygon2f>& InOutPerimeterCurves, float SimplificationAmount)
+	{
+		for (FPolygon2f& PerimeterCurve : InOutPerimeterCurves)
+		{
+			const float LineDeviationTolerance = SimplificationAmount * 0.001;
+			const float ClusterTolerance = LineDeviationTolerance * 0.001;
+			PerimeterCurve.Simplify(ClusterTolerance, LineDeviationTolerance);
+		}
+	}
+	
+	static void TriangulateCurvesAndSplitEdges(const TArray<FPolygon2f>& InPerimeterCurves, UE::Geometry::FDynamicMesh3& OutMesh, float Scale)
+	{
+		/*
+		* This function convertes the perimeter curves into a DynamicMesh and adds a vertex at the center of each edge spanning two perimeter points.
+		* 
+		* It does this by first triangulating the curve using FDelaunay2, it then looks at each edge and if said edge has two perimeter points, it inserts a new vertex
+		* at the center of the edge. It does this for all edges and then re-triangulates the entire mesh with this new vertex set.
+		* 
+		* The end result will look something like the following diagram:
+		*      *      
+		*     / \           /|\
+		*    /   \         / | \
+		*   /     \       /__|__\
+		*   *     *       |\ | /|
+		*   |     |  -->  | \|/ |
+		*   *     *       |--|--|
+		*   |     |       | /|\ |
+		*   *     *       |/ | \|
+		*/
+
+		PerimeterAttributeDefinition.AttachAttribute(OutMesh);
+
+		for (int32 CurveIndex = 0; CurveIndex < InPerimeterCurves.Num(); CurveIndex++)
+		{
+			const FPolygon2f& PerimeterCurve = InPerimeterCurves[CurveIndex];
+			const TArray<FVector2f>& CurveVertices = PerimeterCurve.GetVertices();
+
+			const TArray<FIndex2i> OriginalCurveEdges = Invoke([&]() 
+			{
+				TArray<FIndex2i> Edges;
+				Edges.Reserve(CurveVertices.Num());
+
+				for (int32 Last = CurveVertices.Num() - 1, Idx = 0; Idx < CurveVertices.Num(); Last = Idx++)
+				{
+					Edges.Emplace(Last, Idx);
+				}
+
+				return Edges;
+			});
+			
+			FDelaunay2 Delaunay;
+			Delaunay.Triangulate(CurveVertices, OriginalCurveEdges);
+			
+			bool bTriangulateSucess = Delaunay.GetResult() == UE::Geometry::FDelaunay2::EResult::MissingEdges
+				|| Delaunay.GetResult() == UE::Geometry::FDelaunay2::EResult::Success;
+
+			if (!bTriangulateSucess)
+			{
+				continue;
+			}
+
+			const FDelaunay2::EFillMode FillMode = PerimeterCurve.IsClockwise() 
+				? FDelaunay2::EFillMode::NegativeWinding 
+				: FDelaunay2::EFillMode::PositiveWinding;
+			
+			TArray<FIndex3i> OriginalCurveTriangles;
+			Delaunay.GetFilledTrianglesGeneralizedWinding(OriginalCurveTriangles, CurveVertices, OriginalCurveEdges, FillMode);
+
+			TArray<FVector2f> SubdividedMeshPoints;
+			SubdividedMeshPoints.Reserve(CurveVertices.Num() + OriginalCurveTriangles.Num());
+			SubdividedMeshPoints.Append(CurveVertices);
+			
+			TArray<FIndex2i> SubdividedMeshEdges;
+			SubdividedMeshEdges.Reserve(OriginalCurveTriangles.Num() * 4);
+			SubdividedMeshEdges.Append(OriginalCurveEdges);
+			
+			TSet<FIndex2i> CheckedEdges;
+			CheckedEdges.Reserve(OriginalCurveTriangles.Num() * 3);
+			CheckedEdges.Append(OriginalCurveEdges);
+			
+			for (const FIndex3i& Triangle : OriginalCurveTriangles)
+			{
+				for (int32 i = 0; i < 3; i++)
+				{
+					FIndex2i EdgeIndex(Triangle[i], Triangle[(i + 1) % 3]);
+					EdgeIndex.Sort();
+					
+					bool bAlreadyAdded = false;
+					CheckedEdges.Add(EdgeIndex, &bAlreadyAdded);
+					if (!bAlreadyAdded)
+					{
+						const int32 PointIndex1 = EdgeIndex[0];
+						const int32 PointIndex2 = EdgeIndex[1];
+						const FVector2f& Pos1 = CurveVertices[PointIndex1];
+						const FVector2f& Pos2 = CurveVertices[PointIndex2];
+						const int32 NewPointIndex = SubdividedMeshPoints.Add((Pos1 + Pos2) / 2);
+						
+						const FIndex2i NewEdge1(PointIndex1, NewPointIndex);
+						const FIndex2i NewEdge2(PointIndex2, NewPointIndex);
+						
+						SubdividedMeshEdges.Add(NewEdge1);
+						SubdividedMeshEdges.Add(NewEdge2);
+					}
+				}
+			}
+			
+			Delaunay.Triangulate(SubdividedMeshPoints, SubdividedMeshEdges);
+
+			bTriangulateSucess = Delaunay.GetResult() == UE::Geometry::FDelaunay2::EResult::MissingEdges
+				|| Delaunay.GetResult() == UE::Geometry::FDelaunay2::EResult::Success;
+
+			if (!bTriangulateSucess)
+			{
+				continue;
+			}
+			
+			UE::Geometry::FDynamicMesh3 NewDynamicMesh;
+			PerimeterAttributeDefinition.AttachAttribute(NewDynamicMesh);
+			auto& PerimeterAttribute = PerimeterAttributeDefinition.GetAttributeChecked(NewDynamicMesh);
+			
+			const TArray<FVector2f>& Vertices = SubdividedMeshPoints;
+			TArray<FIndex3i> Triangles;
+			Delaunay.GetFilledTrianglesGeneralizedWinding(Triangles, Vertices, OriginalCurveEdges, FillMode);
+			
+			for (int32 i = 0; i < Vertices.Num(); ++i)
+			{
+				const FVector2f& Vertex = Vertices[i];
+				UE::Geometry::FVertexInfo VertexInfo(
+					ToVector3<double, float>(Vertex, Scale),
+					FVector3f(-1, 0, 0)
+				);
+				
+				const int32 VertexID = NewDynamicMesh.AppendVertex(VertexInfo);
+				PerimeterAttribute.SetValue(VertexID, i < CurveVertices.Num());
+			}
+			
+			for (const FIndex3i& Triangle : Triangles)
+			{
+				NewDynamicMesh.AppendTriangle(Triangle);
+			}
+			
+			OutMesh.AppendWithOffsets(NewDynamicMesh);
+		}
+	}
+
+	static void ComputeVertexScale2D(
+		const UE::Geometry::FDynamicMesh3& DynamicMesh,
+		const FBooleanVertexAttribute& PerimeterAttribute,
+		FFloatVertexAttribute& VertexScaleAttribute
+	)
+	{
+		/*
+		* This function computes the vertex scale for each vertex in the mesh by finding the nearest point on the perimeter curve.
+		*/
+
+		const auto FindDistToPerimeter = [&](const FVector3d& Point, const TArray<int32>& PerimeterVertices)
+		{
+			// TODO: This algorithm is N^2, we can probably reduce that by using FSparseDynamicOctree3
+
+			double ClosestDistSqr = MAX_flt;
+
+			for (int32 i = 0; i < PerimeterVertices.Num(); ++i)
+			{
+				const int32 PointIndex0 = PerimeterVertices[i];
+				const int32 PointIndex1 = PerimeterVertices[FMath::WrapExclusive(i + 1, 0, PerimeterVertices.Num())];
+
+				const FVector3d& VertexLocation1 = DynamicMesh.GetVertex(PointIndex0);
+				const FVector3d& VertexLocation2 = DynamicMesh.GetVertex(PointIndex1);
+
+				const FVector3d ClosestPointOnLine = FMath::ClosestPointOnSegment(Point, VertexLocation1, VertexLocation2);
+				const double DistToLineSqr = (Point - ClosestPointOnLine).SizeSquared();
+				if (DistToLineSqr < ClosestDistSqr)
+				{
+					ClosestDistSqr = DistToLineSqr;
+				}
+			}
+
+			return ClosestDistSqr > 0 ? FMath::Sqrt(ClosestDistSqr) : 0;
+		};
+
+		UE::Geometry::FMeshConnectedComponents ConnectedComponents(&DynamicMesh);
+		ConnectedComponents.FindConnectedVertices();
+
+		for (const auto& Component : ConnectedComponents)
+		{
+			TArray<int32> PerimeterVertices = PerimeterAttribute.FindAllNonZero(Component.Indices);
+			PerimeterVertices.Sort(); // This assumes perimeter vertices are ordered sequentially, which is currently the case.
+
+			for (int32 VertexID : Component.Indices)
+			{
+				if (!PerimeterAttribute.GetValue(VertexID))
+				{
+					const float VertexScale = FindDistToPerimeter(DynamicMesh.GetVertex(VertexID), PerimeterVertices);
+					VertexScaleAttribute.SetValue(VertexID, VertexScale);
+				}
+			}
+		}
+	}
+
+	static void FindEndPoints2D(
+		const UE::Geometry::FDynamicMesh3& DynamicMesh,
+		float MaxTipAngleSearchDist,
+		float TipAngleThresholdDeg,
+		FBooleanVertexAttribute& EndPointAttribute,
+		const FBooleanVertexAttribute& PerimeterAttribute
+	)
+	{
+		/*
+		* This function finds the end points (or "tips") on a 2D mesh by walking the perimeter vertices and checking the angle between itself and its nearest 
+		* two neighbours on the perimeter. To allow for curved (or flat-tipped) end points to be detected, the function takes a second pass where it walks outwards along the
+		* perimited (limited by MaxTipAngleSearchDist), and checks the angle of the current vertex towards the subsequent neighbours. (see diagram a for example if flat-tipped end point)
+		* 
+		*     a: ___     b:
+		*       /   \      /\
+		*       |   |     |  |
+		*       |   |     |  |
+		*/
+
+		using namespace PV::DynamicMeshHelper;
+
+		const static auto ComputeAngleDeg = [](
+			const UE::Geometry::FDynamicMesh3& Mesh,
+			const TArray<int32>& PerimeterVertices, 
+			int32 PerimeterIndex, 
+			int32 LeftPerimeterIndex, 
+			int32 RightPerimeterIndex)->float
+		{
+			check(PerimeterVertices.IsValidIndex(PerimeterIndex));
+			check(Mesh.IsVertex(PerimeterVertices[PerimeterIndex]));
+
+			LeftPerimeterIndex = FMath::WrapExclusive(LeftPerimeterIndex, 0, PerimeterVertices.Num());
+			RightPerimeterIndex = FMath::WrapExclusive(RightPerimeterIndex, 0, PerimeterVertices.Num());
+
+			const int32 VertexIndex = PerimeterVertices[PerimeterIndex];
+			const int32 NeighbourVertexIndex1 = PerimeterVertices[LeftPerimeterIndex];
+			const int32 NeighbourVertexIndex2 = PerimeterVertices[RightPerimeterIndex];
+
+			const FVector3d& PointPos = Mesh.GetVertex(VertexIndex);
+			const FVector3d& PerimeterNeighbour1Pos = Mesh.GetVertex(NeighbourVertexIndex1);
+			const FVector3d& PerimeterNeighbour2Pos = Mesh.GetVertex(NeighbourVertexIndex2);
+
+			const FVector3d Normal1 = (PointPos - PerimeterNeighbour1Pos).GetSafeNormal();
+			const FVector3d Normal2 = (PointPos - PerimeterNeighbour2Pos).GetSafeNormal();
+
+			const double ClampedDot = FMath::Clamp(FVector3d::DotProduct(Normal1, Normal2), -1.0, 1.0); // Clamp to ensure range does not exceed [-1, 1] due to floating point imprecision
+			const double AngleRadians = FMath::Acos(ClampedDot);
+			return FMath::RadiansToDegrees(AngleRadians);
+		};
+
+		UE::Geometry::FMeshConnectedComponents ConnectedComponents(&DynamicMesh);
+		ConnectedComponents.FindConnectedVertices();
+
+		for (const auto& Component : ConnectedComponents)
+		{
+			const TArray<int32> PerimeterVertices = PerimeterAttribute.FindAllNonZero(Component.Indices);
+			 	
+			TArray<int32> ExcludedPerimeterVertices;
+			ExcludedPerimeterVertices.Reserve(PerimeterVertices.Num());
+			 	
+			TArray<float> PerimeterAngles;
+			PerimeterAngles.SetNumZeroed(PerimeterVertices.Num());
+			 	
+			for (int32 PerimeterIndex = 0; PerimeterIndex < PerimeterVertices.Num(); ++PerimeterIndex)
+			{
+				const int32 PointIndex = PerimeterVertices[PerimeterIndex];
+			 	check(DynamicMesh.IsVertex(PointIndex));
+			 		
+			 	int32 NumInternalNeighbours = 0;
+				const TArray<int32> PointNeighbours = PV::DynamicMeshHelper::GetVertexNeighbours(DynamicMesh, PointIndex);
+			 	for (int32 NeighbourIndex : PointNeighbours)
+			 	{
+			 		if (!PerimeterAttribute.GetValue(NeighbourIndex))
+			 		{
+			 			NumInternalNeighbours += 1;
+			 		}
+			 	}
+			 	
+				// TODO: Since we do not know whether the tip is pointing inwards or outwards, we use the fact that most outward pointing tips only have one
+				// internal connection (e.g. a connection to a vertec which is not a perimeter vertex). This is however not a very reliable method of
+				// determining the direction of the tip. A better alternative would be to calculate the normal for each vertex in the perimeter curve
+				// and use said normal to determine whether the tip is pointing outwads.
+			 	if (NumInternalNeighbours == 1)
+			 	{
+			 		const double TipAngle = ComputeAngleDeg(DynamicMesh, PerimeterVertices, PerimeterIndex, PerimeterIndex - 1, PerimeterIndex + 1);
+			 		if (TipAngle >= TipAngleThresholdDeg)
+			 		{
+						EndPointAttribute.SetValue(PointIndex, true);
+						PerimeterAngles[PerimeterIndex] = TipAngle;
+			 		}
+			 		else 
+			 		{
+						ExcludedPerimeterVertices.Add(PerimeterIndex);
+			 		}
+			 	}
+			}
+			 	
+			const float MaxTipAngleSearchDistSqr = MaxTipAngleSearchDist * MaxTipAngleSearchDist;
+			for (int32 PerimeterIndex : ExcludedPerimeterVertices)
+			{
+				const int32 PointIndex = PerimeterVertices[PerimeterIndex];
+
+			 	float LeftDistSqr = 0.f;
+			 	float RightDistSqr = 0.f;
+			 		
+			 	int32 LeftPerimeterIndex = PerimeterIndex;
+			 	int32 RightPerimeterIndex = PerimeterIndex;
+			 		
+			 	while (LeftDistSqr < MaxTipAngleSearchDistSqr || RightDistSqr < MaxTipAngleSearchDistSqr)
+			 	{
+			 		if (LeftDistSqr < MaxTipAngleSearchDistSqr)
+			 		{
+						LeftPerimeterIndex = FMath::WrapExclusive(LeftPerimeterIndex - 1, 0, PerimeterVertices.Num());
+			 				
+			 			if (LeftPerimeterIndex == PerimeterIndex)
+			 			{
+			 				break;
+			 			}
+			 			
+						const int32 LeftVertexIndex = PerimeterVertices[LeftPerimeterIndex];
+
+			 			if (EndPointAttribute.GetValue(LeftVertexIndex))
+			 			{
+			 				LeftDistSqr = MaxTipAngleSearchDistSqr;
+			 			}
+			 			else
+			 			{
+			 				LeftDistSqr = (DynamicMesh.GetVertex(PointIndex) - DynamicMesh.GetVertex(LeftVertexIndex)).SizeSquared();
+			 			}
+			 		}
+			 			
+			 		if (RightDistSqr < MaxTipAngleSearchDistSqr)
+			 		{
+						RightPerimeterIndex = FMath::WrapExclusive(RightPerimeterIndex + 1, 0, PerimeterVertices.Num());
+			 			
+			 			if (RightPerimeterIndex == PerimeterIndex)
+			 			{
+			 				break;
+			 			}
+
+						const int32 RightPointIndex = PerimeterVertices[RightPerimeterIndex];
+			 				
+			 			if (EndPointAttribute.GetValue(RightPointIndex))
+			 			{
+			 				RightDistSqr = MaxTipAngleSearchDistSqr;
+			 			}
+			 			else
+			 			{
+							RightDistSqr = (DynamicMesh.GetVertex(PointIndex) - DynamicMesh.GetVertex(RightPointIndex)).SizeSquared();
+			 			}
+			 		}
+			 			
+			 		const float TipAngle = ComputeAngleDeg(DynamicMesh, PerimeterVertices, PerimeterIndex, LeftPerimeterIndex, RightPerimeterIndex);
+			 		if (TipAngle >= TipAngleThresholdDeg)
+			 		{
+						EndPointAttribute.SetValue(PointIndex, true);
+						PerimeterAngles[PerimeterIndex] = TipAngle;
+			 			break;
+			 		}
+			 	}
+			}
+			 	
+			// When we have 2 neighbouring end points (e.g. when the tip is a plane with 2 45* corners that can both be considered end points),
+			// then select the one with the steepest angle
+			for (int32 PerimeterIndex = 0; PerimeterIndex < PerimeterVertices.Num(); ++PerimeterIndex)
+			{
+				const int32 PointIndex = PerimeterVertices[PerimeterIndex];
+				if (!EndPointAttribute.GetValue(PointIndex))
+				{
+					continue;
+				}
+
+			 	const int32 NextPerimeterIndex = FMath::WrapExclusive(PerimeterIndex + 1, 0, PerimeterVertices.Num());
+				const int32 NextPointIndex = PerimeterVertices[NextPerimeterIndex];
+
+			 	if (EndPointAttribute.GetValue(NextPointIndex))
+			 	{
+			 		const int32 EndPointToRemove = PerimeterAngles[NextPerimeterIndex] >= PerimeterAngles[PerimeterIndex] ? PointIndex : NextPointIndex;
+					EndPointAttribute.SetValue(EndPointToRemove, false);
+			 	}
+			}
+		}
+	}
+
+	struct FPlantMesh
+	{
+		FMeshConnectedComponents::FComponent MeshComponent;
+		int32 RootVertexIndex = INDEX_NONE;
+	};
+
+	static void MapPlantSettingsToMeshComponents(
+		const UE::Geometry::FDynamicMesh3& DynamicMesh,
+		float Scale,
+		const TArray<FPVImportTexture2DPlantSettings>& InputPlantSettings,
+		TArray<FPlantMesh>& OutPlantMapping
+	)
+	{
+		using namespace PV::DynamicMeshHelper;
+
+		/*
+		* 1. Compute the distance to the closest mesh component for each InputPlantSettings
+		* 2. For each mesh component, select the closest InputPlantSettings
+		* 3. Determine which vertex to use as root point
+		*/
+
+		FMeshConnectedComponents Components(&DynamicMesh);
+		Components.FindConnectedTriangles();
+
+		TArray<FDynamicMeshAABBTree3> AABBTrees;
+		AABBTrees.Reserve(Components.Components.Num());
+
+		for (const auto& Component : Components)
+		{
+			FDynamicMeshAABBTree3 DynamicMeshAABBTree(&DynamicMesh, false);
+			DynamicMeshAABBTree.Build(Component.Indices);
+
+			AABBTrees.Add(MoveTemp(DynamicMeshAABBTree));
+		}
+
+		struct FClosestTriangle
+		{
+			int32 ComponentIndex = INDEX_NONE;
+			double DistSqr = TNumericLimits<double>::Max();
+			int32 ClosestTriangle = INDEX_NONE;
+		};
+		TArray<FClosestTriangle> PlantSettingsClosestTriangle;
+		PlantSettingsClosestTriangle.SetNum(InputPlantSettings.Num());
+
+		for (int32 PlantIndex = 0; PlantIndex < InputPlantSettings.Num(); ++PlantIndex)
+		{
+			const FPVImportTexture2DPlantSettings& PlantSettings = InputPlantSettings[PlantIndex];
+			FClosestTriangle& ClosestTriangle = PlantSettingsClosestTriangle[PlantIndex];
+
+			for (int32 ComponentIndex = 0; ComponentIndex < Components.Num(); ComponentIndex++)
+			{
+				const FDynamicMeshAABBTree3& AABBTree = AABBTrees[ComponentIndex];
+
+				double DistSqr = 0;
+				const FVector3d RootPoint = ToVector3<double, float>(PlantSettings.RootPosition, Scale);
+				const int32 NearTriID = AABBTree.FindNearestTriangle(RootPoint, DistSqr);
+				if (NearTriID != INDEX_NONE)
+				{
+					UE::Geometry::FDistPoint3Triangle3d Query = FDynamicMeshQueries3::TriangleDistance(DynamicMesh, NearTriID, RootPoint);
+					DistSqr = Query.GetSquared();
+					if (DistSqr < ClosestTriangle.DistSqr)
+					{
+						ClosestTriangle.ComponentIndex = ComponentIndex;
+						ClosestTriangle.DistSqr = DistSqr;
+						ClosestTriangle.ClosestTriangle = NearTriID;
+					}
+				}
+			}
+		}
+
+		struct FPlantIndexDistPair
+		{
+			int32 PlantIndex = INDEX_NONE;
+			double DistSqr = TNumericLimits<double>::Max();
+		};
+		TArray<FPlantIndexDistPair> ComponentToPlantDist;
+		ComponentToPlantDist.SetNum(Components.Num());
+
+		for (int32 PlantIndex = 0; PlantIndex < InputPlantSettings.Num(); ++PlantIndex)
+		{
+			const int32 ComponentIndex = PlantSettingsClosestTriangle[PlantIndex].ComponentIndex;
+			if (ComponentIndex == INDEX_NONE)
+			{
+				continue;
+			}
+
+			if (PlantSettingsClosestTriangle[PlantIndex].DistSqr < ComponentToPlantDist[ComponentIndex].DistSqr)
+			{
+				ComponentToPlantDist[ComponentIndex].PlantIndex = PlantIndex;
+				ComponentToPlantDist[ComponentIndex].DistSqr = PlantSettingsClosestTriangle[PlantIndex].DistSqr;
+			}
+		}
+
+		OutPlantMapping.SetNum(InputPlantSettings.Num());
+		for (int32 ComponentIndex = 0; ComponentIndex < Components.Num(); ComponentIndex++)
+		{
+			if (ComponentToPlantDist[ComponentIndex].PlantIndex == INDEX_NONE)
+			{
+				// Unable to map this plant setting to a mesh component
+				continue;
+			}
+
+			const int32 PlantIndex = ComponentToPlantDist[ComponentIndex].PlantIndex;
+			const FPVImportTexture2DPlantSettings& PlantSettings = InputPlantSettings[PlantIndex];
+
+			const FVector3d& RootPosition = ToVector3<double, float>(PlantSettings.RootPosition, Scale);
+			const int32 TriangleID = PlantSettingsClosestTriangle[PlantIndex].ClosestTriangle;
+			const FIndex3i& Triangle = DynamicMesh.GetTriangle(TriangleID);
+			const int32 ClosestTriangleVertex = FMath::Min3Index(
+				(RootPosition - DynamicMesh.GetVertex(Triangle[0])).SizeSquared(),
+				(RootPosition - DynamicMesh.GetVertex(Triangle[1])).SizeSquared(),
+				(RootPosition - DynamicMesh.GetVertex(Triangle[2])).SizeSquared()
+			);
+
+			const int32 RootVertexID = Triangle[ClosestTriangleVertex];
+
+			OutPlantMapping[PlantIndex] = { MoveTemp(Components[ComponentIndex]), RootVertexID };
+		}
+	}
+
+	static void MapPlantIndicesToBranchHierarchies(
+		const TArray<FPlantMesh>& PlantMapping,
+		const TArray<PV::DynamicMeshHelper::FMeshBranchHierarchy>& DynamicMeshBranchHierarchies,
+		TArray<int32>& OutBranchHierarchiesToPlantIndices
+	)
+	{
+		TMap<int32, int32> RootVertexPlantIndexMap;
+		RootVertexPlantIndexMap.Reserve(PlantMapping.Num());
+		for (int32 PlantIndex = 0; PlantIndex < PlantMapping.Num(); ++PlantIndex)
+		{
+			const FPlantMesh& PlantMesh = PlantMapping[PlantIndex];
+			RootVertexPlantIndexMap.Emplace(PlantMesh.RootVertexIndex, PlantIndex);
+		}
+
+		OutBranchHierarchiesToPlantIndices.SetNum(DynamicMeshBranchHierarchies.Num());
+		for (int32& PlantIndex : OutBranchHierarchiesToPlantIndices)
+		{
+			PlantIndex = INDEX_NONE;
+		}
+
+		for (int32 i = 0; i < DynamicMeshBranchHierarchies.Num(); ++i)
+		{
+			check(DynamicMeshBranchHierarchies[i].Branches.Num() > 0);
+			check(DynamicMeshBranchHierarchies[i].Branches[0].BranchPoints.Num() > 0);
+
+			const int32 RootVertexID = DynamicMeshBranchHierarchies[i].Branches[0].BranchPoints[0];
+			if (int32* PlantIndex = RootVertexPlantIndexMap.Find(RootVertexID))
+			{
+				OutBranchHierarchiesToPlantIndices[i] = *PlantIndex;
+			}
+		}
+	}
+
+	static void ApplyPlantMappingToDynamicMesh(
+		UE::Geometry::FDynamicMesh3& DynamicMesh,
+		const TArray<FPVImportTexture2DPlantSettings>& InputPlantSettings,
+		const TArray<FPlantMesh>& InPlantMapping,
+		FBooleanVertexAttribute& RootPointAttribute,
+		FBooleanVertexAttribute& EndPointAttribute
+	)
+	{
+		check(InputPlantSettings.Num() == InPlantMapping.Num());
+
+		for (int32 PlantIndex = 0; PlantIndex < InputPlantSettings.Num(); ++PlantIndex)
+		{
+			const FPlantMesh& PlantMesh = InPlantMapping[PlantIndex];
+			if (PlantMesh.RootVertexIndex == INDEX_NONE)
+			{
+				continue;
+			}
+
+			const int32 RootVertexID = PlantMesh.RootVertexIndex;
+			const FPVImportTexture2DPlantSettings& PlantSettings = InputPlantSettings[PlantIndex];
+
+			RootPointAttribute.SetValue(RootVertexID, true);
+			EndPointAttribute.SetValue(RootVertexID, false);
+
+			const FVector VertexTranslation = DynamicMesh.GetVertex(RootVertexID);
+			const FQuat VertexRotation = FQuat(FVector::ForwardVector, FMath::DegreesToRadians(PlantSettings.Rotation));
+
+			for (int32 VertexIndex : PV::DynamicMeshHelper::GetUniqueVertexIndices(DynamicMesh, PlantMesh.MeshComponent.Indices))
+			{
+				FVector Vertex = DynamicMesh.GetVertex(VertexIndex);
+				Vertex -= VertexTranslation;
+				Vertex = VertexRotation.UnrotateVector(Vertex);
+				DynamicMesh.SetVertex(VertexIndex, Vertex);
+			}
+		}
+	}
+};
+
+PV::Texture2DImport::EImportResult PV::Texture2DImport::ImportGrowthDataFromTexture2D(const FPVImportTexture2DParams& InParams, const UMaterialInterface* TextureVisualizationMaterial, FPVImportTexture2DOutput& Output)
+{
+	using namespace PV::DynamicMeshHelper;
+	using namespace Internal;
+	
+	if (!InParams.Texture2DAsset)
+	{
+		return EImportResult::InvalidSourceTexture;
+	}
+
+	if (InParams.Scale <= 0)
+	{
+		return EImportResult::InvalidScale;
+	}
+
+	FVector2i SampleSize = FVector2i(InParams.SampleWidth, InParams.SampleHeight);
+	if (SampleSize.X <= 1 || SampleSize.Y <= 1)
+	{
+		return EImportResult::InvalidSampleSize;
+	}
+	
+	TArray<bool> BlackAndWhiteImageData;
+
+	if (!SampleImage(*InParams.Texture2DAsset.Get(), SampleSize, BlackAndWhiteImageData, InParams.bInvertImage, InParams.WhiteLevel))
+	{
+		return EImportResult::InvalidSourceTexture;
+	}
+
+	if (SampleSize.X <= 1 || SampleSize.Y <= 1)
+	{
+		return EImportResult::InvalidTextureSize;
+	}
+	
+	if (InParams.DebugState == EPVImportTexture2DDebugState::SampleImage)
+	{
+		FManagedArrayCollection& DebugCollection = Output.DebugCollection.Emplace();
+		AddDebugPixelValuesToCollection(DebugCollection, BlackAndWhiteImageData, SampleSize, InParams.Scale);
+		AddTextureQuadToCollection(DebugCollection, *InParams.Texture2DAsset, TextureVisualizationMaterial, InParams.Scale);
+	}
+
+	TArray<FPolygon2f> PerimeterCurves;
+	TraceImage(BlackAndWhiteImageData, SampleSize, PerimeterCurves);
+	RemoveSmallCurves(PerimeterCurves, InParams.MinBoundsArea);
+	SmoothCurves(PerimeterCurves, InParams.SmoothingIterations);
+	SimplifyCurves(PerimeterCurves, InParams.SimplificationAmount);
+
+	UE::Geometry::FDynamicMesh3 DynamicMesh;
+	TriangulateCurvesAndSplitEdges(PerimeterCurves, DynamicMesh, InParams.Scale);
+
+	if (InParams.DebugState == EPVImportTexture2DDebugState::TriangulateCurves)
+	{
+		FManagedArrayCollection& DebugCollection = Output.DebugCollection.Emplace();
+		AddDynamicMeshToCollection(DynamicMesh, DebugCollection);
+		PerimeterAttributeDefinition.CopyAttributeToCollection(DynamicMesh, DebugCollection);
+	}
+
+	ComputeVertexScale2D(
+		DynamicMesh, 
+		PerimeterAttributeDefinition.GetOrAttachAttribute(DynamicMesh),
+		VertexScaleAttributeDefinition.GetOrAttachAttribute(DynamicMesh)
+	);
+
+	if (InParams.DebugState == EPVImportTexture2DDebugState::ComputeVertexScale)
+	{
+		FManagedArrayCollection& DebugCollection = Output.DebugCollection.Emplace();
+		AddDynamicMeshToCollection(DynamicMesh, DebugCollection);
+		VertexScaleAttributeDefinition.CopyAttributeToCollection(DynamicMesh, DebugCollection);
+	}
+
+	FindEndPoints2D(
+		DynamicMesh,
+		InParams.MaxTipAngleSearchDist * InParams.Scale,
+		InParams.TipAngleThresholdInDegrees,
+		EndPointAttributeDefinition.GetOrAttachAttribute(DynamicMesh),
+		PerimeterAttributeDefinition.GetOrAttachAttribute(DynamicMesh)
+	);
+
+	const float HypotenuseLength = FMath::Sqrt(InParams.Scale * InParams.Scale * 2);
+	const FVector VisualizationOffset = FVector(0, -InParams.Scale / 2, HypotenuseLength);
+
+	// Add end points visualization before applying root points as that transforms the DynamicMesh
+	AddDynamicMeshVerticesToCollection(
+		Output.VisualizationCollection,
+		TEXT("EndPoints"), 
+		TEXT("Vertex"), 
+		DynamicMesh, 
+		EndPointAttributeDefinition.GetOrAttachAttribute(DynamicMesh), 
+		VisualizationOffset
+	);
+
+	TArray<FPlantMesh> PlantMapping;
+	MapPlantSettingsToMeshComponents(
+		DynamicMesh,
+		InParams.Scale,
+		InParams.PlantSettings,
+		PlantMapping
+	);
+
+	ApplyPlantMappingToDynamicMesh(
+		DynamicMesh,
+		InParams.PlantSettings,
+		PlantMapping,
+		RootPointAttributeDefinition.GetOrAttachAttribute(DynamicMesh),
+		EndPointAttributeDefinition.GetOrAttachAttribute(DynamicMesh)
+	);
+
+	auto& PerimeterAttribute = PerimeterAttributeDefinition.GetOrAttachAttribute(DynamicMesh);
+	
+	TArray<FMeshBranchHierarchy> DynamicMeshBranchHierarchies;
+	ComputeDynamicMeshBranchHierarchies(
+		DynamicMesh,
+		RootPointAttributeDefinition.GetOrAttachAttribute(DynamicMesh),
+		EndPointAttributeDefinition.GetOrAttachAttribute(DynamicMesh),
+		VertexScaleAttributeDefinition.GetOrAttachAttribute(DynamicMesh),
+		LengthFromRootAttributeDefinition.GetOrAttachAttribute(DynamicMesh),
+		NextBranchPointAttributeDefinition.GetOrAttachAttribute(DynamicMesh),
+		[&](int32 VertexID) { return !PerimeterAttribute.GetValue(VertexID); },
+		DynamicMeshBranchHierarchies
+	);
+
+	TArray<FPVBranchHierarchyDescription> BranchHierarchyDescriptions;
+	ConvertDynamicMeshBranchHierarchiesToBranchDescription(
+		DynamicMesh,
+		VertexScaleAttributeDefinition.GetOrAttachAttribute(DynamicMesh),
+		EndPointAttributeDefinition.GetOrAttachAttribute(DynamicMesh),
+		DynamicMeshBranchHierarchies,
+		BranchHierarchyDescriptions
+	);
+
+	check(DynamicMeshBranchHierarchies.Num() == BranchHierarchyDescriptions.Num());
+
+	if (InParams.DebugState == EPVImportTexture2DDebugState::ComputeBranchHierarchies)
+	{
+		FManagedArrayCollection& DebugCollection = Output.DebugCollection.Emplace();
+		AddDynamicMeshToCollection(DynamicMesh, DebugCollection);
+		AddBranchHierarchiesToCollection(BranchHierarchyDescriptions, DebugCollection);
+		RootPointAttributeDefinition.CopyAttributeToCollection(DynamicMesh, DebugCollection);
+		EndPointAttributeDefinition.CopyAttributeToCollection(DynamicMesh, DebugCollection);
+		VertexScaleAttributeDefinition.CopyAttributeToCollection(DynamicMesh, DebugCollection);
+		LengthFromRootAttributeDefinition.CopyAttributeToCollection(DynamicMesh, DebugCollection);
+		NextBranchPointAttributeDefinition.CopyAttributeToCollection(DynamicMesh, DebugCollection);
+	}
+
+	TArray<int32> BranchHierarchiesToPlantIndices;
+	MapPlantIndicesToBranchHierarchies(PlantMapping, DynamicMeshBranchHierarchies, BranchHierarchiesToPlantIndices);
+
+	// TODO: It would probably be faster to combine all the outputs of the PlantCollections
+	if (!PV::ImportHelper::GenerateGrowthDataFromBranchHierarchies(Output.CombinedPlantCollection, BranchHierarchyDescriptions))
+	{
+		return EImportResult::FailedToGenerateGrowthData;
+	}
+
+	Output.PlantCollections.SetNum(InParams.PlantSettings.Num());
+	for (int32 i = 0; i < BranchHierarchyDescriptions.Num(); ++i)
+	{
+		const int32 PlantIndex = BranchHierarchiesToPlantIndices[i];
+		if (PlantIndex == INDEX_NONE)
+		{
+			continue;
+		}
+
+		if (!PV::ImportHelper::GenerateGrowthDataFromBranchHierarchy(Output.PlantCollections[PlantIndex], BranchHierarchyDescriptions[i]))
+		{
+			return EImportResult::FailedToGenerateGrowthData;
+		}
+	}
+
+	// Build visualization output
+	{
+		Output.CombinedPlantCollection.CopyTo(&Output.VisualizationCollection);
+		ArrangePlantsInGrid(Output.VisualizationCollection, BranchHierarchiesToPlantIndices, VisualizationOffset);
+		AddTextureQuadToCollection(Output.VisualizationCollection, *InParams.Texture2DAsset, TextureVisualizationMaterial, InParams.Scale, VisualizationOffset);
+		AddPerimeterCurvesToCollection(Output.VisualizationCollection, PerimeterCurves, InParams.Scale, VisualizationOffset);
+		AddPlantSettingsToCollection(Output.VisualizationCollection, InParams.PlantSettings, InParams.Scale, VisualizationOffset);
+	}
+
+	return EImportResult::Success;
+}
diff --git a/Engine/Plugins/Experimental/ProceduralVegetationEditor/Source/ProceduralVegetation/Private/Implementations/PVImporter_Texture2D.h b/Engine/Plugins/Experimental/ProceduralVegetationEditor/Source/ProceduralVegetation/Private/Implementations/PVImporter_Texture2D.h
new file mode 100644
--- /dev/null
+++ b/Engine/Plugins/Experimental/ProceduralVegetationEditor/Source/ProceduralVegetation/Private/Implementations/PVImporter_Texture2D.h
@@ -0,0 +1,23 @@
+// Copyright Epic Games, Inc. All Rights Reserved.
+#pragma once
+
+#include "CoreMinimal.h"
+#include "Params/PVImportTexture2DParams.h"
+
+class UMaterialInterface;
+class UTexture2D;
+
+namespace PV::Texture2DImport
+{
+	enum class EImportResult
+	{
+		Success,
+		InvalidScale,
+		InvalidSourceTexture,
+		InvalidSampleSize,
+		InvalidTextureSize,
+		FailedToGenerateGrowthData
+	};
+
+	EImportResult ImportGrowthDataFromTexture2D(const FPVImportTexture2DParams& InParams, const UMaterialInterface* TextureVisualizationMaterial, FPVImportTexture2DOutput& Output);
+};
diff --git a/Engine/Plugins/Experimental/ProceduralVegetationEditor/Source/ProceduralVegetation/Private/Implementations/PVMeshBuilder.h b/Engine/Plugins/Experimental/ProceduralVegetationEditor/Source/ProceduralVegetation/Private/Implementations/PVMeshBuilder.h
--- a/Engine/Plugins/Experimental/ProceduralVegetationEditor/Source/ProceduralVegetation/Private/Implementations/PVMeshBuilder.h
+++ b/Engine/Plugins/Experimental/ProceduralVegetationEditor/Source/ProceduralVegetation/Private/Implementations/PVMeshBuilder.h
@@ -3,7 +3,7 @@
 #pragma once
 
 #include "CoreMinimal.h"
-#include "PVFloatRamp.h"
+#include "Utils/PVFloatRamp.h"
 #include "PVMaterialSettings.h"
 #include "UDynamicMesh.h"
 
diff --git a/Engine/Plugins/Experimental/ProceduralVegetationEditor/Source/ProceduralVegetation/Private/Nodes/PVFoliageDistributorSettings.h b/Engine/Plugins/Experimental/ProceduralVegetationEditor/Source/ProceduralVegetation/Private/Nodes/PVFoliageDistributorSettings.h
--- a/Engine/Plugins/Experimental/ProceduralVegetationEditor/Source/ProceduralVegetation/Private/Nodes/PVFoliageDistributorSettings.h
+++ b/Engine/Plugins/Experimental/ProceduralVegetationEditor/Source/ProceduralVegetation/Private/Nodes/PVFoliageDistributorSettings.h
@@ -2,7 +2,7 @@
 #pragma once
 
 #include "CoreMinimal.h"
-#include "PVFloatRamp.h"
+#include "Utils/PVFloatRamp.h"
 
 #include "DataTypes/PVData.h"
 
diff --git a/Engine/Plugins/Experimental/ProceduralVegetationEditor/Source/ProceduralVegetation/Private/Nodes/PVGrafterSettings.h b/Engine/Plugins/Experimental/ProceduralVegetationEditor/Source/ProceduralVegetation/Private/Nodes/PVGrafterSettings.h
--- a/Engine/Plugins/Experimental/ProceduralVegetationEditor/Source/ProceduralVegetation/Private/Nodes/PVGrafterSettings.h
+++ b/Engine/Plugins/Experimental/ProceduralVegetationEditor/Source/ProceduralVegetation/Private/Nodes/PVGrafterSettings.h
@@ -2,7 +2,7 @@
 #pragma once
 
 #include "CoreMinimal.h"
-#include "PVFloatRamp.h"
+#include "Utils/PVFloatRamp.h"
 #include "Nodes/PVBaseSettings.h"
 #include "Implementations/PVGrafter.h"
 
diff --git a/Engine/Plugins/Experimental/ProceduralVegetationEditor/Source/ProceduralVegetation/Private/Nodes/PVStaticMeshImporterSettings.cpp b/Engine/Plugins/Experimental/ProceduralVegetationEditor/Source/ProceduralVegetation/Private/Nodes/PVStaticMeshImporterSettings.cpp
--- a/Engine/Plugins/Experimental/ProceduralVegetationEditor/Source/ProceduralVegetation/Private/Nodes/PVStaticMeshImporterSettings.cpp
+++ b/Engine/Plugins/Experimental/ProceduralVegetationEditor/Source/ProceduralVegetation/Private/Nodes/PVStaticMeshImporterSettings.cpp
@@ -1,7 +1,7 @@
 // Copyright Epic Games, Inc. All Rights Reserved.
 
 #include "PVStaticMeshImporterSettings.h"
-#include "IPVEditorInterface.h"
+#include "Implementations/PVImporter_StaticMesh.h"
 #include "PCGContext.h"
 #include "ProceduralVegetationModule.h"
 #include "DataTypes/PVGrowthData.h"
@@ -66,14 +66,9 @@ bool FPVStaticMeshImporterElement::ExecuteInternal(FPCGContext* InContext) const
 		return true;
 	}
 	
-#if WITH_EDITOR
-	if (!PV::GPVEditorInterface)
-	{
-		return true;
-	}
-	
 	FManagedArrayCollection Collection;
-	if (!PV::GPVEditorInterface->ImportGrowthDataFromStaticMesh(Settings->Params, Collection))
+	const auto Result = PV::StaticMeshImport::ImportGrowthDataFromStaticMesh(Settings->Params, Collection);
+	if (Result != PV::StaticMeshImport::EImportResult::Success)
 	{
 		return false;
 	}
@@ -84,7 +79,6 @@ bool FPVStaticMeshImporterElement::ExecuteInternal(FPCGContext* InContext) const
 	FPCGTaggedData& CollectionOutput = InContext->OutputData.TaggedData.Emplace_GetRef();
 	CollectionOutput.Data = OutVariantData;
 	CollectionOutput.Pin = TEXT("Output");
-#endif
 
 	return true;
 }
diff --git a/Engine/Plugins/Experimental/ProceduralVegetationEditor/Source/ProceduralVegetation/Private/Nodes/PVStaticMeshImporterSettings.h b/Engine/Plugins/Experimental/ProceduralVegetationEditor/Source/ProceduralVegetation/Private/Nodes/PVStaticMeshImporterSettings.h
--- a/Engine/Plugins/Experimental/ProceduralVegetationEditor/Source/ProceduralVegetation/Private/Nodes/PVStaticMeshImporterSettings.h
+++ b/Engine/Plugins/Experimental/ProceduralVegetationEditor/Source/ProceduralVegetation/Private/Nodes/PVStaticMeshImporterSettings.h
@@ -3,7 +3,7 @@
 
 #include "CoreMinimal.h"
 #include "Nodes/PVBaseSettings.h"
-#include "PVImportStaticMeshParams.h"
+#include "Params/PVImportStaticMeshParams.h"
 
 #include "PVStaticMeshImporterSettings.generated.h"
 
diff --git a/Engine/Plugins/Experimental/ProceduralVegetationEditor/Source/ProceduralVegetation/Private/Nodes/PVTexture2DImporterSettings.cpp b/Engine/Plugins/Experimental/ProceduralVegetationEditor/Source/ProceduralVegetation/Private/Nodes/PVTexture2DImporterSettings.cpp
--- a/Engine/Plugins/Experimental/ProceduralVegetationEditor/Source/ProceduralVegetation/Private/Nodes/PVTexture2DImporterSettings.cpp
+++ b/Engine/Plugins/Experimental/ProceduralVegetationEditor/Source/ProceduralVegetation/Private/Nodes/PVTexture2DImporterSettings.cpp
@@ -2,11 +2,11 @@
 
 #include "PVTexture2DImporterSettings.h"
 
-#include "IPVEditorInterface.h"
 #include "ProceduralVegetationModule.h"
 #include "DataTypes/PVGrowthData.h"
 #include "ProceduralVegetationPreset.h"
-#include "PVImportTexture2DParams.h"
+#include "Params/PVImportTexture2DParams.h"
+#include "Implementations/PVImporter_Texture2D.h"
 
 #include "PCGContext.h"
 
@@ -116,12 +116,6 @@ bool FPVTexture2DImporterElement::ExecuteInternal(FPCGContext* InContext) const
 		return true;
 	}
 
-#if WITH_EDITOR
-	if (!PV::GPVEditorInterface)
-	{
-		return true;
-	}
-
 	const static TArray<FPVVisualizationSettings> VisualizationSettings = Invoke([]()
 	{
 		TArray<FPVVisualizationSettings> OutVisualizationSettings;
@@ -172,6 +166,7 @@ bool FPVTexture2DImporterElement::ExecuteInternal(FPCGContext* InContext) const
 	{
 		UMaterialInstanceDynamic* OutMaterialInstance = nullptr;
 
+#if WITH_EDITOR
 		if (Settings->TextureVisualizationMaterial)
 		{
 			OutMaterialInstance = UMaterialInstanceDynamic::Create(Settings->TextureVisualizationMaterial, nullptr);
@@ -181,12 +176,18 @@ bool FPVTexture2DImporterElement::ExecuteInternal(FPCGContext* InContext) const
 			OutMaterialInstance->SetTextureParameterValue(TEXT("Diffuse"), Settings->Params.Texture2DAsset);
 			OutMaterialInstance->SetTextureParameterValue(TEXT("OpacityMask"), Settings->Params.Texture2DAsset);
 		}
+#endif
 
 		return OutMaterialInstance;
 	});
 
 	FPVImportTexture2DOutput ImportTexture2DOutput;
-	if (!PV::GPVEditorInterface->ImportGrowthDataFromTexture2D(Settings->Params, TextureVisualizationMaterialInstance, ImportTexture2DOutput))
+	const auto ImportResult = PV::Texture2DImport::ImportGrowthDataFromTexture2D(
+		Settings->Params, 
+		TextureVisualizationMaterialInstance, 
+		ImportTexture2DOutput
+	);
+	if (ImportResult != PV::Texture2DImport::EImportResult::Success)
 	{
 		return false;
 	}
@@ -196,7 +197,9 @@ bool FPVTexture2DImporterElement::ExecuteInternal(FPCGContext* InContext) const
 		UPVGrowthData* OutVariantData = FPCGContext::NewObject_AnyThread<UPVGrowthData>(InContext);
 
 		OutVariantData->AddManagedResource(TextureVisualizationMaterialInstance);
+#if WITH_EDITOR
 		OutVariantData->AddVisualizationCollection(ImportTexture2DOutput.VisualizationCollection, VisualizationSettings);
+#endif
 
 		if (ImportTexture2DOutput.DebugCollection.IsSet())
 		{
@@ -220,13 +223,14 @@ bool FPVTexture2DImporterElement::ExecuteInternal(FPCGContext* InContext) const
 
 		// TODO: For now we add the visualization to all outputs as we cannot control which pin gets inspected
 		OutVariantData->AddManagedResource(TextureVisualizationMaterialInstance);
+#if WITH_EDITOR
 		OutVariantData->AddVisualizationCollection(ImportTexture2DOutput.VisualizationCollection, VisualizationSettings);
+#endif
 
 		FPCGTaggedData& CollectionOutput = InContext->OutputData.TaggedData.Emplace_GetRef();
 		CollectionOutput.Data = OutVariantData;
 		CollectionOutput.Pin = PVTexture2DImporterSettings::GetPlantPinName(i);
 	}
-#endif
 
 	return true;
 }
diff --git a/Engine/Plugins/Experimental/ProceduralVegetationEditor/Source/ProceduralVegetation/Private/Nodes/PVTexture2DImporterSettings.h b/Engine/Plugins/Experimental/ProceduralVegetationEditor/Source/ProceduralVegetation/Private/Nodes/PVTexture2DImporterSettings.h
--- a/Engine/Plugins/Experimental/ProceduralVegetationEditor/Source/ProceduralVegetation/Private/Nodes/PVTexture2DImporterSettings.h
+++ b/Engine/Plugins/Experimental/ProceduralVegetationEditor/Source/ProceduralVegetation/Private/Nodes/PVTexture2DImporterSettings.h
@@ -3,7 +3,7 @@
 
 #include "CoreMinimal.h"
 #include "Nodes/PVBaseSettings.h"
-#include "PVImportTexture2DParams.h"
+#include "Params/PVImportTexture2DParams.h"
 
 #include "PVTexture2DImporterSettings.generated.h"
 
diff --git a/Engine/Plugins/Experimental/ProceduralVegetationEditor/Source/ProceduralVegetation/Private/PVExportParams.cpp b/Engine/Plugins/Experimental/ProceduralVegetationEditor/Source/ProceduralVegetation/Private/PVExportParams.cpp
deleted file mode 100644
diff --git a/Engine/Plugins/Experimental/ProceduralVegetationEditor/Source/ProceduralVegetation/Private/Params/PVExportParams.cpp b/Engine/Plugins/Experimental/ProceduralVegetationEditor/Source/ProceduralVegetation/Private/Params/PVExportParams.cpp
new file mode 100644
--- /dev/null
+++ b/Engine/Plugins/Experimental/ProceduralVegetationEditor/Source/ProceduralVegetation/Private/Params/PVExportParams.cpp
@@ -0,0 +1,55 @@
+// Copyright Epic Games, Inc. All Rights Reserved.
+
+#include "Params/PVExportParams.h"
+#include "Helpers/PVUtilities.h"
+
+#include "Engine/StaticMesh.h"
+#include "Engine/SkeletalMesh.h"
+#include "Animation/Skeleton.h"
+
+UClass* FPVExportParams::GetMeshClass() const
+{
+	return ExportMeshType == EPVExportMeshType::StaticMesh ? UStaticMesh::StaticClass() : USkeletalMesh::StaticClass();
+}
+
+bool FPVExportParams::Validate(FString& OutError) const
+{
+	const FString& OutputPath = ContentBrowserFolder.Path;
+
+	if (!PV::Utilities::ValidateAssetPathAndName(MeshName.ToString(), OutputPath, GetMeshClass(), OutError))
+	{
+		return false;
+	}
+
+	if (ExportMeshType == EPVExportMeshType::SkeletalMesh)
+	{
+		if (!PV::Utilities::ValidateAssetPathAndName(GetOutputSkeletonName(), OutputPath, USkeleton::StaticClass(), OutError))
+		{
+			return false;
+		}
+	}
+
+	if (ReplacePolicy == EPVAssetReplacePolicy::Replace)
+	{
+		const TCHAR* ConflictingAssetError = TEXT("Create a Unique name for the output as another asset of different type {0} exists with same name {1}");
+
+		const FString MeshPackagePath = GetOutputMeshPackagePath();
+		if (PV::Utilities::DoesConflictingPackageExist(MeshPackagePath, GetMeshClass()))
+		{
+			OutError = FString::Format(ConflictingAssetError, { *GetMeshClass()->GetName(), *MeshPackagePath });
+			return false;
+		}
+
+		if (ExportMeshType == EPVExportMeshType::SkeletalMesh)
+		{
+			const FString SkeletonPackagePath = GetOutputSkeletonPackagePath();
+			if (PV::Utilities::DoesConflictingPackageExist(SkeletonPackagePath, USkeleton::StaticClass()))
+			{
+				OutError = FString::Format(ConflictingAssetError, { *USkeleton::StaticClass()->GetName(), *SkeletonPackagePath });
+				return false;
+			}
+		}
+	}
+
+	return true;
+}
\ No newline at end of file
diff --git a/Engine/Plugins/Experimental/ProceduralVegetationEditor/Source/ProceduralVegetation/Private/Params/PVImportStaticMeshParams.h b/Engine/Plugins/Experimental/ProceduralVegetationEditor/Source/ProceduralVegetation/Private/Params/PVImportStaticMeshParams.h
new file mode 100644
--- /dev/null
+++ b/Engine/Plugins/Experimental/ProceduralVegetationEditor/Source/ProceduralVegetation/Private/Params/PVImportStaticMeshParams.h
@@ -0,0 +1,39 @@
+// Copyright Epic Games, Inc. All Rights Reserved.
+#pragma once
+
+#include "Engine/EngineTypes.h"
+
+#include "PVImportStaticMeshParams.generated.h"
+
+class UStaticMesh;
+
+UENUM(BlueprintType)
+enum class EPVImportStaticMeshDebugState : uint8
+{
+	None,
+	InitialMesh,
+	KeepMaterial,
+	ComputeNormals,
+	ComputeVerticesCenterAndRadius,
+	CollapseBranchesToLine,
+	FindEndPoints,
+	ConnectBranchesToParent,
+	FindRootPoints,
+	ComputeBranchHierarchies,
+	MAX UMETA(Hidden)
+};
+
+USTRUCT(BlueprintType)
+struct FPVImportStaticMeshParams
+{
+	GENERATED_BODY()
+
+	UPROPERTY(EditAnywhere, Category = "Asset")
+	TObjectPtr<UStaticMesh> StaticMeshAsset = nullptr;
+
+	UPROPERTY(EditAnywhere, Category = "Asset")
+	TArray<FName> MaterialsToKeep;
+
+	UPROPERTY(EditAnywhere, Category = "Asset")
+	EPVImportStaticMeshDebugState DebugState = EPVImportStaticMeshDebugState::None;
+};
diff --git a/Engine/Plugins/Experimental/ProceduralVegetationEditor/Source/ProceduralVegetation/Private/Params/PVImportTexture2DParams.h b/Engine/Plugins/Experimental/ProceduralVegetationEditor/Source/ProceduralVegetation/Private/Params/PVImportTexture2DParams.h
new file mode 100644
--- /dev/null
+++ b/Engine/Plugins/Experimental/ProceduralVegetationEditor/Source/ProceduralVegetation/Private/Params/PVImportTexture2DParams.h
@@ -0,0 +1,87 @@
+// Copyright Epic Games, Inc. All Rights Reserved.
+#pragma once
+
+#include "Engine/EngineTypes.h"
+#include "GeometryCollection/ManagedArrayCollection.h"
+#include "PVImportTexture2DParams.generated.h"
+
+class UTexture2D;
+
+UENUM(BlueprintType)
+enum class EPVImportTexture2DDebugState : uint8
+{
+	None,
+	SampleImage,
+	TriangulateCurves,
+	ComputeVertexScale,
+	ComputeBranchHierarchies,
+	MAX UMETA(Hidden)
+};
+
+USTRUCT(BlueprintType)
+struct FPVImportTexture2DPlantSettings
+{
+	GENERATED_BODY()
+
+	UPROPERTY(EditAnywhere, Category = "Plant Settings")
+	FVector2f RootPosition = FVector2f::ZeroVector;
+
+	UPROPERTY(EditAnywhere, Category = "Plant Settings")
+	float Rotation = 0;
+};
+
+USTRUCT(BlueprintType)
+struct FPVImportTexture2DParams
+{
+	GENERATED_BODY()
+
+	UPROPERTY(EditAnywhere, Category = "Image Trace", meta=(ClampMin=0, UIMin=0))
+	int32 SampleWidth = 512;
+	
+	UPROPERTY(EditAnywhere, Category = "Image Trace", meta=(ClampMin=0, UIMin=0))
+    int32 SampleHeight = 512;
+    
+    // Any plants with a bounds area lower than this value will be ignored.
+    // Value is normalized meaning a value of 0.1 will filter out any plants whose bounds cover less than 10% of the image. 
+	UPROPERTY(EditAnywhere, Category = "Image Trace", meta=(ClampMin=0, UIMin=0, UIMax=0.2))
+	float MinBoundsArea = 0.05f;
+	
+	UPROPERTY(EditAnywhere, Category = "Image Trace", meta=(ClampMin=0, ClampMax=100, UIMin=0, UIMax=10))
+	int32 SmoothingIterations = 2;
+	
+	UPROPERTY(EditAnywhere, Category = "Image Trace", meta=(ClampMin=0, ClampMax=100, UIMin=0, UIMax=10.0))
+	float SimplificationAmount = 1.0;
+
+	UPROPERTY(EditAnywhere, Category = "Texture")
+	TObjectPtr<UTexture2D> Texture2DAsset = nullptr;
+	
+	UPROPERTY(EditAnywhere, Category = "Texture")
+	bool bInvertImage = false;
+	
+	UPROPERTY(EditAnywhere, Category = "Texture", meta=(ClampMin=0, ClampMax=1, UIMin=0, UIMax=1))
+	float WhiteLevel = 0.01;
+
+	UPROPERTY(EditAnywhere, Category = "Plant Geneneration")
+	float Scale = 100.f;
+
+	UPROPERTY(EditAnywhere, Category = "Plant Geneneration")
+	TArray<FPVImportTexture2DPlantSettings> PlantSettings;
+
+	UPROPERTY(EditAnywhere, Category = "Plant Geneneration|Tip detection", meta = (ClampMin=0.001, ClampMax=180, UIMin=0.01, UIMax=180))
+	float TipAngleThresholdInDegrees = 45.0;
+
+	UPROPERTY(EditAnywhere, Category = "Plant Geneneration|Tip detection", meta = (ClampMin=0, ClampMax=0.01, UIMin=0, UIMax=0.01))
+	float MaxTipAngleSearchDist = 0.005;
+
+	UPROPERTY(EditAnywhere, Category = "DebugSettings")
+	EPVImportTexture2DDebugState DebugState = EPVImportTexture2DDebugState::None;
+};
+
+struct FPVImportTexture2DOutput
+{
+	TOptional<FManagedArrayCollection> DebugCollection;
+	FManagedArrayCollection VisualizationCollection;
+
+	FManagedArrayCollection CombinedPlantCollection;
+	TArray<FManagedArrayCollection> PlantCollections;
+};
\ No newline at end of file
diff --git a/Engine/Plugins/Experimental/ProceduralVegetationEditor/Source/ProceduralVegetation/Private/Tests/PVImportTests.cpp b/Engine/Plugins/Experimental/ProceduralVegetationEditor/Source/ProceduralVegetation/Private/Tests/PVImportTests.cpp
new file mode 100644
--- /dev/null
+++ b/Engine/Plugins/Experimental/ProceduralVegetationEditor/Source/ProceduralVegetation/Private/Tests/PVImportTests.cpp
@@ -0,0 +1,353 @@
+// Copyright Epic Games, Inc. All Rights Reserved.
+
+#if WITH_DEV_AUTOMATION_TESTS
+
+#include "Tests/PVTestsCommon.h"
+#include "Helpers/PVImportHelpers.h"
+#include "Facades/PVAttributesNames.h"
+
+#define PV_IMPORTER_TEST(TestName) PV_SIMPLE_AUTOMATION_TEST(Importer, TestName)
+
+namespace PVMockTree
+{
+	FPVBranchHierarchyDescription CreateBranchHierarchyDescription()
+	{
+		return PV::ImportHelper::CreateBranchHierarchyFromPoints(PVMockTree::Points, PVMockTree::PointsRadii, PVMockTree::BranchPointIndices);
+	}
+}
+
+PV_IMPORTER_TEST(FillAttributesFromBranchHierarchy)
+{
+	const FPVBranchHierarchyDescription BranchHierarchy = PVMockTree::CreateBranchHierarchyDescription();
+
+	const int32 NumBranchPoints = BranchHierarchy.GetNumBranchPoints();
+	const int32 NumFoliagePoints = 0;
+	const int32 NumBranches = BranchHierarchy.GetNumBranches();
+	const int32 NumFoliageMeshes = 0;
+
+	FManagedArrayCollection ExpectedCollection;
+	PV::ImportHelper::CreateEmptyGrowthData(ExpectedCollection, NumBranchPoints, NumFoliagePoints, NumBranches, NumFoliageMeshes);
+
+	auto ExpectedPointPositionAttribute = PV::FPointPositionAttribute::GetAttribute(ExpectedCollection);
+	auto ExpectedBranchPointsAttribute = PV::FBranchPointsAttribute::GetAttribute(ExpectedCollection);
+	auto ExpectedBranchParentNumberAttribute = PV::FBranchParentNumberAttribute::GetAttribute(ExpectedCollection);
+	auto ExpectedBranchParentsAttribute = PV::FBranchParentsAttribute::GetAttribute(ExpectedCollection);
+	auto ExpectedBranchChildrenAttribute = PV::FBranchChildrenAttribute::GetAttribute(ExpectedCollection);
+
+	// Init Expected Point Attributes
+	{
+		ExpectedPointPositionAttribute[0] = PVMockTree::Point0;
+		ExpectedPointPositionAttribute[1] = PVMockTree::Point1;
+		ExpectedPointPositionAttribute[2] = PVMockTree::Point2;
+		ExpectedPointPositionAttribute[3] = PVMockTree::Point3;
+		ExpectedPointPositionAttribute[4] = PVMockTree::Point4;
+		ExpectedPointPositionAttribute[5] = PVMockTree::Point5;
+		ExpectedPointPositionAttribute[6] = PVMockTree::Point6;
+		ExpectedPointPositionAttribute[7] = PVMockTree::Point7;
+		ExpectedPointPositionAttribute[8] = PVMockTree::Point8;
+		ExpectedPointPositionAttribute[9] = PVMockTree::Point9;
+		ExpectedPointPositionAttribute[10] = PVMockTree::Point10;
+		ExpectedPointPositionAttribute[11] = PVMockTree::Point11;
+	}
+
+	// Init Expected Branch attributes
+	{
+		ExpectedBranchPointsAttribute[0] = PVMockTree::Branch0_PointIndices;
+		ExpectedBranchPointsAttribute[1] = PVMockTree::Branch1_PointIndices;
+		ExpectedBranchPointsAttribute[2] = PVMockTree::Branch2_PointIndices;
+		ExpectedBranchPointsAttribute[3] = PVMockTree::Branch3_PointIndices;
+
+		ExpectedBranchParentNumberAttribute[0] = 0;
+		ExpectedBranchParentNumberAttribute[1] = PVMockTree::Branch0_BranchNumber;
+		ExpectedBranchParentNumberAttribute[2] = PVMockTree::Branch1_BranchNumber;
+		ExpectedBranchParentNumberAttribute[3] = PVMockTree::Branch0_BranchNumber;
+
+		ExpectedBranchParentsAttribute[0] = { };
+		ExpectedBranchParentsAttribute[1] = { PVMockTree::Branch0_BranchNumber };
+		ExpectedBranchParentsAttribute[2] = { PVMockTree::Branch0_BranchNumber, PVMockTree::Branch1_BranchNumber };
+		ExpectedBranchParentsAttribute[3] = { PVMockTree::Branch0_BranchNumber };
+
+		ExpectedBranchChildrenAttribute[0] = { PVMockTree::Branch1_BranchNumber, PVMockTree::Branch3_BranchNumber };
+		ExpectedBranchChildrenAttribute[1] = { PVMockTree::Branch2_BranchNumber };
+		ExpectedBranchChildrenAttribute[2] = { };
+		ExpectedBranchChildrenAttribute[3] = { };
+	}
+
+	FManagedArrayCollection ActualCollection;
+	PV::ImportHelper::CreateEmptyGrowthData(ActualCollection, NumBranchPoints, NumFoliagePoints, NumBranches, NumFoliageMeshes);
+
+	UTEST_EQUAL("CorrectNrOfPointsInGroup", ActualCollection.NumElements(PV::GroupNames::PointGroup), NumBranchPoints);
+	UTEST_EQUAL("CorrectNrOfBranchesInGroup", ActualCollection.NumElements(PV::GroupNames::BranchGroup), NumBranches);
+
+	PV::ImportHelper::FillAttributesFromBranchHierarchy(ActualCollection, { BranchHierarchy });
+
+	const auto& ActualPointPositionAttribute = PV::FPointPositionAttribute::GetAttribute(ActualCollection);
+	const auto& ActualBranchPointsAttribute = PV::FBranchPointsAttribute::GetAttribute(ActualCollection);
+	const auto& ActualBranchParentNumberAttribute = PV::FBranchParentNumberAttribute::GetAttribute(ActualCollection);
+	const auto& ActualBranchParentsAttribute = PV::FBranchParentsAttribute::GetAttribute(ActualCollection);
+	const auto& ActualBranchChildrenAttribute = PV::FBranchChildrenAttribute::GetAttribute(ActualCollection);
+
+	// Validate that the points are correctly indexed and in the correct location
+	UTEST_TRUE("PointPosition0",  ActualPointPositionAttribute[0].Equals(ExpectedPointPositionAttribute[0]));
+	UTEST_TRUE("PointPosition1",  ActualPointPositionAttribute[1].Equals(ExpectedPointPositionAttribute[1]));
+	UTEST_TRUE("PointPosition2",  ActualPointPositionAttribute[2].Equals(ExpectedPointPositionAttribute[2]));
+	UTEST_TRUE("PointPosition3",  ActualPointPositionAttribute[3].Equals(ExpectedPointPositionAttribute[3]));
+	UTEST_TRUE("PointPosition4",  ActualPointPositionAttribute[4].Equals(ExpectedPointPositionAttribute[4]));
+	UTEST_TRUE("PointPosition5",  ActualPointPositionAttribute[5].Equals(ExpectedPointPositionAttribute[5]));
+	UTEST_TRUE("PointPosition6",  ActualPointPositionAttribute[6].Equals(ExpectedPointPositionAttribute[6]));
+	UTEST_TRUE("PointPosition7",  ActualPointPositionAttribute[7].Equals(ExpectedPointPositionAttribute[7]));
+	UTEST_TRUE("PointPosition8",  ActualPointPositionAttribute[8].Equals(ExpectedPointPositionAttribute[8]));
+	UTEST_TRUE("PointPosition9",  ActualPointPositionAttribute[9].Equals(ExpectedPointPositionAttribute[9]));
+	UTEST_TRUE("PointPosition10", ActualPointPositionAttribute[10].Equals(ExpectedPointPositionAttribute[10]));
+	UTEST_TRUE("PointPosition11", ActualPointPositionAttribute[11].Equals(ExpectedPointPositionAttribute[11]));
+
+	// Validate that the branch point indices are correct
+	UTEST_EQUAL("BranchPoints0", ActualBranchPointsAttribute[0], ExpectedBranchPointsAttribute[0]);
+	UTEST_EQUAL("BranchPoints1", ActualBranchPointsAttribute[1], ExpectedBranchPointsAttribute[1]);
+	UTEST_EQUAL("BranchPoints2", ActualBranchPointsAttribute[2], ExpectedBranchPointsAttribute[2]);
+	UTEST_EQUAL("BranchPoints3", ActualBranchPointsAttribute[3], ExpectedBranchPointsAttribute[3]);
+
+	// Validate that the branch parents are correct
+	UTEST_EQUAL("ParentBranchNumber0", ActualBranchParentNumberAttribute[0], ExpectedBranchParentNumberAttribute[0]);
+	UTEST_EQUAL("ParentBranchNumber1", ActualBranchParentNumberAttribute[1], ExpectedBranchParentNumberAttribute[1]);
+	UTEST_EQUAL("ParentBranchNumber2", ActualBranchParentNumberAttribute[2], ExpectedBranchParentNumberAttribute[2]);
+	UTEST_EQUAL("ParentBranchNumber3", ActualBranchParentNumberAttribute[3], ExpectedBranchParentNumberAttribute[3]);
+	UTEST_EQUAL("BranchParents0", ActualBranchParentsAttribute[0], ExpectedBranchParentsAttribute[0]);
+	UTEST_EQUAL("BranchParents1", ActualBranchParentsAttribute[1], ExpectedBranchParentsAttribute[1]);
+	UTEST_EQUAL("BranchParents2", ActualBranchParentsAttribute[2], ExpectedBranchParentsAttribute[2]);
+	UTEST_EQUAL("BranchParents3", ActualBranchParentsAttribute[3], ExpectedBranchParentsAttribute[3]);
+
+	// Validate branch children
+	UTEST_EQUAL("BranchChildren0", ActualBranchChildrenAttribute[0], ExpectedBranchChildrenAttribute[0]);
+	UTEST_EQUAL("BranchChildren1", ActualBranchChildrenAttribute[1], ExpectedBranchChildrenAttribute[1]);
+	UTEST_EQUAL("BranchChildren2", ActualBranchChildrenAttribute[2], ExpectedBranchChildrenAttribute[2]);
+	UTEST_EQUAL("BranchChildren3", ActualBranchChildrenAttribute[3], ExpectedBranchChildrenAttribute[3]);
+
+	return true;
+}
+
+PV_IMPORTER_TEST(FindRootBranch)
+{
+	const TArray<TArray<int32>>& Branches = PVMockTree::BranchPointIndices;
+	const int32 RootBranchIndex = PV::ImportHelper::FindRootBranch(PVMockTree::Points, Branches);
+
+	UTEST_NOT_EQUAL("ValidRootBranch", RootBranchIndex, int32(INDEX_NONE));
+
+	const TArray<int32> RootBranchIndices = Branches[RootBranchIndex];
+	const FVector3f RootVertex = PVMockTree::Points[RootBranchIndices[0]];
+
+	UTEST_NEARLY_EQUAL("CorrectRootBranchLocation", FVector(RootVertex), FVector::ZeroVector, UE_KINDA_SMALL_NUMBER);
+
+	return true;
+}
+
+PV_IMPORTER_TEST(EstimateBranchHierarchy)
+{
+	// Test standard MockTree
+	{
+		TArray<TArray<int32>> Branches = {
+			PVMockTree::Branch0_PointIndices,
+			PVMockTree::Branch1_PointIndices,
+			PVMockTree::Branch2_PointIndices,
+			PVMockTree::Branch3_PointIndices
+		};
+
+		TArray<int32> ExpectedBranchHierarchy = { -1, 0, 1, 0 };
+
+		// Naive testing of hierarchy
+		{
+			const int32 RootBranchIndex = PV::ImportHelper::FindRootBranch(PVMockTree::Points, Branches);
+			const TArray<int32> BranchHierarchy = PV::ImportHelper::EstimateBranchHierarchy(PVMockTree::Points, Branches, RootBranchIndex);
+
+			UTEST_EQUAL("CorrectNrOfBranches", BranchHierarchy.Num(), ExpectedBranchHierarchy.Num());
+			for (int32 i = 0; i < Branches.Num(); ++i)
+			{
+				UTEST_EQUAL(FString::Printf(TEXT("Branch%d_CorrectParent"), i), BranchHierarchy[i], ExpectedBranchHierarchy[i]);
+			}
+		}
+
+		// Test with branch list reversed
+		{
+			Algo::Reverse(Branches);
+			Algo::Reverse(ExpectedBranchHierarchy);
+			for (int32& Hierarchy : ExpectedBranchHierarchy)
+			{
+				if (Hierarchy != -1)
+				{
+					Hierarchy = Branches.Num() - 1 - Hierarchy;
+				}
+			}
+
+			const int32 RootBranchIndex = PV::ImportHelper::FindRootBranch(PVMockTree::Points, Branches);
+			const TArray<int32> BranchHierarchy = PV::ImportHelper::EstimateBranchHierarchy(PVMockTree::Points, Branches, RootBranchIndex);
+			
+			UTEST_EQUAL("CorrectNrOfBranches_Reverse", BranchHierarchy.Num(), ExpectedBranchHierarchy.Num());
+			for (int32 i = 0; i < Branches.Num(); ++i)
+			{
+				UTEST_EQUAL(FString::Printf(TEXT("Branch%d_CorrectParent_Reverse"), i), BranchHierarchy[i], ExpectedBranchHierarchy[i]);
+			}
+		}
+	}
+
+	// Add additional branches to MockTree and ensure hierarchy stays correct
+	{
+		// Add a few branches starting from the same point as another as we want to make sure the algorithm can handle multiple 
+		// branches being at the same distance from the start point of another.
+		TArray<FVector3f> PointsCpy = PVMockTree::Points;
+		const int32 Point2_1_Index = PointsCpy.Add(PVMockTree::Point2 + FVector3f(-10, 0, 0));
+		const int32 Point2_2_Index = PointsCpy.Add(PVMockTree::Point2 + FVector3f(10, 0, 0));
+		const int32 Point2_3_Index = PointsCpy.Add(PVMockTree::Point2 + FVector3f(-10, 10, 0));
+		const int32 Point2_4_Index = PointsCpy.Add(PVMockTree::Point2 + FVector3f(10, 10, 0));
+		const int32 Point2_5_Index = PointsCpy.Add(PVMockTree::Point2 + FVector3f(-10, 10, -10));
+		const int32 Point2_6_Index = PointsCpy.Add(PVMockTree::Point2 + FVector3f(10, 8, 1));
+
+		const int32 Point8_1_Index = PointsCpy.Add(PVMockTree::Point8 + FVector3f(10, 0, 0));
+		const int32 Point8_2_Index = PointsCpy.Add(PVMockTree::Point8 + FVector3f(-10, 0, 0));
+		const int32 Point8_3_Index = PointsCpy.Add(PVMockTree::Point8 + FVector3f(10, 10, 0));
+		const int32 Point8_4_Index = PointsCpy.Add(PVMockTree::Point8 + FVector3f(-10, 10, 0));
+		const int32 Point8_5_Index = PointsCpy.Add(PVMockTree::Point8 + FVector3f(10, 10, 10));
+		const int32 Point8_6_Index = PointsCpy.Add(PVMockTree::Point8 + FVector3f(-10, 10, 10));
+
+		const int32 Point3_0_Index = PointsCpy.Add(PVMockTree::Point3 + FVector3f(10, 0, 0));
+		const int32 Point3_1_Index = PointsCpy.Add(PVMockTree::Point3 + FVector3f(-10, 0, 0));
+
+		TArray<TArray<int32>> Branches = {
+			PVMockTree::Branch0_PointIndices,
+			PVMockTree::Branch1_PointIndices,
+			PVMockTree::Branch2_PointIndices,
+			PVMockTree::Branch3_PointIndices,
+			{ 2, Point2_1_Index },
+			{ 2, Point2_2_Index },
+			{ 2, Point2_3_Index },
+			{ 2, Point2_4_Index },
+			{ 2, Point2_5_Index },
+			{ 2, Point2_6_Index },
+
+			{ 8, Point8_1_Index },
+			{ 8, Point8_2_Index },
+			{ 8, Point8_3_Index },
+			{ 8, Point8_4_Index },
+			{ 8, Point8_5_Index },
+			{ 8, Point8_6_Index },
+
+			{ 3, Point3_0_Index },
+			{ 3, Point3_1_Index },
+		};
+
+		TArray<int32> ExpectedBranchHierarchy = { 
+			-1, 0, 1, 0, 
+			0, 0, 0, 0, 0, 0,
+			1, 1, 1, 1, 1, 1,
+			0, 0
+		};
+
+		// Naive testing of hierarchy
+		{
+			const int32 RootBranchIndex = PV::ImportHelper::FindRootBranch(PointsCpy, Branches);
+			const TArray<int32> BranchHierarchy = PV::ImportHelper::EstimateBranchHierarchy(PointsCpy, Branches, RootBranchIndex);
+
+			UTEST_EQUAL("CorrectNrOfBranches", BranchHierarchy.Num(), ExpectedBranchHierarchy.Num());
+			for (int32 i = 0; i < Branches.Num(); ++i)
+			{
+				UTEST_EQUAL(FString::Printf(TEXT("Branch%d_CorrectParent"), i), BranchHierarchy[i], ExpectedBranchHierarchy[i]);
+			}
+		}
+
+		// Test with branch list reversed
+		{
+			Algo::Reverse(Branches);
+			Algo::Reverse(ExpectedBranchHierarchy);
+			for (int32& Hierarchy : ExpectedBranchHierarchy)
+			{
+				if (Hierarchy != -1)
+				{
+					Hierarchy = Branches.Num() - 1 - Hierarchy;
+				}
+			}
+
+			const int32 RootBranchIndex = PV::ImportHelper::FindRootBranch(PointsCpy, Branches);
+			const TArray<int32> BranchHierarchy = PV::ImportHelper::EstimateBranchHierarchy(PointsCpy, Branches, RootBranchIndex);
+
+			UTEST_EQUAL("CorrectNrOfBranches_Reverse", BranchHierarchy.Num(), ExpectedBranchHierarchy.Num());
+			for (int32 i = 0; i < Branches.Num(); ++i)
+			{
+				UTEST_EQUAL(FString::Printf(TEXT("Branch%d_CorrectParent_Reverse"), i), BranchHierarchy[i], ExpectedBranchHierarchy[i]);
+			}
+		}
+	}
+
+	return true;
+}
+
+PV_IMPORTER_TEST(CreateBranchHierarchyFromEdges)
+{
+	FPVBranchHierarchyDescription ExpectedBranchHierarchy;
+	ExpectedBranchHierarchy.RootBranchIndex = 0;
+	ExpectedBranchHierarchy.Points = PVMockTree::Points;
+
+	FPVBranchDescription& B0 = ExpectedBranchHierarchy.Branches.AddDefaulted_GetRef();
+	B0.BranchIndex = 0;
+	B0.PointIndices = PVMockTree::Branch0_PointIndices;
+	B0.ParentBranchIndex = -1;
+	B0.ChildBranchIndices = { 1, 3 };
+
+	FPVBranchDescription& B1 = ExpectedBranchHierarchy.Branches.AddDefaulted_GetRef();
+	B1.BranchIndex = 1;
+	B1.PointIndices = PVMockTree::Branch1_PointIndices;
+	B1.ParentBranchIndex = 0;
+	B1.ChildBranchIndices = { 2 };
+
+	FPVBranchDescription& B2 = ExpectedBranchHierarchy.Branches.AddDefaulted_GetRef();
+	B2.BranchIndex = 2;
+	B2.PointIndices = PVMockTree::Branch2_PointIndices;
+	B2.ParentBranchIndex = 1;
+	B2.ChildBranchIndices = { };
+
+	FPVBranchDescription& B3 = ExpectedBranchHierarchy.Branches.AddDefaulted_GetRef();
+	B3.BranchIndex = 3;
+	B3.PointIndices = PVMockTree::Branch3_PointIndices;
+	B3.ParentBranchIndex = 0;
+	B3.ChildBranchIndices = { };
+
+	const FPVBranchHierarchyDescription BranchHierarchy = PVMockTree::CreateBranchHierarchyDescription();
+
+	UTEST_EQUAL("CorrectRootIndex", BranchHierarchy.RootBranchIndex, ExpectedBranchHierarchy.RootBranchIndex);
+	UTEST_EQUAL("CorrectNrOfBranches", BranchHierarchy.Branches.Num(), ExpectedBranchHierarchy.Branches.Num());
+
+	for (int32 i = 0; i < BranchHierarchy.Branches.Num(); ++i)
+	{
+		const FPVBranchDescription& Branch = BranchHierarchy.Branches[i];
+		const FPVBranchDescription* ExpectedBranchPtr = ExpectedBranchHierarchy.Branches.FindByPredicate([&](const FPVBranchDescription& InBranch) { return InBranch.BranchIndex == Branch.BranchIndex; });
+
+		UTEST_NOT_NULL("ContainsBranchIndex", ExpectedBranchPtr);
+
+		const FPVBranchDescription& ExpectedBranch = *ExpectedBranchPtr;
+
+		const FString BranchName = FString::Printf(TEXT("Branch%d: "), i);
+
+		UTEST_EQUAL(BranchName + "_CorrectBranchIndex", Branch.BranchIndex, ExpectedBranch.BranchIndex);
+		UTEST_EQUAL(BranchName + "_CorrectParent", Branch.ParentBranchIndex, ExpectedBranch.ParentBranchIndex);
+		UTEST_EQUAL(BranchName + "_CorrectNrOfChildren", Branch.ChildBranchIndices.Num(), ExpectedBranch.ChildBranchIndices.Num());
+		UTEST_EQUAL(BranchName + "_CorrectNrOfPoints", Branch.PointIndices.Num(), ExpectedBranch.PointIndices.Num());
+
+		for (int32 j = 0; j < Branch.PointIndices.Num(); ++j)
+		{
+			const int32 PointIndex = Branch.PointIndices[j];
+			const int32 ExpectedPointIndex = ExpectedBranch.PointIndices[j];
+			UTEST_EQUAL(BranchName + FString::Printf(TEXT("_CorrectPointIndex %d"), j), PointIndex, ExpectedPointIndex);
+
+			const FVector Point = FVector(BranchHierarchy.Points[PointIndex]);
+			const FVector ExpectedPoint = FVector(ExpectedBranchHierarchy.Points[ExpectedPointIndex]);
+			UTEST_NEARLY_EQUAL(BranchName + FString::Printf(TEXT("_CorrectPoint %d"), j), Point, ExpectedPoint, UE_KINDA_SMALL_NUMBER);
+		}
+
+		for (int32 j = 0; j < Branch.ChildBranchIndices.Num(); ++j)
+		{
+			UTEST_TRUE(BranchName + FString::Printf(TEXT("_CorrectChildIndex %d"), j), ExpectedBranch.ChildBranchIndices.Contains(Branch.ChildBranchIndices[j]));
+		}
+	}
+
+	return true;
+}
+
+#endif
\ No newline at end of file
diff --git a/Engine/Plugins/Experimental/ProceduralVegetationEditor/Source/ProceduralVegetation/Private/Tests/PVMockTree.cpp b/Engine/Plugins/Experimental/ProceduralVegetationEditor/Source/ProceduralVegetation/Private/Tests/PVMockTree.cpp
--- a/Engine/Plugins/Experimental/ProceduralVegetationEditor/Source/ProceduralVegetation/Private/Tests/PVMockTree.cpp
+++ b/Engine/Plugins/Experimental/ProceduralVegetationEditor/Source/ProceduralVegetation/Private/Tests/PVMockTree.cpp
@@ -3,7 +3,7 @@
 #if WITH_DEV_AUTOMATION_TESTS
 
 #include "Tests/PVMockTree.h"
-#include "IPVEditorInterface.h"
+#include "Helpers/PVImportHelpers.h"
 #include "GeometryCollection/ManagedArrayCollection.h"
 
 #include "Facades/PVBranchFacade.h"
@@ -31,7 +31,6 @@ const TCHAR* PVMockTreeCollection::CompareResultToString(PVMockTreeCollection::C
 
 PVMockTreeCollection::CompareResult PVMockTreeCollection::Compare(const FManagedArrayCollection& Collection, bool bCheckPointPositions)
 {
-#if WITH_EDITOR
 	const PV::Facades::FPointFacade InputPointFacade(Collection);
 	const PV::Facades::FBranchFacade InputBranchFacade(Collection);
 	if (!InputPointFacade.IsValid() || !InputBranchFacade.IsValid())
@@ -67,20 +66,19 @@ PVMockTreeCollection::CompareResult PVMockTreeCollection::Compare(const FManaged
 			}
 		}
 	}
-#endif
+
 	return CompareResult::Success;
 }
 
 TSharedRef<FManagedArrayCollection> PVMockTreeCollection::CreateCollection(bool bSplitPoints)
 {
-#if WITH_EDITOR
-	if (PV::GPVEditorInterface)
-	{
-		return PV::GPVEditorInterface->CreateMockTreeGrowthDataCollection(bSplitPoints);
-	}
-#endif // WITH_EDITOR
+	const FPVBranchHierarchyDescription BranchHierarchy = PV::ImportHelper::CreateBranchHierarchyFromPoints(PVMockTree::Points, PVMockTree::PointsRadii, PVMockTree::BranchPointIndices);
+
+	TSharedRef<FManagedArrayCollection> OutCollection = MakeShared<FManagedArrayCollection>();
+	const bool bResult = PV::ImportHelper::GenerateGrowthDataFromBranchHierarchy(*OutCollection, BranchHierarchy, bSplitPoints);
+	check(bResult != false); // Mock tree is staticly defined to always be valid, so it should always return valid growth data
 
-	return MakeShared<FManagedArrayCollection>();
+	return OutCollection;
 }
 
 #endif // WITH_DEV_AUTOMATION_TESTS
diff --git a/Engine/Plugins/Experimental/ProceduralVegetationEditor/Source/ProceduralVegetation/Private/Utils/PVAttributes.h b/Engine/Plugins/Experimental/ProceduralVegetationEditor/Source/ProceduralVegetation/Private/Utils/PVAttributes.h
new file mode 100644
--- /dev/null
+++ b/Engine/Plugins/Experimental/ProceduralVegetationEditor/Source/ProceduralVegetation/Private/Utils/PVAttributes.h
@@ -0,0 +1,547 @@
+// Copyright Epic Games, Inc. All Rights Reserved.
+#pragma once
+
+#include "CoreMinimal.h"
+
+#include "GeometryCollection/ManagedArrayCollection.h"
+
+#include "Facades/PVAttributesNames.h"
+
+/*
+This file defines the TAttributeDefinition type which allows for compile-time type checking of attributes.
+
+- Defining a new attribute:
+	To define a new attribute, use the PV_DEFINE_ATTRIBUTE as follows:
+
+	inline constexpr TCHAR MyPropertyNameStr[] = TEXT("MyPropertyName");
+	inline constexpr TCHAR MyPropertyGroupStr[] = TEXT("MyPropertyGroup");
+    PV_DEFINE_ATTRIBUTE(MyProperty, FVector3f, MyPropertyNameStr, MyPropertyGroupStr);
+	PV_DEFINE_STRUCTURED_ATTRIBUTE can be used when custom element accessor views are needed.
+
+- Usage in code: 
+	PV_DEFINE_ATTRIBUTE defines a new attribute type, the example above defines the type FMyPropertyAttribute.
+	This type can be used with the FManagedArrayCollection to get, find, and add attributes using `FMyPropertyAttribute::FindAttribute(Collection)`, 
+	`FMyPropertyAttribute::GetAttribute(Collection)`, or `FMyPropertyAttribute::AddAttribute(Collection)`.
+
+	PV_DEFINE_ATTRIBUTE also defines a view and a const view for the attribute which can be used to access elements in the attribute.
+	Usage example:
+	void MultiplyMyProperty(FMyPropertyAttributeView MyPropertyAttribute, float Multiplier)
+	{
+		for (FVector3f& Vec : MyPropertyAttribute)
+		{
+			Vec *= Multiplier;
+		}
+	}
+	
+	// At call site:
+	MultiplyMyProperty(FMyPropertyAttribute::GetAttribute(Collection), 10.f);
+*/
+
+namespace PV
+{
+template<
+	typename InElementType,
+	typename InElementViewType,
+	typename InElementConstViewType,
+	const TCHAR* InAttributeNameStr, 
+	const TCHAR* InGroupNameStr
+>
+class TAttributeView
+{
+public:
+	using ElementType = std::remove_const_t<InElementType>;
+	
+	using ManagedArrayType =
+		std::conditional_t<std::is_const_v<InElementType>,
+		const TManagedArray<ElementType>,
+		TManagedArray<ElementType>>;
+	
+	using ElementViewType = 
+		std::conditional_t<std::is_const_v<InElementType>,
+		InElementConstViewType,
+		InElementViewType>;
+
+	using NonConstAttributeViewType = TAttributeView<ElementType, InElementViewType, InElementConstViewType, InAttributeNameStr, InGroupNameStr>;
+
+	using SizeType = TArray<ElementType>::SizeType;
+
+	inline static const FName AttributeName = FName(InAttributeNameStr);
+	inline static const FName GroupName = FName(InGroupNameStr);
+
+	ManagedArrayType* ManagedArray;
+
+	TAttributeView()
+		: ManagedArray(nullptr)
+	{}
+
+	TAttributeView(ManagedArrayType* InManagedArray)
+		: ManagedArray(InManagedArray)
+	{}
+
+	TAttributeView(const NonConstAttributeViewType& Other)
+		: ManagedArray(Other.ManagedArray)
+	{}
+
+	bool IsValid() const
+	{
+		return ManagedArray != nullptr;
+	}
+
+	ElementViewType operator[](SizeType Index) const
+	{
+		check(IsValid() && ManagedArray->IsValidIndex(Index));
+		return ElementViewType((*ManagedArray)[Index]);
+	}
+
+	operator bool() const 
+	{
+		return IsValid();
+	}
+
+	bool operator!() const
+	{
+		return !IsValid();
+	}
+
+	/**
+	* return true if index is in array range.
+	*
+	* @param Index Index to check.
+	*/
+	bool IsValidIndex(SizeType Index) const
+	{
+		return ManagedArray && ManagedArray->IsValidIndex(Index);
+	}
+
+	/**
+	* Find first index of the element
+	*/
+	int32 Find(const ElementType& Item) const
+	{
+		return ManagedArray ? ManagedArray->Find(Item) : INDEX_NONE;
+	}
+
+	SizeType Num() const
+	{
+		return ManagedArray ? ManagedArray->Num() : 0;
+	}
+
+	using RangedForIteratorType      = ManagedArrayType::RangedForIteratorType;
+	using RangedForConstIteratorType = ManagedArrayType::RangedForConstIteratorType;
+
+	using AttributeRangedForIteratorType =
+		std::conditional_t<std::is_const_v<InElementType>,
+		RangedForConstIteratorType,
+		RangedForIteratorType>;
+
+	/**
+	* DO NOT USE DIRECTLY
+	* STL-like iterators to enable range-based for loop support.
+	*/
+	FORCEINLINE AttributeRangedForIteratorType begin() const { check(ManagedArray); return ManagedArray->begin(); }
+	FORCEINLINE AttributeRangedForIteratorType end() const { check(ManagedArray); return ManagedArray->end(); }
+};
+
+template<
+	typename InAttributeType, 
+	typename InElementViewType,
+	typename InElementConstViewType,
+	const TCHAR* InAttributeNameStr, 
+	const TCHAR* InGroupNameStr
+>
+class TAttributeDefinition
+{
+	static_assert(!std::is_const_v<InAttributeType>, "InAttributeType cannot be const");
+
+public:
+	using AttributeType = InAttributeType;
+	using AttributeView = TAttributeView<AttributeType, InElementViewType, InElementConstViewType, InAttributeNameStr, InGroupNameStr>;
+	using AttributeConstView = TAttributeView<const AttributeType, InElementViewType, InElementConstViewType, InAttributeNameStr, InGroupNameStr>;
+
+	inline static const FName AttributeName = FName(InAttributeNameStr);
+	inline static const FName GroupName = FName(InGroupNameStr);
+
+	/**
+	* Tries to find the attribute in the supplied collection. Returns a null view if attribute is not in collection.
+	* @return A const view of the found attribute.
+	*/
+	static AttributeConstView FindAttribute(const FManagedArrayCollection& InCollection)
+	{
+		return AttributeConstView(InCollection.FindAttributeTyped<AttributeType>(AttributeName, GroupName));
+	}
+	
+	/**
+	* Tries to find the attribute in the supplied collection. Returns a null view if attribute is not in collection.
+	* @return A mutable view of the found attribute.
+	*/
+	static AttributeView FindAttribute(FManagedArrayCollection& InCollection)
+	{
+		return AttributeView(InCollection.FindAttributeTyped<AttributeType>(AttributeName, GroupName));
+	}
+
+	/**
+	* Checked version of FindAttribute. Will assert if attribute is missing from collection.
+	* @return A const view of the found attribute.
+	*/
+	static AttributeConstView GetAttribute(const FManagedArrayCollection& InCollection)
+	{
+		AttributeConstView AttributeView = FindAttribute(InCollection);
+		check(AttributeView.IsValid());
+		return AttributeView;
+	}
+
+	/**
+	* Checked version of FindAttribute. Will assert if attribute is missing from collection.
+	* @return A mutable view of the found attribute.
+	*/ 
+	static AttributeView GetAttribute(FManagedArrayCollection& InCollection)
+	{
+		AttributeView AttributeView = FindAttribute(InCollection);
+		check(AttributeView.IsValid());
+		return AttributeView;
+	}
+
+	/**
+	* Add the attribute to the collection and returns a view to the added collection. Returns existing attribute in the collection if already present.
+	* @return A mutable view of the found attribute.
+	*/
+	static AttributeView AddAttribute(FManagedArrayCollection& InCollection)
+	{
+		return AttributeView(&InCollection.AddAttribute<AttributeType>(AttributeName, GroupName));
+	}
+
+	/**
+	* Add the attribute to the collection and returns a view to the added collection. Returns existing attribute in the collection if already present.
+	* @param DefaultValue - Initialize all elements in the added attribute with this value (only if the attribute is not already present)
+	* @return A mutable view of the found attribute.
+	*/
+	static AttributeView AddAttribute(FManagedArrayCollection& InCollection, const AttributeType& DefaultValue)
+	{
+		AttributeView AttributeView = FindAttribute(InCollection);
+		if (AttributeView.IsValid())
+		{
+			return AttributeView;
+		}
+
+		AttributeView = AddAttribute(InCollection);
+		for (int32 i = 0; i < AttributeView.Num(); ++i)
+		{
+			AttributeView[i] = DefaultValue;
+		}
+		return AttributeView;
+	}
+
+	/**
+	* Check for the existence of the attribute in the collection.
+	*/
+	static bool HasAttribute(const FManagedArrayCollection& InCollection)
+	{
+		return InCollection.HasAttribute(AttributeName, GroupName);
+	}
+};
+
+// Helper class to allow for initialization, and selecton of, a fixed size array based of off default input array.
+template<typename InElementType>
+struct TFixedSizeArrayHelper
+{
+	using ElementType = std::remove_const_t<InElementType>;
+	using MutableArrayType = TArray<ElementType>;
+	using ConstArrayType = const TArray<ElementType>;
+	using ArrayType = std::conditional_t<std::is_const_v<InElementType>, ConstArrayType, MutableArrayType>;
+
+	TFixedSizeArrayHelper() = delete;
+
+	static MutableArrayType& GetFixedArray(MutableArrayType& InMutableArray, const ConstArrayType& InDefaultData)
+	{
+		if (InMutableArray.Num() < InDefaultData.Num())
+		{
+			const int32 FirstIndex = InMutableArray.Num();
+			InMutableArray.SetNum(InDefaultData.Num());
+			for (int32 i = FirstIndex; i < InDefaultData.Num(); ++i)
+			{
+				InMutableArray[i] = InDefaultData[i];
+			}
+		}
+
+		return InMutableArray;
+	}
+
+	static ConstArrayType& GetFixedArray(ConstArrayType& InConstArray, const ConstArrayType& InDefaultData)
+	{
+		if (InConstArray.Num() < InDefaultData.Num())
+		{
+			return InDefaultData;
+		}
+		return InConstArray;
+	}
+};
+
+#define PV_DEFINE_STRUCTURED_ATTRIBUTE(Name, Type, ElementViewType, ElementConstViewType, AttributeNameStr, AttributeGroupStr)          \
+	using F##Name##Attribute = TAttributeDefinition<Type, ElementViewType, ElementConstViewType, AttributeNameStr, AttributeGroupStr>;  \
+	using F##Name##AttributeView = F##Name##Attribute::AttributeView;                                                                   \
+	using F##Name##AttributeConstView = F##Name##Attribute::AttributeConstView;
+
+#define PV_DEFINE_ATTRIBUTE(Name, Type, AttributeNameStr, AttributeGroupStr) \
+	PV_DEFINE_STRUCTURED_ATTRIBUTE(Name, Type, Type&, const Type&, AttributeNameStr, AttributeGroupStr)
+
+// Branch Attributes
+PV_DEFINE_ATTRIBUTE(BranchParents, TArray<int32>, AttributeNames::BranchParentsStr, GroupNames::BranchGroupStr);
+PV_DEFINE_ATTRIBUTE(BranchChildren, TArray<int32>, AttributeNames::BranchChildrenStr, GroupNames::BranchGroupStr);
+PV_DEFINE_ATTRIBUTE(BranchPoints, TArray<int32>, AttributeNames::BranchPointsStr, GroupNames::BranchGroupStr);
+PV_DEFINE_ATTRIBUTE(BranchNumber, int32, AttributeNames::BranchNumberStr, GroupNames::BranchGroupStr);
+PV_DEFINE_ATTRIBUTE(BranchSourceBudNumber, int32, AttributeNames::BranchSourceBudNumberStr, GroupNames::BranchGroupStr);
+PV_DEFINE_ATTRIBUTE(BranchFoliageIDs, TArray<int32>, AttributeNames::BranchFoliageIDsStr, GroupNames::BranchGroupStr);
+PV_DEFINE_ATTRIBUTE(BranchUVMaterial, int32, AttributeNames::BranchUVMaterialStr, GroupNames::BranchGroupStr);
+PV_DEFINE_ATTRIBUTE(BranchHierarchyNumber, int32, AttributeNames::BranchHierarchyNumberStr, GroupNames::BranchGroupStr);
+PV_DEFINE_ATTRIBUTE(BranchSimulationGroupIndex, int32, AttributeNames::BranchSimulationGroupIndexStr, GroupNames::BranchGroupStr);
+PV_DEFINE_ATTRIBUTE(BranchPlantNumber, int32, AttributeNames::PlantNumberStr, GroupNames::BranchGroupStr);
+PV_DEFINE_ATTRIBUTE(BranchParentNumber, int32, AttributeNames::BranchParentNumberStr, GroupNames::BranchGroupStr);
+
+// Point Attributes
+PV_DEFINE_ATTRIBUTE(PointPosition, FVector3f, AttributeNames::PointPositionStr, GroupNames::PointGroupStr);
+PV_DEFINE_ATTRIBUTE(PointLengthFromRoot, float, AttributeNames::LengthFromRootStr, GroupNames::PointGroupStr);
+PV_DEFINE_ATTRIBUTE(PointLengthFromSeed, float, AttributeNames::LengthFromSeedStr, GroupNames::PointGroupStr);
+PV_DEFINE_ATTRIBUTE(PointScale, float, AttributeNames::PointScaleStr, GroupNames::PointGroupStr);
+PV_DEFINE_ATTRIBUTE(PointTextureCoordV, float, AttributeNames::TextureCoordVStr, GroupNames::PointGroupStr);
+PV_DEFINE_ATTRIBUTE(PointTextureCoordUOffset, float, AttributeNames::TextureCoordUOffsetStr, GroupNames::PointGroupStr);
+PV_DEFINE_ATTRIBUTE(PointURange, FVector2f, AttributeNames::URangeStr, GroupNames::PointGroupStr);
+PV_DEFINE_ATTRIBUTE(PointHullGradient, float, AttributeNames::HullGradientStr, GroupNames::PointGroupStr);
+PV_DEFINE_ATTRIBUTE(PointMainTrunkGradient, float, AttributeNames::MainTrunkGradientStr, GroupNames::PointGroupStr);
+PV_DEFINE_ATTRIBUTE(PointGroundGradient, float, AttributeNames::GroundGradientStr, GroupNames::PointGroupStr);
+PV_DEFINE_ATTRIBUTE(PointScaleGradient, float, AttributeNames::PointScaleGradientStr, GroupNames::PointGroupStr);
+PV_DEFINE_ATTRIBUTE(PointBranchGradient, float, AttributeNames::BranchGradientStr, GroupNames::PointGroupStr);
+PV_DEFINE_ATTRIBUTE(PointPlantGradient, float, AttributeNames::PlantGradientStr, GroupNames::PointGroupStr);
+PV_DEFINE_ATTRIBUTE(PointBudNumber, int32, AttributeNames::BudNumberStr, GroupNames::PointGroupStr);
+PV_DEFINE_ATTRIBUTE(PointNjordPixelIndex, float, AttributeNames::NjordPixelIndexStr, GroupNames::PointGroupStr);
+
+template<typename TFloatType>
+struct TBudLightDetectedView
+{
+	using ArrayType = TFixedSizeArrayHelper<TFloatType>::ArrayType;
+
+	inline const static TArray<float> DefaultData = { 0, 0, 0, 0 };
+
+	ArrayType& Array;
+	TFloatType& Availible;
+	TFloatType& Resource;
+	TFloatType& Branch;
+	TFloatType& Collision;
+
+	TBudLightDetectedView() = delete;
+	TBudLightDetectedView(ArrayType& InData)
+		: Array(TFixedSizeArrayHelper<TFloatType>::GetFixedArray(InData, DefaultData))
+		, Availible(Array[0])
+		, Resource(Array[1])
+		, Branch(Array[2])
+		, Collision(Array[3])
+	{
+	}
+};
+using FBudLightDetectedView = TBudLightDetectedView<float>;
+using FBudLightDetectedConstView = TBudLightDetectedView<const float>;
+PV_DEFINE_STRUCTURED_ATTRIBUTE(BudLightDetected, TArray<float>, FBudLightDetectedView, FBudLightDetectedConstView, AttributeNames::BudLightDetectedStr, GroupNames::PointGroupStr);
+
+template<typename TFloatType>
+struct TBudLateralMeristemView
+{
+	using ArrayType = TFixedSizeArrayHelper<TFloatType>::ArrayType;
+	
+	inline const static TArray<float> DefaultData = { 0, 0, 0, 0, 0, 0, 0 };
+
+	ArrayType& Array;
+	TFloatType& LateralMeristem;
+	TFloatType& Multiplier;
+	TFloatType& Inactive;
+	TFloatType& Davinci;
+	TFloatType& ParentDot;
+	TFloatType& RootDistance;
+	TFloatType& Degredation;
+
+	TBudLateralMeristemView() = delete;
+	TBudLateralMeristemView(ArrayType& InData)
+		: Array(TFixedSizeArrayHelper<TFloatType>::GetFixedArray(InData, DefaultData))
+		, LateralMeristem(Array[0])
+		, Multiplier(Array[1])
+		, Inactive(Array[2])
+		, Davinci(Array[3])
+		, ParentDot(Array[4])
+		, RootDistance(Array[5])
+		, Degredation(Array[6])
+	{
+	}
+};
+using FBudLateralMeristemView = TBudLateralMeristemView<float>;
+using FBudLateralMeristemConstView = TBudLateralMeristemView<const float>;
+PV_DEFINE_STRUCTURED_ATTRIBUTE(BudLateralMeristem, TArray<float>, FBudLateralMeristemView, FBudLateralMeristemConstView, AttributeNames::BudLateralMeristemStr, GroupNames::PointGroupStr);
+
+template<typename TFloatType>
+struct TBudHormoneLevelsView
+{
+	using ArrayType = TFixedSizeArrayHelper<TFloatType>::ArrayType;
+
+	inline const static TArray<float> DefaultData = { 0, 0, 0, 0, 0, 0 };
+
+	ArrayType& Array;
+	TFloatType& Apical;
+	TFloatType& Axillary;
+	TFloatType& AxillaryInhibition;
+	TFloatType& Radical;
+	TFloatType& Ethylene;
+	TFloatType& Cykininin;
+
+	TBudHormoneLevelsView() = delete;
+	TBudHormoneLevelsView(ArrayType& InData)
+		: Array(TFixedSizeArrayHelper<TFloatType>::GetFixedArray(InData, DefaultData))
+		, Apical(Array[0])
+		, Axillary(Array[1])
+		, AxillaryInhibition(Array[2])
+		, Radical(Array[3])
+		, Ethylene(Array[4])
+		, Cykininin(Array[5])
+	{
+	}
+};
+using FBudHormoneLevelsView = TBudHormoneLevelsView<float>;
+using FBudHormoneLevelsConstView = TBudHormoneLevelsView<const float>;
+PV_DEFINE_STRUCTURED_ATTRIBUTE(BudHormoneLevels, TArray<float>, FBudHormoneLevelsView, FBudHormoneLevelsConstView, AttributeNames::BudHormoneLevelsStr, GroupNames::PointGroupStr);
+
+template<typename TVec3Type>
+struct TBudDirectionView
+{
+	using ArrayType = TFixedSizeArrayHelper<TVec3Type>::ArrayType;
+
+	inline const static TArray<FVector3f> DefaultData = { FVector3f::ZeroVector, FVector3f::ZeroVector, FVector3f::ZeroVector, FVector3f::ZeroVector, FVector3f::ZeroVector, FVector3f::ZeroVector };
+
+	ArrayType& Array;
+	TVec3Type& Apical;
+	TVec3Type& Axillary;
+	TVec3Type& LightOptimal;
+	TVec3Type& LightSubOptimal;
+	TVec3Type& GuideCurve;
+	TVec3Type& UpVector;
+
+	TBudDirectionView() = delete;
+	TBudDirectionView(ArrayType& InData)
+		: Array(TFixedSizeArrayHelper<TVec3Type>::GetFixedArray(InData, DefaultData))
+		, Apical(Array[0])
+		, Axillary(Array[1])
+		, LightOptimal(Array[2])
+		, LightSubOptimal(Array[3])
+		, GuideCurve(Array[4])
+		, UpVector(Array[5])
+	{
+	}
+};
+using FBudDirectionView = TBudDirectionView<FVector3f>;
+using FBudDirectionConstView = TBudDirectionView<const FVector3f>;
+PV_DEFINE_STRUCTURED_ATTRIBUTE(BudDirection, TArray<FVector3f>, FBudDirectionView, FBudDirectionConstView, AttributeNames::BudDirectionStr, GroupNames::PointGroupStr);
+
+template<typename TIntType>
+struct TBudStatusView
+{
+	using ArrayType = TFixedSizeArrayHelper<TIntType>::ArrayType;
+	
+	inline const static TArray<int32> DefaultData = { 0, 0, 0, 0, 0, 0, 0, 0, 0, 0 };
+
+	ArrayType& Array;
+	TIntType& ApicalMeristem;
+	TIntType& Codominant;
+	TIntType& Axillary;
+	TIntType& Seed;
+	TIntType& Dormant;
+	TIntType& Triggered;
+	TIntType& NumTriggered;
+	TIntType& Inactive;
+	TIntType& BrokenTip;
+	TIntType& Broken;
+
+	TBudStatusView() = delete;
+	TBudStatusView(ArrayType& InData)
+		: Array(TFixedSizeArrayHelper<TIntType>::GetFixedArray(InData, DefaultData))
+		, ApicalMeristem(Array[0])
+		, Codominant(Array[1])
+		, Axillary(Array[2])
+		, Seed(Array[3])
+		, Dormant(Array[4])
+		, Triggered(Array[5])
+		, NumTriggered(Array[6])
+		, Inactive(Array[7])
+		, BrokenTip(Array[8])
+		, Broken(Array[9])
+	{
+	}
+};
+using FBudStatusView = TBudStatusView<int32>;
+using FBudStatusConstView = TBudStatusView<const int32>;
+PV_DEFINE_STRUCTURED_ATTRIBUTE(BudStatus, TArray<int32>, FBudStatusView, FBudStatusConstView, AttributeNames::BudStatusStr, GroupNames::PointGroupStr);
+
+template<typename TIntType>
+struct TBudDevelopmentView
+{
+	using ArrayType = TFixedSizeArrayHelper<TIntType>::ArrayType;
+
+	inline const static TArray<int32> DefaultData = { 0, 0, 0, 0, 0, 0 };
+
+	ArrayType& Array;
+	TIntType& Generation;
+	TIntType& BudAge;
+	TIntType& BranchAge;
+	TIntType& AgeSenescense;
+	TIntType& LightSenescense;
+	TIntType& RelativeBudAge;
+
+	TBudDevelopmentView() = delete;
+	TBudDevelopmentView(ArrayType& InData)
+		: Array(TFixedSizeArrayHelper<TIntType>::GetFixedArray(InData, DefaultData))
+		, Generation(Array[0])
+		, BudAge(Array[1])
+		, BranchAge(Array[2])
+		, AgeSenescense(Array[3])
+		, LightSenescense(Array[4])
+		, RelativeBudAge(Array[5])
+	{
+	}
+};
+using FBudDevelopmentView = TBudDevelopmentView<int32>;
+using FBudDevelopmentConstView = TBudDevelopmentView<const int32>;
+PV_DEFINE_STRUCTURED_ATTRIBUTE(BudDevelopment, TArray<int32>, FBudDevelopmentView, FBudDevelopmentConstView, AttributeNames::BudDevelopmentStr, GroupNames::PointGroupStr);
+
+// Foliage Attributes
+PV_DEFINE_ATTRIBUTE(FoliageNameID, int32, AttributeNames::FoliageNameIDStr, GroupNames::FoliageGroupStr);
+PV_DEFINE_ATTRIBUTE(FoliageBranchID, int32, AttributeNames::FoliageBranchIDStr, GroupNames::FoliageGroupStr);
+PV_DEFINE_ATTRIBUTE(FoliagePivotPoint, FVector3f, AttributeNames::FoliagePivotPointStr, GroupNames::FoliageGroupStr);
+PV_DEFINE_ATTRIBUTE(FoliageUpVector, FVector3f, AttributeNames::FoliageUPVectorStr, GroupNames::FoliageGroupStr);
+PV_DEFINE_ATTRIBUTE(FoliageNormalVector, FVector3f, AttributeNames::FoliageNormalVectorStr, GroupNames::FoliageGroupStr);
+PV_DEFINE_ATTRIBUTE(FoliageScale, float, AttributeNames::FoliageScaleStr, GroupNames::FoliageGroupStr);
+PV_DEFINE_ATTRIBUTE(FoliageLengthFromRoot, float, AttributeNames::FoliageLengthFromRootStr, GroupNames::FoliageGroupStr);
+PV_DEFINE_ATTRIBUTE(FoliageParentBoneID, int32, AttributeNames::FoliageParentBoneIDStr, GroupNames::FoliageGroupStr);
+
+// FoliageNames Attributes
+PV_DEFINE_ATTRIBUTE(FoliageName, FString, AttributeNames::FoliageNameStr, GroupNames::FoliageNamesGroupStr);
+
+// Details Attributes
+PV_DEFINE_ATTRIBUTE(DetailLeafPhyllotaxy, TArray<float>, AttributeNames::LeafPhyllotaxyStr, GroupNames::DetailsGroupStr);
+PV_DEFINE_ATTRIBUTE(DetailFoliagePath, FString, AttributeNames::FoliagePathStr, GroupNames::DetailsGroupStr);
+PV_DEFINE_ATTRIBUTE(DetailGuid, FGuid, AttributeNames::GuidStr, GroupNames::DetailsGroupStr);
+PV_DEFINE_ATTRIBUTE(DetailLeafGrowth, TArray<float>, AttributeNames::LeafGrowthStr, GroupNames::DetailsGroupStr);
+PV_DEFINE_ATTRIBUTE(DetailAbscissionSenescense, TArray<float>, AttributeNames::AbscissionSenescenseStr, GroupNames::DetailsGroupStr);
+PV_DEFINE_ATTRIBUTE(DetailLateralElongation, TArray<float>, AttributeNames::LateralElongationStr, GroupNames::DetailsGroupStr);
+PV_DEFINE_ATTRIBUTE(DetailTrunkURange, TArray<FVector2f>, AttributeNames::TrunkURangeStr, GroupNames::DetailsGroupStr);
+PV_DEFINE_ATTRIBUTE(DetailTrunkMaterialPath, FString, AttributeNames::TrunkMaterialPathStr, GroupNames::DetailsGroupStr);
+
+// PlantProfiles Attributes
+PV_DEFINE_ATTRIBUTE(ProfilePoints, TArray<float>, AttributeNames::ProfilePointsStr, GroupNames::PlantProfilesGroupStr);
+
+// Bone Attributes
+PV_DEFINE_ATTRIBUTE(BoneName, FString, AttributeNames::BoneNameStr, GroupNames::BonesGroupStr);
+PV_DEFINE_ATTRIBUTE(BoneParentIndex, int32, AttributeNames::BoneParentIndexStr, GroupNames::BonesGroupStr);
+PV_DEFINE_ATTRIBUTE(BonePose, FTransform, AttributeNames::BonePoseStr, GroupNames::BonesGroupStr);
+PV_DEFINE_ATTRIBUTE(BonePointIndex, int32, AttributeNames::BonePointIndexStr, GroupNames::BonesGroupStr);
+PV_DEFINE_ATTRIBUTE(BoneAbsolutePosition, FVector3f, AttributeNames::BoneAbsolutePositionStr, GroupNames::BonesGroupStr);
+PV_DEFINE_ATTRIBUTE(BoneBranchIndex, int32, AttributeNames::BoneBranchIndexStr, GroupNames::BonesGroupStr);
+PV_DEFINE_ATTRIBUTE(BoneId, int32, AttributeNames::BoneIdStr, GroupNames::BonesGroupStr);
+
+// VerticesGroup
+PV_DEFINE_ATTRIBUTE(VerticesPointIds, int32, AttributeNames::VertexPointIdsStr, GroupNames::VerticesGroupStr);
+
+};
\ No newline at end of file
diff --git a/Engine/Plugins/Experimental/ProceduralVegetationEditor/Source/ProceduralVegetation/Private/Utils/PVDynamicMeshVertexAttribute.h b/Engine/Plugins/Experimental/ProceduralVegetationEditor/Source/ProceduralVegetation/Private/Utils/PVDynamicMeshVertexAttribute.h
new file mode 100644
--- /dev/null
+++ b/Engine/Plugins/Experimental/ProceduralVegetationEditor/Source/ProceduralVegetation/Private/Utils/PVDynamicMeshVertexAttribute.h
@@ -0,0 +1,203 @@
+// Copyright Epic Games, Inc. All Rights Reserved.
+#pragma once
+
+#include "DynamicMesh/DynamicMesh3.h"
+#include "DynamicMesh/DynamicVertexAttribute.h"
+#include "DynamicMesh/DynamicMeshAttributeSet.h"
+
+#include "GeometryCollection/GeometryCollection.h" // Needed for CopyAttributeToCollection
+
+namespace PV
+{
+template<typename AttribValueType>
+class TDynamicMeshVertexAttributeExt 
+	: public UE::Geometry::TDynamicMeshVertexAttribute<AttribValueType, 1>
+{
+public:
+	using Super = UE::Geometry::TDynamicMeshVertexAttribute<AttribValueType, 1>;
+
+	TDynamicMeshVertexAttributeExt() : Super() {}
+	TDynamicMeshVertexAttributeExt(UE::Geometry::FDynamicMesh3* ParentIn, bool bAutoInit = true) : Super(ParentIn, bAutoInit) {}
+
+	inline AttribValueType GetValue(int32 VertexID) const
+	{
+		AttribValueType Value = (AttribValueType)0;
+		Super::GetValue(VertexID, &Value);
+		return Value;
+	}
+
+	inline void SetValue(int32 VertexID, const AttribValueType& NewValue)
+	{
+		Super::SetValue(VertexID, &NewValue);
+	}
+
+	template <typename Predicate, typename Iteratable>
+	TArray<int32> FindAllByPredicate(Predicate Pred, const Iteratable& InIteratable) const
+	{
+		TArray<int32> Result;
+
+		if (Super::Parent)
+		{
+			Result.Reserve(Super::Parent->VertexCount());
+
+			for (int32 VertexID : InIteratable)
+			{
+				check(Super::Parent->IsVertex(VertexID));
+
+				const AttribValueType Value = GetValue(VertexID);
+				if (Pred(Value))
+				{
+					Result.Add(VertexID);
+				}
+			}
+		}
+
+		return Result;
+	}
+
+	TArray<int32> FindAllNonZero() const
+	{
+		if (!Super::Parent)
+		{
+			return TArray<int32>();
+		}
+
+		return FindAllByPredicate(
+			[](const auto& Value) { return Value != (AttribValueType)0; },
+			Super::Parent->VertexIndicesItr()
+		);
+	}
+
+	TArray<int32> FindAllZero() const
+	{
+		if (!Super::Parent)
+		{
+			return TArray<int32>();
+		}
+
+		return FindAllByPredicate(
+			[](const auto& Value) { return Value == (AttribValueType)0; },
+			Super::Parent->VertexIndicesItr()
+		);
+	}
+
+	TArray<int32> FindAllNonZero(const TArray<int32>& InVertexIDsToSearch) const
+	{
+		return FindAllByPredicate(
+			[](const auto& Value) { return Value != (AttribValueType)0; },
+			InVertexIDsToSearch
+		);
+	}
+
+	TArray<int32> FindAllZero(const TArray<int32>& InVertexIDsToSearch) const
+	{
+		return FindAllByPredicate(
+			[](const auto& Value) { return Value == (AttribValueType)0; },
+			InVertexIDsToSearch
+		);
+	}
+};
+
+template<typename T>
+class TDynamicMeshVertexAttributeDefinition
+{
+public:
+	using FAttributeType = TDynamicMeshVertexAttributeExt<T>;
+
+	const FName AttributeName;
+
+	TDynamicMeshVertexAttributeDefinition(FName InAttributeName)
+		: AttributeName(InAttributeName)
+	{
+	}
+
+	const FAttributeType* GetAttribute(const UE::Geometry::FDynamicMesh3& DynamicMesh) const
+	{
+		if (DynamicMesh.HasAttributes() && DynamicMesh.Attributes()->HasAttachedAttribute(AttributeName))
+		{
+			return static_cast<const FAttributeType*>(DynamicMesh.Attributes()->GetAttachedAttribute(AttributeName));
+		}
+		return nullptr;
+	}
+
+	FAttributeType* GetAttribute(UE::Geometry::FDynamicMesh3& DynamicMesh) const
+	{
+		if (DynamicMesh.HasAttributes() && DynamicMesh.Attributes()->HasAttachedAttribute(AttributeName))
+		{
+			return static_cast<FAttributeType*>(DynamicMesh.Attributes()->GetAttachedAttribute(AttributeName));
+		}
+		return nullptr;
+	}
+
+	const FAttributeType& GetAttributeChecked(const UE::Geometry::FDynamicMesh3& DynamicMesh) const
+	{
+		const FAttributeType* Attribute = GetAttribute(DynamicMesh);
+		check(Attribute != nullptr);
+		return *Attribute;
+	}
+
+	FAttributeType& GetAttributeChecked(UE::Geometry::FDynamicMesh3& DynamicMesh) const
+	{
+		FAttributeType* Attribute = GetAttribute(DynamicMesh);
+		check(Attribute != nullptr);
+		return *Attribute;
+	}
+
+	FAttributeType& GetOrAttachAttribute(UE::Geometry::FDynamicMesh3& DynamicMesh) const
+	{
+		FAttributeType* Attribute = GetAttribute(DynamicMesh);
+		if (Attribute == nullptr)
+		{
+			AttachAttribute(DynamicMesh);
+		}
+
+		return GetAttributeChecked(DynamicMesh);
+	}
+
+	void AttachAttribute(UE::Geometry::FDynamicMesh3& DynamicMesh) const
+	{
+		if (!DynamicMesh.HasAttributes())
+		{
+			DynamicMesh.EnableAttributes();
+		}
+
+		if (!DynamicMesh.Attributes()->HasAttachedAttribute(AttributeName))
+		{
+			auto* Attribute = new FAttributeType(&DynamicMesh);
+			Attribute->Initialize();
+			DynamicMesh.Attributes()->AttachAttribute(AttributeName, Attribute);
+		}
+	}
+
+	void CopyAttributeToCollection(const UE::Geometry::FDynamicMesh3& DynamicMesh, FManagedArrayCollection& OutCollection) const
+	{
+		const FAttributeType* AttachedAttribute = GetAttribute(DynamicMesh);
+		if (!AttachedAttribute)
+		{
+			return;
+		}
+
+		const static FName VertexAttributesGroupName = TEXT("VertexAttributes");
+
+		if (!OutCollection.HasGroup(VertexAttributesGroupName))
+		{
+			OutCollection.AddGroup(VertexAttributesGroupName);
+			OutCollection.AddElements(DynamicMesh.MaxVertexID(), VertexAttributesGroupName);
+
+			auto& DebugVertexPositionAttribute = OutCollection.AddAttribute<FVector3f>(FGeometryCollection::VertexPositionAttribute, VertexAttributesGroupName);
+			for (int32 VertexIndex : DynamicMesh.VertexIndicesItr())
+			{
+				DebugVertexPositionAttribute[VertexIndex] = FVector3f(DynamicMesh.GetVertex(VertexIndex));
+			}
+		}
+
+		TManagedArray<T>& CollectionAttribute = OutCollection.AddAttribute<T>(AttributeName, VertexAttributesGroupName);
+		check(CollectionAttribute.Num() >= DynamicMesh.MaxVertexID());
+
+		for (int32 VertexIndex : DynamicMesh.VertexIndicesItr())
+		{
+			CollectionAttribute[VertexIndex] = AttachedAttribute->GetValue(VertexIndex);
+		}
+	}
+};
+}
\ No newline at end of file
diff --git a/Engine/Plugins/Experimental/ProceduralVegetationEditor/Source/ProceduralVegetation/ProceduralVegetation.Build.cs b/Engine/Plugins/Experimental/ProceduralVegetationEditor/Source/ProceduralVegetation/ProceduralVegetation.Build.cs
--- a/Engine/Plugins/Experimental/ProceduralVegetationEditor/Source/ProceduralVegetation/ProceduralVegetation.Build.cs
+++ b/Engine/Plugins/Experimental/ProceduralVegetationEditor/Source/ProceduralVegetation/ProceduralVegetation.Build.cs
@@ -26,12 +26,14 @@ namespace UnrealBuildTool.Rules
 					"Json",
 					"JsonUtilities",
 					"MeshDescription",
+					"MeshConversionEngineTypes",
 					"ModelingComponents",
 					"PCG",
 					"RenderCore",
 					"StaticMeshDescription",
 					"RHI",
-					"Projects"
+					"Projects",
+					"ImageCore",
 				}
 			);
 		}
diff --git a/Engine/Plugins/Experimental/ProceduralVegetationEditor/Source/ProceduralVegetation/Public/Facades/PVAttributes.h b/Engine/Plugins/Experimental/ProceduralVegetationEditor/Source/ProceduralVegetation/Public/Facades/PVAttributes.h
deleted file mode 100644
diff --git a/Engine/Plugins/Experimental/ProceduralVegetationEditor/Source/ProceduralVegetation/Public/Helpers/PVAttributesHelper.h b/Engine/Plugins/Experimental/ProceduralVegetationEditor/Source/ProceduralVegetation/Public/Helpers/PVAttributesHelper.h
--- a/Engine/Plugins/Experimental/ProceduralVegetationEditor/Source/ProceduralVegetation/Public/Helpers/PVAttributesHelper.h
+++ b/Engine/Plugins/Experimental/ProceduralVegetationEditor/Source/ProceduralVegetation/Public/Helpers/PVAttributesHelper.h
@@ -1,7 +1,7 @@
 // Copyright Epic Games, Inc. All Rights Reserved.
 #pragma once
 
-#include "Facades/PVAttributes.h"
+#include "Utils/PVAttributes.h"
 
 namespace PV::AttributesHelper
 {
@@ -11,7 +11,9 @@ namespace PV::AttributesHelper
 		int32 BranchIndex
 	)
 	{
-		return BranchNumberAttribute.Find(BranchParentNumberAttribute[BranchIndex]);
+		return BranchParentNumberAttribute.IsValidIndex(BranchIndex) 
+			? BranchNumberAttribute.Find(BranchParentNumberAttribute[BranchIndex])
+			: INDEX_NONE;
 	}
 
 	inline TArray<int32> GetPlantBranchIndices(
@@ -577,15 +579,18 @@ namespace PV::AttributesHelper
 		int32 BranchIndex
 	)
 	{
-		if (GetBranchParentIndex(BranchParentNumberAttribute, BranchNumberAttribute, BranchIndex) == INDEX_NONE)
+		if (BranchPointsAttribute.IsValidIndex(BranchIndex))
 		{
-			return TArrayView<const int32>(BranchPointsAttribute[BranchIndex]);
-		}
+			if (GetBranchParentIndex(BranchParentNumberAttribute, BranchNumberAttribute, BranchIndex) == INDEX_NONE)
+			{
+				return TArrayView<const int32>(BranchPointsAttribute[BranchIndex]);
+			}
 
-		const TArray<int32>& Points = BranchPointsAttribute[BranchIndex];
-		if (Points.Num() > 1)
-		{
-			return TArrayView<const int32>(&Points[1], Points.Num() - 1);
+			const TArray<int32>& Points = BranchPointsAttribute[BranchIndex];
+			if (Points.Num() > 1)
+			{
+				return TArrayView<const int32>(&Points[1], Points.Num() - 1);
+			}
 		}
 
 		return TArrayView<const int32>();
diff --git a/Engine/Plugins/Experimental/ProceduralVegetationEditor/Source/ProceduralVegetation/Public/IPVEditorInterface.h b/Engine/Plugins/Experimental/ProceduralVegetationEditor/Source/ProceduralVegetation/Public/IPVEditorInterface.h
deleted file mode 100644
diff --git a/Engine/Plugins/Experimental/ProceduralVegetationEditor/Source/ProceduralVegetation/Public/Nodes/PVOutputSettings.h b/Engine/Plugins/Experimental/ProceduralVegetationEditor/Source/ProceduralVegetation/Public/Nodes/PVOutputSettings.h
--- a/Engine/Plugins/Experimental/ProceduralVegetationEditor/Source/ProceduralVegetation/Public/Nodes/PVOutputSettings.h
+++ b/Engine/Plugins/Experimental/ProceduralVegetationEditor/Source/ProceduralVegetation/Public/Nodes/PVOutputSettings.h
@@ -2,7 +2,7 @@
 #pragma once
 
 #include "CoreMinimal.h"
-#include "PVExportParams.h"
+#include "Params/PVExportParams.h"
 #include "Nodes/PVBaseSettings.h"
 #include "PVOutputSettings.generated.h"
 
diff --git a/Engine/Plugins/Experimental/ProceduralVegetationEditor/Source/ProceduralVegetation/Public/PVExportParams.h b/Engine/Plugins/Experimental/ProceduralVegetationEditor/Source/ProceduralVegetation/Public/PVExportParams.h
deleted file mode 100644
diff --git a/Engine/Plugins/Experimental/ProceduralVegetationEditor/Source/ProceduralVegetation/Public/PVFloatRamp.h b/Engine/Plugins/Experimental/ProceduralVegetationEditor/Source/ProceduralVegetation/Public/PVFloatRamp.h
deleted file mode 100644
diff --git a/Engine/Plugins/Experimental/ProceduralVegetationEditor/Source/ProceduralVegetation/Public/PVImportStaticMeshParams.h b/Engine/Plugins/Experimental/ProceduralVegetationEditor/Source/ProceduralVegetation/Public/PVImportStaticMeshParams.h
deleted file mode 100644
diff --git a/Engine/Plugins/Experimental/ProceduralVegetationEditor/Source/ProceduralVegetation/Public/PVImportTexture2DParams.h b/Engine/Plugins/Experimental/ProceduralVegetationEditor/Source/ProceduralVegetation/Public/PVImportTexture2DParams.h
deleted file mode 100644
diff --git a/Engine/Plugins/Experimental/ProceduralVegetationEditor/Source/ProceduralVegetation/Public/Params/PVExportParams.h b/Engine/Plugins/Experimental/ProceduralVegetationEditor/Source/ProceduralVegetation/Public/Params/PVExportParams.h
new file mode 100644
--- /dev/null
+++ b/Engine/Plugins/Experimental/ProceduralVegetationEditor/Source/ProceduralVegetation/Public/Params/PVExportParams.h
@@ -0,0 +1,130 @@
+// Copyright Epic Games, Inc. All Rights Reserved.
+#pragma once
+
+#include "Engine/EngineTypes.h"
+#include "PVWindSettings.h"
+
+#include "Misc/Paths.h"
+#include "PVExportParams.generated.h"
+
+UENUM()
+enum class EPVExportMeshType: uint8
+{
+	StaticMesh,
+	SkeletalMesh,
+};
+
+UENUM()
+enum class EPVAssetReplacePolicy : uint8
+{
+	/** Create new assets with numbered suffixes */
+	Append,
+
+	/** Replaces existing asset with new asset */
+	Replace,
+
+	/** Ignores the new asset and keeps the existing asset */
+	Ignore,
+};
+
+UENUM()
+enum class EPVCollisionGeneration : uint8
+{
+	/** Do not create collision for mesh */
+	None,
+
+	/** Creates collision for trunk only */
+	TrunkOnly,
+
+	/** Creates collision for all branch generations */
+	AllGenerations,
+};
+
+USTRUCT()
+struct PROCEDURALVEGETATION_API FPVExportParams
+{
+	GENERATED_BODY()
+
+	UPROPERTY()
+	bool bShouldExport = true;
+
+	UPROPERTY()
+	bool bShouldExportFoliage = true;
+
+	UPROPERTY(EditAnywhere, Category = "Asset", Meta = (ContentDir, DisplayName = "Content Browser Folder", Tooltip="Path to the folder where exported mesh will be saved\n\nPath to the folder where exported mesh will be saved."))
+	FDirectoryPath ContentBrowserFolder;
+
+	UPROPERTY(EditAnywhere, Category = "Asset", Meta = (DisplayName = "Mesh Name", Tooltip="Name of the exported mesh\n\nName of the exported mesh."))
+	FName MeshName;
+
+	UPROPERTY(EditAnywhere, Category = "Asset", Meta = (DisplayName = "Asset Replace Policy", Tooltip="Export behaviour if asset already exists."))
+	EPVAssetReplacePolicy ReplacePolicy = EPVAssetReplacePolicy::Replace;
+
+	UPROPERTY(EditAnywhere, Category = "Mesh", Meta = (DisplayName = "Export Mesh Type", Tooltip="Static Mesh or Skeletal Mesh."))
+	EPVExportMeshType ExportMeshType = EPVExportMeshType::SkeletalMesh;
+
+	UPROPERTY(EditAnywhere, Category = "Mesh | Nanite", Meta = (DisplayName = "Create Nanite Foliage", Tooltip="Add Nanite foliage support for the exported mesh."))
+	bool bCreateNaniteFoliage = true;
+	
+	UPROPERTY(EditAnywhere, Category = "Mesh | Nanite", Meta = (DisplayName = "Nanite Shape Preservation Method", Tooltip="Shape preservation setting to be used in nanite settings."))
+	ENaniteShapePreservation NaniteShapePreservation = ENaniteShapePreservation::Voxelize;
+
+	UPROPERTY(EditAnywhere, Category = "Mesh | Collision", Meta = (EditCondition = "ExportMeshType == EPVExportMeshType::StaticMesh", EditConditionHides, Tooltip="Creates the collision for static mesh."))
+	bool bCollision = false;
+
+	UPROPERTY(EditAnywhere, Category = "Mesh | Collision", Meta = (EditCondition = "ExportMeshType == EPVExportMeshType::SkeletalMesh", EditConditionHides, Tooltip="Creates the physics asset for skeletal mesh."))
+	EPVCollisionGeneration CollisionGeneration = EPVCollisionGeneration::None;
+	
+	UPROPERTY(EditAnywhere, Category = "Mesh | Dynamic Wind", Meta = (EditCondition = "ExportMeshType == EPVExportMeshType::SkeletalMesh", EditConditionHides, Tooltip="Wind Settings contains the wind simulation group data, Default presets are available for Saplings and Trees."))
+	TObjectPtr<UPVWindSettings> WindSettings;
+
+	void Initialize(const FString& InAssetPath, const FString& InName, const FString& InDefaultWindSettingsPath)
+	{
+		ContentBrowserFolder.Path = FPaths::GetPath(InAssetPath);
+		MeshName = FName(InName);
+		ReplacePolicy = EPVAssetReplacePolicy::Replace;
+
+		WindSettings = LoadObject<UPVWindSettings>(nullptr, InDefaultWindSettingsPath);
+	}
+
+	FString GetOutputObjectPath() const
+	{
+		const FString AssetName = MeshName.ToString();
+		const FString OutputPath = ContentBrowserFolder.Path / AssetName + '.' + AssetName;
+		return OutputPath;
+	}
+
+	FString GetOutputMeshPackagePath() const
+	{
+		return FPaths::Combine(ContentBrowserFolder.Path, MeshName.ToString());
+	}
+
+	FString GetOutputSkeletonName() const
+	{
+		return FString::Printf(TEXT("%s_Skeleton"), *MeshName.ToString());
+	}
+
+	FString GetOutputSkeletonPackagePath() const
+	{
+		return FPaths::Combine(ContentBrowserFolder.Path, GetOutputSkeletonName());
+	}
+
+	FString GetOutputPhysicsAssetName() const
+	{
+		return FString::Printf(TEXT("%s_Physics"), *MeshName.ToString());
+	}
+
+	FString GetOutputPhysicsAssetPackagePath() const
+	{
+		return FPaths::Combine(ContentBrowserFolder.Path, GetOutputPhysicsAssetName());
+	}
+	
+	bool IsCollisionEnable() const
+	{
+		return CollisionGeneration != EPVCollisionGeneration::None;
+	}
+
+	UClass* GetMeshClass() const;
+
+	bool Validate(FString& OutError) const;
+};
\ No newline at end of file
diff --git a/Engine/Plugins/Experimental/ProceduralVegetationEditor/Source/ProceduralVegetation/Public/Utils/PVFloatRamp.h b/Engine/Plugins/Experimental/ProceduralVegetationEditor/Source/ProceduralVegetation/Public/Utils/PVFloatRamp.h
new file mode 100644
--- /dev/null
+++ b/Engine/Plugins/Experimental/ProceduralVegetationEditor/Source/ProceduralVegetation/Public/Utils/PVFloatRamp.h
@@ -0,0 +1,25 @@
+// Copyright Epic Games, Inc. All Rights Reserved.
+
+#pragma once
+
+#include "Curves/RichCurve.h"
+
+#include "PVFloatRamp.generated.h"
+
+USTRUCT(BlueprintType)
+struct PROCEDURALVEGETATION_API FPVFloatRamp
+{
+	GENERATED_BODY()
+
+	UPROPERTY()
+	FRichCurve EditorCurveData;
+
+	FRichCurve* GetRichCurve()
+	{
+		return &EditorCurveData;
+	}
+	const FRichCurve* GetRichCurveConst() const
+	{
+		return &EditorCurveData;
+	}
+};
diff --git a/Engine/Plugins/Experimental/ProceduralVegetationEditor/Source/ProceduralVegetationEditor/Private/Customizations/PVFloatRampCustomization.cpp b/Engine/Plugins/Experimental/ProceduralVegetationEditor/Source/ProceduralVegetationEditor/Private/Customizations/PVFloatRampCustomization.cpp
--- a/Engine/Plugins/Experimental/ProceduralVegetationEditor/Source/ProceduralVegetationEditor/Private/Customizations/PVFloatRampCustomization.cpp
+++ b/Engine/Plugins/Experimental/ProceduralVegetationEditor/Source/ProceduralVegetationEditor/Private/Customizations/PVFloatRampCustomization.cpp
@@ -6,7 +6,7 @@
 #include "DetailWidgetRow.h"
 #include "Editor.h"
 #include "IDetailChildrenBuilder.h"
-#include "PVFloatRamp.h"
+#include "Utils/PVFloatRamp.h"
 #include "ScopedTransaction.h"
 #include "SCurveEditor.h"
 
diff --git a/Engine/Plugins/Experimental/ProceduralVegetationEditor/Source/ProceduralVegetationEditor/Private/Helpers/PVExportHelper.cpp b/Engine/Plugins/Experimental/ProceduralVegetationEditor/Source/ProceduralVegetationEditor/Private/Helpers/PVExportHelper.cpp
--- a/Engine/Plugins/Experimental/ProceduralVegetationEditor/Source/ProceduralVegetationEditor/Private/Helpers/PVExportHelper.cpp
+++ b/Engine/Plugins/Experimental/ProceduralVegetationEditor/Source/ProceduralVegetationEditor/Private/Helpers/PVExportHelper.cpp
@@ -14,7 +14,7 @@
 #include "ProceduralVegetation.h"
 #include "ProceduralVegetationEditorModule.h"
 #include "ProceduralVegetationLink.h"
-#include "PVExportParams.h"
+#include "Params/PVExportParams.h"
 #include "PVWindSettings.h"
 #include "RenderUtils.h"
 #include "SkinnedAssetCompiler.h"
diff --git a/Engine/Plugins/Experimental/ProceduralVegetationEditor/Source/ProceduralVegetationEditor/Private/Importer/PVDynamicMeshVertexAttribute.h b/Engine/Plugins/Experimental/ProceduralVegetationEditor/Source/ProceduralVegetationEditor/Private/Importer/PVDynamicMeshVertexAttribute.h
deleted file mode 100644
diff --git a/Engine/Plugins/Experimental/ProceduralVegetationEditor/Source/ProceduralVegetationEditor/Private/Importer/PVImportHelpers.cpp b/Engine/Plugins/Experimental/ProceduralVegetationEditor/Source/ProceduralVegetationEditor/Private/Importer/PVImportHelpers.cpp
deleted file mode 100644
diff --git a/Engine/Plugins/Experimental/ProceduralVegetationEditor/Source/ProceduralVegetationEditor/Private/Importer/PVImportHelpers.h b/Engine/Plugins/Experimental/ProceduralVegetationEditor/Source/ProceduralVegetationEditor/Private/Importer/PVImportHelpers.h
deleted file mode 100644
diff --git a/Engine/Plugins/Experimental/ProceduralVegetationEditor/Source/ProceduralVegetationEditor/Private/Importer/PVImporter_StaticMesh.cpp b/Engine/Plugins/Experimental/ProceduralVegetationEditor/Source/ProceduralVegetationEditor/Private/Importer/PVImporter_StaticMesh.cpp
deleted file mode 100644
diff --git a/Engine/Plugins/Experimental/ProceduralVegetationEditor/Source/ProceduralVegetationEditor/Private/Importer/PVImporter_StaticMesh.h b/Engine/Plugins/Experimental/ProceduralVegetationEditor/Source/ProceduralVegetationEditor/Private/Importer/PVImporter_StaticMesh.h
deleted file mode 100644
diff --git a/Engine/Plugins/Experimental/ProceduralVegetationEditor/Source/ProceduralVegetationEditor/Private/Importer/PVImporter_Texture2D.cpp b/Engine/Plugins/Experimental/ProceduralVegetationEditor/Source/ProceduralVegetationEditor/Private/Importer/PVImporter_Texture2D.cpp
deleted file mode 100644
diff --git a/Engine/Plugins/Experimental/ProceduralVegetationEditor/Source/ProceduralVegetationEditor/Private/Importer/PVImporter_Texture2D.h b/Engine/Plugins/Experimental/ProceduralVegetationEditor/Source/ProceduralVegetationEditor/Private/Importer/PVImporter_Texture2D.h
deleted file mode 100644
diff --git a/Engine/Plugins/Experimental/ProceduralVegetationEditor/Source/ProceduralVegetationEditor/Private/PVEditorInterfaceImpl.cpp b/Engine/Plugins/Experimental/ProceduralVegetationEditor/Source/ProceduralVegetationEditor/Private/PVEditorInterfaceImpl.cpp
deleted file mode 100644
diff --git a/Engine/Plugins/Experimental/ProceduralVegetationEditor/Source/ProceduralVegetationEditor/Private/PVEditorInterfaceImpl.h b/Engine/Plugins/Experimental/ProceduralVegetationEditor/Source/ProceduralVegetationEditor/Private/PVEditorInterfaceImpl.h
deleted file mode 100644
diff --git a/Engine/Plugins/Experimental/ProceduralVegetationEditor/Source/ProceduralVegetationEditor/Private/ProceduralVegetationEditorModule.cpp b/Engine/Plugins/Experimental/ProceduralVegetationEditor/Source/ProceduralVegetationEditor/Private/ProceduralVegetationEditorModule.cpp
--- a/Engine/Plugins/Experimental/ProceduralVegetationEditor/Source/ProceduralVegetationEditor/Private/ProceduralVegetationEditorModule.cpp
+++ b/Engine/Plugins/Experimental/ProceduralVegetationEditor/Source/ProceduralVegetationEditor/Private/ProceduralVegetationEditorModule.cpp
@@ -8,8 +8,7 @@
 #include "PropertyEditorModule.h"
 #include "PVEditorCommands.h"
 #include "PVEditorSettings.h"
-#include "PVEditorInterfaceImpl.h"
-#include "PVFloatRamp.h"
+#include "Utils/PVFloatRamp.h"
 #include "Selection.h"
 #include "ToolMenuContext.h"
 #include "ToolMenuDelegates.h"
@@ -29,7 +28,7 @@
 #include "DataTypes/PVFoliageMeshData.h"
 #include "DataTypes/PVGrowthData.h"
 #include "DataTypes/PVMeshData.h"
-#include "PVExportParams.h"
+#include "Params/PVExportParams.h"
 
 #include "Engine/SkeletalMesh.h"
 #include "Engine/StaticMesh.h"
@@ -77,9 +76,6 @@ void FProceduralVegetationEditorModule::StartupModule()
 
 	FMessageLogModule& MessageLogModule = FModuleManager::LoadModuleChecked<FMessageLogModule>("MessageLog");
 	MessageLogModule.RegisterLogListing(PVEditor::MessageLogName, LOCTEXT("PVEditorLogLabel", "Procedural Vegetation Editor"));
-	
-	check(PV::GPVEditorInterface == nullptr);
-	PV::GPVEditorInterface = new PV::FPVEditorInterfaceImpl();
 }
 
 void FProceduralVegetationEditorModule::ShutdownModule()
@@ -102,10 +98,6 @@ void FProceduralVegetationEditorModule::ShutdownModule()
 	{
 		MessageLogModule->UnregisterLogListing(PVEditor::MessageLogName);
 	}
-
-	check(PV::GPVEditorInterface != nullptr);
-	delete PV::GPVEditorInterface;
-	PV::GPVEditorInterface = nullptr;
 }
 
 void FProceduralVegetationEditorModule::RegisterMenus()
diff --git a/Engine/Plugins/Experimental/ProceduralVegetationEditor/Source/ProceduralVegetationEditor/Private/Tests/Importer/PVImportTests.cpp b/Engine/Plugins/Experimental/ProceduralVegetationEditor/Source/ProceduralVegetationEditor/Private/Tests/Importer/PVImportTests.cpp
deleted file mode 100644
diff --git a/Engine/Plugins/Experimental/ProceduralVegetationEditor/Source/ProceduralVegetationEditor/ProceduralVegetationEditor.Build.cs b/Engine/Plugins/Experimental/ProceduralVegetationEditor/Source/ProceduralVegetationEditor/ProceduralVegetationEditor.Build.cs
--- a/Engine/Plugins/Experimental/ProceduralVegetationEditor/Source/ProceduralVegetationEditor/ProceduralVegetationEditor.Build.cs
+++ b/Engine/Plugins/Experimental/ProceduralVegetationEditor/Source/ProceduralVegetationEditor/ProceduralVegetationEditor.Build.cs
@@ -22,17 +22,14 @@ public class ProceduralVegetationEditor : ModuleRules
 				"CurveEditor",
 				"DataflowEditor",
 				"Engine",
-				"GeometryCollectionEngine",
 				"GeometryCore",
 				"GeometryFramework",
 				"GeometryScriptingCore",
 				"InputCore",
-				"ImageCore",
 				"Json",
 				"JsonUtilities",
 				"LevelEditor",
 				"MeshConversionEngineTypes",
-				"MeshDescription",
 				"MeshConversion",
 				"ModelingComponents",
 				"PackagesDialog",
