Fix accidental inversion of PVE 2D importer tip angle

- Fix inversion of tip angle accidentally submitted in CL50864534
- Adjust default TipAngleThreshold to something more reasonable

#Synced-CL 50901855

diff --git a/Engine/Plugins/Experimental/ProceduralVegetationEditor/Source/ProceduralVegetation/Private/Implementations/PVImporter_Texture2D.cpp b/Engine/Plugins/Experimental/ProceduralVegetationEditor/Source/ProceduralVegetation/Private/Implementations/PVImporter_Texture2D.cpp
--- a/Engine/Plugins/Experimental/ProceduralVegetationEditor/Source/ProceduralVegetation/Private/Implementations/PVImporter_Texture2D.cpp
+++ b/Engine/Plugins/Experimental/ProceduralVegetationEditor/Source/ProceduralVegetation/Private/Implementations/PVImporter_Texture2D.cpp
@@ -695,7 +695,7 @@ namespace Internal
 
 			const double ClampedDot = FMath::Clamp(FVector3d::DotProduct(Normal1, Normal2), -1.0, 1.0); // Clamp to ensure range does not exceed [-1, 1] due to floating point imprecision
 			const double AngleRadians = FMath::Acos(ClampedDot);
-			return FMath::RadiansToDegrees(AngleRadians);
+			return 180.f - FMath::RadiansToDegrees(AngleRadians);
 		};
 
 		UE::Geometry::FMeshConnectedComponents ConnectedComponents(&DynamicMesh);
diff --git a/Engine/Plugins/Experimental/ProceduralVegetationEditor/Source/ProceduralVegetation/Private/Params/PVImportTexture2DParams.h b/Engine/Plugins/Experimental/ProceduralVegetationEditor/Source/ProceduralVegetation/Private/Params/PVImportTexture2DParams.h
--- a/Engine/Plugins/Experimental/ProceduralVegetationEditor/Source/ProceduralVegetation/Private/Params/PVImportTexture2DParams.h
+++ b/Engine/Plugins/Experimental/ProceduralVegetationEditor/Source/ProceduralVegetation/Private/Params/PVImportTexture2DParams.h
@@ -68,7 +68,7 @@ struct FPVImportTexture2DParams
 	TArray<FPVImportTexture2DPlantSettings> PlantSettings;
 
 	UPROPERTY(EditAnywhere, Category = "Plant Geneneration|Tip detection", meta = (ClampMin=0.001, ClampMax=180, UIMin=0.01, UIMax=180))
-	float TipAngleThresholdInDegrees = 45.0;
+	float TipAngleThresholdInDegrees = 100.0;
 
 	UPROPERTY(EditAnywhere, Category = "Plant Geneneration|Tip detection", meta = (ClampMin=0, ClampMax=0.01, UIMin=0, UIMax=0.01))
 	float MaxTipAngleSearchDist = 0.005;
