[PVE] Improvements and bugfixes for grafter node

- If either the source skeleton or the source graft have any foliage data, it is removed and a warning is show on the node
- Grafts added to the tip of a branch are merged with the parent branch
- A point is added to the skeleton at the point of each graft and the first point of each graft is removed before attachment to compensate for the point added
- Fixed scaling issue with the grafts
- Recomputes branch children since legacy data has inconsistent data in the children attribute, to be removed in the future once support for legacy data is dropped
- Refactored some code for clarity and readability

#Synced-CL 50929882

diff --git a/Engine/Plugins/Experimental/ProceduralVegetationEditor/Source/ProceduralVegetation/Private/Facades/PVBranchFacade.cpp b/Engine/Plugins/Experimental/ProceduralVegetationEditor/Source/ProceduralVegetation/Private/Facades/PVBranchFacade.cpp
--- a/Engine/Plugins/Experimental/ProceduralVegetationEditor/Source/ProceduralVegetation/Private/Facades/PVBranchFacade.cpp
+++ b/Engine/Plugins/Experimental/ProceduralVegetationEditor/Source/ProceduralVegetation/Private/Facades/PVBranchFacade.cpp
@@ -529,6 +529,16 @@ namespace PV::Facades
 		FILL_COLLECTION_ATTRIBUTE(PlantNumbers, InputArray);
 	}
 
+	bool FBranchFacade::IsTrunk(int32 Index) const
+	{
+		if (BranchParentNumbers.IsValid() && BranchParentNumbers.IsValidIndex(Index))
+		{
+			return BranchParentNumbers[Index] == 0;
+		}
+		
+		return false;
+	}
+
 	void FBranchFacade::GetParents(int32 Index, TArray<int32>& OutValue) const
 	{
 		GetValueFromAccessor<TArray<int32>>(Parents, OutValue, Index, TArray<int32>());
@@ -554,6 +564,16 @@ namespace PV::Facades
 		GetValueFromAccessor<int32>(BranchSourceBudNumber, OutValue, Index, -1);
 	}
 
+	const TManagedArray<int32>& FBranchFacade::GetBranchSourceBudNumbers() const
+	{
+		return BranchSourceBudNumber.Get();
+	}
+
+	void FBranchFacade::SetBranchSourceBudNumbersFromIndex(const TArray<int32>& InputArray, const int32 StartIndex)
+	{
+		FILL_COLLECTION_ATTRIBUTE_FROM_INDEX(BranchSourceBudNumber, InputArray, StartIndex);
+	}
+
 	void FBranchFacade::GetBranchHierarchyNumber(int32 Index, int32& OutValue) const
 	{
 		GetValueFromAccessor<int32>(BranchHierarchyNumber, OutValue, Index, -1);
diff --git a/Engine/Plugins/Experimental/ProceduralVegetationEditor/Source/ProceduralVegetation/Private/Facades/PVPointFacade.cpp b/Engine/Plugins/Experimental/ProceduralVegetationEditor/Source/ProceduralVegetation/Private/Facades/PVPointFacade.cpp
--- a/Engine/Plugins/Experimental/ProceduralVegetationEditor/Source/ProceduralVegetation/Private/Facades/PVPointFacade.cpp
+++ b/Engine/Plugins/Experimental/ProceduralVegetationEditor/Source/ProceduralVegetation/Private/Facades/PVPointFacade.cpp
@@ -7,7 +7,9 @@
 namespace PV::Facades
 {
 	FPointFacade::FPointFacade(FManagedArrayCollection& InCollection)
-		: Positions(InCollection, AttributeNames::PointPosition, GroupNames::PointGroup)
+		: Collection(&InCollection)
+		, BranchPoints(InCollection, AttributeNames::BranchPoints, GroupNames::BranchGroup)
+		, Positions(InCollection, AttributeNames::PointPosition, GroupNames::PointGroup)
 		, LengthFromRoot(InCollection, AttributeNames::LengthFromRoot, GroupNames::PointGroup)
 		, PointScaleGradient(InCollection, AttributeNames::PointScaleGradient, GroupNames::PointGroup)
 		, HullGradient(InCollection, AttributeNames::HullGradient, GroupNames::PointGroup)
@@ -34,7 +36,9 @@ namespace PV::Facades
 	}
 
 	FPointFacade::FPointFacade(const FManagedArrayCollection& InCollection)
-		: Positions(InCollection, AttributeNames::PointPosition, GroupNames::PointGroup)
+		: Collection(nullptr)
+		, BranchPoints(InCollection, AttributeNames::BranchPoints, GroupNames::BranchGroup)
+		, Positions(InCollection, AttributeNames::PointPosition, GroupNames::PointGroup)
 		, LengthFromRoot(InCollection, AttributeNames::LengthFromRoot, GroupNames::PointGroup)
 		, PointScaleGradient(InCollection, AttributeNames::PointScaleGradient, GroupNames::PointGroup)
 		, HullGradient(InCollection, AttributeNames::HullGradient, GroupNames::PointGroup)
@@ -96,6 +100,14 @@ namespace PV::Facades
 		return DefaultPosition;
 	}
 
+	void FPointFacade::SetPosition(int32 Index, const FVector3f& InPosition)
+	{
+		if (Positions.IsValid() && Positions.IsValidIndex(Index))
+		{
+			Positions.ModifyAt(Index, InPosition);
+		}
+	}
+
 	float FPointFacade::GetLengthFromRoot(int32 Index) const
 	{
 		if (LengthFromRoot.IsValid() && LengthFromRoot.IsValidIndex(Index))
@@ -106,6 +118,14 @@ namespace PV::Facades
 		return 0;
 	}
 
+	void FPointFacade::SetLengthFromRoot(int32 Index, float InLength)
+	{
+		if (LengthFromRoot.IsValid() && LengthFromRoot.IsValidIndex(Index))
+		{
+			LengthFromRoot.ModifyAt(Index, InLength);
+		}
+	}
+
 	float FPointFacade::GetLengthFromSeed(int32 Index) const
 	{
 		if (LengthFromSeed.IsValid() && LengthFromSeed.IsValidIndex(Index))
@@ -116,6 +136,14 @@ namespace PV::Facades
 		return 0;
 	}
 
+	void FPointFacade::SetLengthFromSeed(int32 Index, float InLength)
+	{
+		if (LengthFromSeed.IsValid() && LengthFromSeed.IsValidIndex(Index))
+		{
+			LengthFromSeed.ModifyAt(Index, InLength);
+		}
+	}
+
 	float FPointFacade::GetPointScaleGradient(int32 Index) const
 	{
 		if (PointScaleGradient.IsValid() && PointScaleGradient.IsValidIndex(Index))
@@ -203,11 +231,11 @@ namespace PV::Facades
 		return EmptyArray;
 	}
 
-	void FPointFacade::SetBudLightDetected(int32 Index, TArray<float> InBudLightDetected)
+	void FPointFacade::SetBudLightDetected(int32 Index, const TArray<float>& InBudLightDetected)
 	{
 		if (BudLightDetected.IsValid() && BudLightDetected.IsValidIndex(Index))
 		{
-			BudLightDetected.Modify()[Index] = MoveTemp(InBudLightDetected);
+			BudLightDetected.Modify()[Index] = InBudLightDetected;
 		}
 	}
 
@@ -222,11 +250,11 @@ namespace PV::Facades
 		return EmptyArray;
 	}
 	
-	void FPointFacade::SetBudDevelopment(int32 Index, TArray<int> InBudDevelopment)
+	void FPointFacade::SetBudDevelopment(int32 Index, const TArray<int>& InBudDevelopment)
 	{
 		if (BudDevelopment.IsValid() && BudDevelopment.IsValidIndex(Index))
 		{
-			BudDevelopment.Modify()[Index] = MoveTemp(InBudDevelopment);
+			BudDevelopment.Modify()[Index] = InBudDevelopment;
 		}
 	}
 
@@ -252,6 +280,31 @@ namespace PV::Facades
 		return INDEX_NONE;
 	}
 
+	int32 FPointFacade::GetMaxBudAge() const
+	{
+		int32 MaxBudAge = -1;
+		if (BudDevelopment.IsValid())
+		{
+			for (int32 i = 0; i < BudDevelopment.Num(); ++i)
+			{
+				MaxBudAge = FMath::Max(MaxBudAge, BudDevelopment[i][1]);
+			}
+		}
+
+		return MaxBudAge;
+	}
+
+	void FPointFacade::AddToBudAgeForAllPoints(const int32 Age)
+	{
+		if (BudDevelopment.IsValid())
+		{
+			for (int32 i = 0; i < BudDevelopment.Num(); ++i)
+			{
+				BudDevelopment.Modify()[i][1] += Age;
+			}
+		}
+	}
+
 	int32 FPointFacade::GetBudBranchAge(int32 Index) const
 	{
 		const TArray<int>& PointBudDevelopment = GetBudDevelopment(Index);
@@ -274,11 +327,11 @@ namespace PV::Facades
 		return EmptyArray;
 	}
 
-	void FPointFacade::SetBudDirections(int32 Index, TArray<FVector3f> InBudDirections)
+	void FPointFacade::SetBudDirections(int32 Index, const TArray<FVector3f>& InBudDirections)
 	{
 		if (BudDirections.IsValid() && BudDirections.IsValidIndex(Index))
 		{
-			BudDirections.Modify()[Index] = MoveTemp(InBudDirections);
+			BudDirections.Modify()[Index] = InBudDirections;
 		}
 	}
 
@@ -293,11 +346,11 @@ namespace PV::Facades
 		return EmptyArray;
 	}
 
-	void FPointFacade::SetBudHormoneLevels(int32 Index, TArray<float> InBudHormoneLevels)
+	void FPointFacade::SetBudHormoneLevels(int32 Index, const TArray<float>& InBudHormoneLevels)
 	{
 		if (BudHormoneLevels.IsValid() && BudHormoneLevels.IsValidIndex(Index))
 		{
-			BudHormoneLevels.Modify()[Index] = MoveTemp(InBudHormoneLevels);
+			BudHormoneLevels.Modify()[Index] = InBudHormoneLevels;
 		}
 	}
 
@@ -312,6 +365,14 @@ namespace PV::Facades
 		return EmptyArray;
 	}
 
+	void FPointFacade::SetBudLateralMeristem(int32 Index, const TArray<float>& InBudLateralMeristem)
+	{
+		if (BudLateralMeristem.IsValid() && BudLateralMeristem.IsValidIndex(Index))
+		{
+			BudLateralMeristem.ModifyAt(Index, InBudLateralMeristem);
+		}
+	}
+
 	float FPointFacade::GetBranchGradient(int32 Index) const
 	{
 		if (BranchGradients.IsValid() && BranchGradients.IsValidIndex(Index))
@@ -704,6 +765,16 @@ namespace PV::Facades
 		GetValueFromAccessor<int32>(BudNumber, OutValue, Index, -1);
 	}
 
+	const TManagedArray<int32>& FPointFacade::GetBudNumbersAttribute() const
+	{
+		return BudNumber.Get();
+	}
+
+	void FPointFacade::SetBudNumbersFromIndex(const TArray<int32>& InputArray, const int32 StartIndex)
+	{
+		FILL_COLLECTION_ATTRIBUTE_FROM_INDEX(BudNumber, InputArray, StartIndex);
+	}
+
 	void FPointFacade::GetBudLightDetected(int32 Index, TArray<float>& OutValue) const
 	{
 		GetValueFromAccessor<TArray<float>>(BudLightDetected, OutValue, Index, TArray<float>());
@@ -724,6 +795,14 @@ namespace PV::Facades
 		GetValueFromAccessor<TArray<int>>(BudStatus, OutValue, Index, TArray<int>());
 	}
 
+	void FPointFacade::SetBudStatus(int32 Index, const TArray<int>& InValue)
+	{
+		if (BudStatus.IsValid() && BudStatus.IsValidIndex(Index))
+		{
+			BudStatus.ModifyAt(Index, InValue);
+		}
+	}
+
 	void FPointFacade::GetBudLateralMeristem(int32 Index, TArray<float>& OutValue) const
 	{
 		GetValueFromAccessor<TArray<float>>(BudLateralMeristem, OutValue, Index, TArray<float>());
@@ -754,6 +833,11 @@ namespace PV::Facades
 		FILL_COLLECTION_ATTRIBUTE_FROM_INDEX(Positions, InputArray, StartIndex);
 	}
 
+	void FPointFacade::SetPointScalesFromIndex(const TArray<float>& InputArray, const int32 StartIndex)
+	{
+		FILL_COLLECTION_ATTRIBUTE_FROM_INDEX(PointScale, InputArray, StartIndex);
+	}
+
 	void FPointFacade::SetLFRsFromIndex(const TArray<float>& InputArray, const int32 StartIndex)
 	{
 		FILL_COLLECTION_ATTRIBUTE_FROM_INDEX(LengthFromRoot, InputArray, StartIndex);
diff --git a/Engine/Plugins/Experimental/ProceduralVegetationEditor/Source/ProceduralVegetation/Private/Implementations/PVGrafter.cpp b/Engine/Plugins/Experimental/ProceduralVegetationEditor/Source/ProceduralVegetation/Private/Implementations/PVGrafter.cpp
--- a/Engine/Plugins/Experimental/ProceduralVegetationEditor/Source/ProceduralVegetation/Private/Implementations/PVGrafter.cpp
+++ b/Engine/Plugins/Experimental/ProceduralVegetationEditor/Source/ProceduralVegetation/Private/Implementations/PVGrafter.cpp
@@ -9,17 +9,149 @@
 #include "Facades/PVAttributesNames.h"
 #include "Facades/PVBranchFacade.h"
 #include "Facades/PVBudVectorsFacade.h"
-#include "Facades/PVFoliageFacade.h"
 #include "Facades/PVPointFacade.h"
+#include "Helpers/PVAttributesHelper.h"
 #include "Helpers/PVFoliageJSONHelper.h"
 
 
+template <typename T>
+static FORCEINLINE T LerpValue(const T& A, const T& B, float Alpha)
+{
+	if constexpr (TIsIntegral<T>::Value)
+	{
+		const float Val = FMath::Lerp(static_cast<float>(A), static_cast<float>(B), Alpha);
+		return static_cast<T>(FMath::RoundToInt(Val));
+	}
+	else
+	{
+		return FMath::Lerp(A, B, Alpha);
+	}
+}
+
+template <typename T>
+TArray<T> FPVGrafter::LerpArrayElements(const TArray<T>& InArray1, const TArray<T>& InArray2, const float Alpha)
+{
+	check(InArray1.Num() == InArray2.Num());
+	TArray<T> OutArray;
+	OutArray.SetNumUninitialized(InArray1.Num());
+
+	for (int32 i = 0; i < InArray1.Num(); i++)
+	{
+		OutArray[i] = LerpValue(InArray1[i], InArray2[i], Alpha);
+	}
+
+	return OutArray;
+}
+
+template <typename Fn>
+void FPVGrafter::ForEachGraftPointExcludingRoot(const int32 NumPoints, const int32 RootIndex, Fn&& Func)
+{
+	for (int32 SourceIndex = 0, TargetIndex = 0; SourceIndex < NumPoints; SourceIndex++)
+	{
+		if (SourceIndex == RootIndex)
+		[[likely]]
+		{
+			continue;
+		}
+
+		Func(SourceIndex, TargetIndex);
+		TargetIndex++;
+	}
+}
+
+int32 FPVGrafter::AddPointToCollection(FManagedArrayCollection& OutCollection, const int32 CurrentPointIndex, const int32 NextPointIndex,
+                                       const float Alpha, const int32 BudNumber)
+{
+	PV::Facades::FPointFacade PointFacade(OutCollection);
+
+	const FVector3f CurrentPointPosition = PointFacade.GetPosition(CurrentPointIndex);
+	const FVector3f NextPointPosition = PointFacade.GetPosition(NextPointIndex);
+	const FVector3f Position = FMath::Lerp(CurrentPointPosition, NextPointPosition, Alpha);
+
+	const float CurrentPointLengthFromRoot = PointFacade.GetLengthFromRoot(CurrentPointIndex);
+	const float NextPointLengthFromRoot = PointFacade.GetLengthFromRoot(NextPointIndex);
+	const float LengthFromRoot = FMath::Lerp(CurrentPointLengthFromRoot, NextPointLengthFromRoot, Alpha);
+
+	const float CurrentPointScale = PointFacade.GetPointScale(CurrentPointIndex);
+	const float NextPointScale = PointFacade.GetPointScale(NextPointIndex);
+	const float PointScale = FMath::Lerp(CurrentPointScale, NextPointScale, Alpha);
+
+	const float CurrentPointLengthFromSeed = PointFacade.GetLengthFromSeed(CurrentPointIndex);
+	const float NextPointLengthFromSeed = PointFacade.GetLengthFromSeed(NextPointIndex);
+	const float LengthFromSeed = FMath::Lerp(CurrentPointLengthFromSeed, NextPointLengthFromSeed, Alpha);
+
+	const TArray<float>& CurrentPointBudLightDetected = PointFacade.GetBudLightDetected(CurrentPointIndex);
+	const TArray<float>& NextPointBudLightDetected = PointFacade.GetBudLightDetected(NextPointIndex);
+	const TArray<float> BudLightDetected = LerpArrayElements(CurrentPointBudLightDetected, NextPointBudLightDetected, Alpha);
+
+	const TArray<int>& CurrentPointBudDevelopment = PointFacade.GetBudDevelopment(CurrentPointIndex);
+	const TArray<int>& NextPointBudDevelopment = PointFacade.GetBudDevelopment(NextPointIndex);
+	const TArray<int> BudDevelopment = LerpArrayElements(CurrentPointBudDevelopment, NextPointBudDevelopment, Alpha);
+
+	const TArray<float>& CurrentPointBudLateralMeristem = PointFacade.GetBudLateralMeristem(CurrentPointIndex);
+	const TArray<float>& NextPointBudLateralMeristem = PointFacade.GetBudLateralMeristem(NextPointIndex);
+	const TArray<float> BudLateralMeristem = LerpArrayElements(CurrentPointBudLateralMeristem, NextPointBudLateralMeristem, Alpha);
+
+	const TArray<float>& CurrentPointBudHormoneLevels = PointFacade.GetBudHormoneLevels(CurrentPointIndex);
+	const TArray<float>& NextPointBudHormoneLevels = PointFacade.GetBudHormoneLevels(NextPointIndex);
+	const TArray<float> BudHormoneLevels = LerpArrayElements(CurrentPointBudHormoneLevels, NextPointBudHormoneLevels, Alpha);
+
+	const TArray<FVector3f>& CurrentPointBudDirections = PointFacade.GetBudDirection(CurrentPointIndex);
+	const TArray<FVector3f>& NextPointBudDirections = PointFacade.GetBudDirection(NextPointIndex);
+	const TArray<FVector3f> BudDirections = LerpArrayElements(CurrentPointBudDirections, NextPointBudDirections, Alpha);
+
+	TArray<int> CurrentPointBudStatus;
+	PointFacade.GetBudStatus(CurrentPointIndex, CurrentPointBudStatus);
+	TArray<int> NextPointBudStatus;
+	PointFacade.GetBudStatus(NextPointIndex, NextPointBudStatus);
+	TArray<int> BudStatus;
+	BudStatus.SetNum(CurrentPointBudStatus.Num());
+
+	check(BudStatus.Num() == 10);
+
+	// 0_ApicalMeristem
+	BudStatus[0] = 0;
+	// 1_Codominant
+	BudStatus[1] = 0;
+	// 2_Axillary
+	BudStatus[2] = 1;
+	// 3_Seed
+	BudStatus[3] = 0;
+	// 4_Dormant
+	BudStatus[4] = 0;
+	// 5_Triggered
+	BudStatus[5] = 1;
+	// 6_NumTriggered
+	BudStatus[6] = 1;
+	// 7_Inactive
+	BudStatus[7] = 0;
+	// 8_BrokenTip
+	BudStatus[8] = 0;
+	// 9_Broken
+	BudStatus[9] = 0;
+
+	const int32 NewIndex = PointFacade.AddElements(1);
+	PointFacade.SetPosition(NewIndex, Position);
+	PointFacade.SetLengthFromRoot(NewIndex, LengthFromRoot);
+	PointFacade.SetPointScale(NewIndex, PointScale);
+	PointFacade.SetLengthFromSeed(NewIndex, LengthFromSeed);
+	PointFacade.SetBudLightDetected(NewIndex, BudLightDetected);
+	PointFacade.SetBudDevelopment(NewIndex, BudDevelopment);
+	PointFacade.SetBudLateralMeristem(NewIndex, BudLateralMeristem);
+	PointFacade.SetBudNumber(NewIndex, BudNumber);
+	PointFacade.SetBudHormoneLevels(NewIndex, BudHormoneLevels);
+	PointFacade.SetBudDirections(NewIndex, BudDirections);
+	PointFacade.SetBudStatus(NewIndex, BudStatus);
+
+	return NewIndex;
+}
+
 FTransform FPVGrafter::GetPointTransform(const FVector3f AttachmentPoint, const FVector3f UpVector, const FVector3f NormalVector, const float Scale)
 {
 	const FVector UpVectorN = FVector(UpVector).GetSafeNormal();
 	if (UpVectorN.IsNearlyZero())
 	{
-		return FTransform(FQuat::Identity, FVector(AttachmentPoint), FVector::OneVector);
+		return FTransform(FQuat::Identity, FVector(AttachmentPoint), FVector(Scale));
 	}
 
 	const FVector NormalVectorN = FVector(NormalVector).GetSafeNormal();
@@ -33,338 +165,580 @@ FTransform FPVGrafter::GetPointTransform(const FVector3f AttachmentPoint, const
 	}
 
 	const FQuat RotationQuat = FQuat(FRotationMatrix::MakeFromYZ(OrthogonalizedNormalVectorN, UpVectorN));
-	return FTransform(RotationQuat, FVector(AttachmentPoint), FVector::OneVector);
+	return FTransform(RotationQuat, FVector(AttachmentPoint), FVector(Scale));
+}
+
+void FPVGrafter::RecomputeBranchChildren(
+	PV::FBranchNumberAttributeConstView BranchNumberAttribute,
+	PV::FBranchParentNumberAttributeConstView BranchParentNumberAttribute,
+	PV::FBranchChildrenAttributeView BranchChildrenAttribute)
+{
+	const int32 NumberOfElements = BranchParentNumberAttribute.Num();
+	TMap<int32, TArray<int32>> BranchNumbersToBranchChildren;
+	BranchNumbersToBranchChildren.Reserve(NumberOfElements);
+
+	for (int32 i = 0; i < NumberOfElements; i++)
+	{
+		const int32 BranchNumber = BranchNumberAttribute[i];
+		const int32 BranchParentNumber = BranchParentNumberAttribute[i];
+
+		if (BranchParentNumber != 0)
+		[[likely]]
+		{
+			BranchNumbersToBranchChildren.FindOrAdd(BranchParentNumber).Add(BranchNumber);
+		}
+	}
+
+	for (int32 i = 0; i < NumberOfElements; i++)
+	{
+		const int32 BranchNumber = BranchNumberAttribute[i];
+		BranchChildrenAttribute[i] = BranchNumbersToBranchChildren.Contains(BranchNumber)
+			? BranchNumbersToBranchChildren[BranchNumber]
+			: TArray<int32>();
+	}
+}
+
+void FPVGrafter::ComputeBranchGradientAttribute(
+	PV::FPointBranchGradientAttributeView PointBranchGradientAttribute,
+	PV::FBranchParentNumberAttributeConstView BranchParentNumberAttribute,
+	PV::FBranchChildrenAttributeConstView BranchChildrenAttribute,
+	PV::FBranchNumberAttributeConstView BranchNumberAttribute,
+	PV::FBranchPointsAttributeConstView BranchPointsAttribute)
+{
+	using namespace PV::AttributesHelper;
+
+	RecursiveWalkBranches(
+		BranchParentNumberAttribute,
+		BranchChildrenAttribute,
+		BranchNumberAttribute,
+		[&](int32 BranchIndex)-> EForEachResult
+			{
+				const TArray<int32>& BranchPoints = BranchPointsAttribute[BranchIndex];
+				check(BranchPoints.Num() > 1);
+
+				for (int32 i = 0; i < BranchPoints.Num(); ++i)
+				{
+					const float BranchAlpha = i / static_cast<float>(BranchPoints.Num() - 1);
+					PointBranchGradientAttribute[BranchPoints[i]] = 1.0 - BranchAlpha;
+				}
+
+				return EForEachResult::Continue;
+			}
+	);
+}
+
+void FPVGrafter::RescaleBranchForGraftingOnTip(FManagedArrayCollection& OutCollection, const int32 BranchIndex, const int32 BranchTipPointIndex,
+                                               const float GraftTrunkLength)
+{
+	PV::Facades::FPointFacade PointFacade(OutCollection);
+	PV::Facades::FBranchFacade BranchFacade(OutCollection);
+
+	const TArray<int32>& BranchPoints = BranchFacade.GetPoints(BranchIndex);
+	if (BranchPoints.Num() < 2)
+	{
+		return;
+	}
+
+	TArray<int32> SortedPointIndices = BranchPoints;
+	SortedPointIndices.Sort([PointFacade](const int32& A, const int32& B)
+		{
+			return PointFacade.GetLengthFromRoot(A) < PointFacade.GetLengthFromRoot(B);
+		});
+
+	const int32 RootPointIndex = SortedPointIndices[0];
+
+	const float RootPointLFR = PointFacade.GetLengthFromRoot(RootPointIndex);
+	const float TipPointLFR = PointFacade.GetLengthFromRoot(BranchTipPointIndex);
+	const float BranchLength = TipPointLFR - RootPointLFR;
+
+	const float CombinedLengths = BranchLength + GraftTrunkLength;
+	if (CombinedLengths <= KINDA_SMALL_NUMBER)
+	{
+		return;
+	}
+
+	const float Ratio = BranchLength / CombinedLengths;
+	const float RootPointScale = PointFacade.GetPointScale(RootPointIndex);
+	const float TipPointScale = PointFacade.GetPointScale(BranchTipPointIndex);
+	const float NewTipPointScale = FMath::Lerp(RootPointScale, TipPointScale, Ratio);
+
+	int32 i = 1;
+	while (SortedPointIndices[i - 1] != BranchTipPointIndex)
+	{
+		const int32 PtIndex = SortedPointIndices[i];
+		const float PtScale = PointFacade.GetPointScale(PtIndex);
+
+		const float NewPtScale = FMath::GetMappedRangeValueClamped(
+			FVector2f(RootPointScale, TipPointScale),
+			FVector2f(RootPointScale, NewTipPointScale),
+			PtScale);
+
+		PointFacade.SetPointScale(PtIndex, NewPtScale);
+		i++;
+	}
 }
 
 void FPVGrafter::SetGraft(FManagedArrayCollection& OutCollection, const FManagedArrayCollection& SourceSkeletonCollection,
-                          const FManagedArrayCollection& SourceGraftCollection, const FPVGraftAttributes& InAttributes,
-                          const TMap<int32, int32>& FoliageNamesOldIDsToNewIDs)
+                          const FManagedArrayCollection& SourceGraftCollection, const FPVGraftAttributes& InAttributes)
 {
 	// Transform and set points
 	PV::Facades::FPointFacade PointFacadeOut(OutCollection);
-	const PV::Facades::FPointFacade PointFacadeSourceGraft(SourceGraftCollection);
+	PV::Facades::FPointFacade PointFacadeSourceGraft(SourceGraftCollection);
+	const PV::Facades::FBranchFacade BranchFacadeSourceGraft(SourceGraftCollection);
+
+	int32 SourceGraftTrunkBranchIndex = INDEX_NONE;
+	for (int32 i = 0; i < BranchFacadeSourceGraft.GetElementCount(); ++i)
+	{
+		if (BranchFacadeSourceGraft.IsTrunk(i))
+		{
+			SourceGraftTrunkBranchIndex = i;
+			break;
+		}
+	}
+	check(SourceGraftTrunkBranchIndex != INDEX_NONE)
+
+	const TArray<int32>& SourceGraftTrunkPoints = BranchFacadeSourceGraft.GetPoints(SourceGraftTrunkBranchIndex);
+	check(SourceGraftTrunkPoints.Num() > 0);
+	const int32 SourceGraftTrunkRootPointIndex = SourceGraftTrunkPoints[0];
+
+	// Compute graft trunk length
+	float GraftTrunkLength = 0;
+	for (int32 PtIndex : SourceGraftTrunkPoints)
+	{
+		GraftTrunkLength = FMath::Max(GraftTrunkLength, PointFacadeSourceGraft.GetLengthFromRoot(PtIndex));
+	}
+	GraftTrunkLength *= InAttributes.Scale;
+
+	// Get max bud age from source graft and add to all points' bud age in the source skeleton
+	const int32 MaxBudAgeInSourceGraft = PointFacadeSourceGraft.GetMaxBudAge();
+	PointFacadeOut.AddToBudAgeForAllPoints(MaxBudAgeInSourceGraft);
+
+	const TManagedArray<int32>& SourceSkeletonBudNumbers = PointFacadeOut.GetBudNumbersAttribute();
+	check(SourceSkeletonBudNumbers.Num() > 0);
+	int32 MaxBudNumber = *Algo::MaxElement(SourceSkeletonBudNumbers.GetConstArray());
+	const TManagedArray<int32>& SourceGraftBudNumbers = PointFacadeSourceGraft.GetBudNumbersAttribute();
+	TArray<int32> TransformedBudNumbers;
+	const int32 NumOfPointsInSourceGraft = PointFacadeSourceGraft.GetElementCount();
+	// We will ignore the main trunk's root point when adding points to skeleton
+	const int32 NumberOfPointsToAddToSkeleton = NumOfPointsInSourceGraft - 1;
+	TransformedBudNumbers.SetNum(NumberOfPointsToAddToSkeleton);
+	TMap<int32, int32> SourceGraftOldToNewBudNumbers;
+	SourceGraftOldToNewBudNumbers.Reserve(NumberOfPointsToAddToSkeleton);
+	int32 SourcePointIndex = 0;
+	int32 TargetPointIndex = 0;
+	while (SourcePointIndex < NumOfPointsInSourceGraft)
+	{
+		const int32 OldBudNumber = SourceGraftBudNumbers[SourcePointIndex];
+		const int32 NewBudNumber = OldBudNumber + MaxBudNumber + 1;
+		SourceGraftOldToNewBudNumbers.Add(OldBudNumber, NewBudNumber);
+		if (SourcePointIndex != SourceGraftTrunkRootPointIndex)
+		[[likely]]
+		{
+			TransformedBudNumbers[TargetPointIndex] = NewBudNumber;
+			TargetPointIndex++;
+		}
+		SourcePointIndex++;
+	}
+
+	const int32 PointFacadeOutStartIndex = PointFacadeOut.GetElementCount();
+	PointFacadeOut.SetBudNumbersFromIndex(TransformedBudNumbers, PointFacadeOutStartIndex);
 
 	const TManagedArray<FVector3f>& PointPositions = PointFacadeSourceGraft.GetPositions();
+	const FVector GraftRootPosition = FVector(PointPositions[SourceGraftTrunkRootPointIndex]);
 	const FTransform Transform = GetPointTransform(InAttributes.AttachmentPoint, InAttributes.UpVector, InAttributes.NormalVector,
 		InAttributes.Scale);
 	TArray<FVector3f> TransformedPoints;
-	const int32 NumOfPoints = PointPositions.Num();
-	TransformedPoints.SetNum(NumOfPoints);
-	for (int32 i = 0; i < NumOfPoints; ++i)
+	TransformedPoints.SetNum(NumberOfPointsToAddToSkeleton);
+
+	ForEachGraftPointExcludingRoot(NumOfPointsInSourceGraft, SourceGraftTrunkRootPointIndex,
+		[&](const int32 SourcePointIndex, const int32 TargetPointIndex)
+			{
+				const FVector SourcePointPosition = FVector(PointPositions[SourcePointIndex]);
+				const FVector SourcePositionWrtRoot = SourcePointPosition - GraftRootPosition;
+				const FVector RotatedScaledOffsetPosition = Transform.TransformPosition(SourcePositionWrtRoot);
+				TransformedPoints[TargetPointIndex] = FVector3f(RotatedScaledOffsetPosition);
+			});
+
+	PointFacadeOut.SetPositionsFromIndex(TransformedPoints, PointFacadeOutStartIndex);
+
+
+	const TManagedArray<float>& SourceGraftPointScales = PointFacadeSourceGraft.GetPointScales();
+	TArray<float> TransformedPointScales;
+	TransformedPointScales.SetNum(NumberOfPointsToAddToSkeleton);
+
+	const float ParentPointPointScale = PointFacadeOut.GetPointScale(InAttributes.PointIndex);
+	float SourceGraftRootPointScale = SourceGraftPointScales[SourceGraftTrunkRootPointIndex];
+
+	if (FMath::IsNearlyZero(SourceGraftRootPointScale))
+	[[unlikely]]
 	{
-		TransformedPoints[i] = FVector3f(Transform.TransformPosition(FVector(PointPositions[i])));
+		SourceGraftRootPointScale = 1;
 	}
+	const float ParentToGraftRootPointScaleRatio = ParentPointPointScale / SourceGraftRootPointScale;
 
-	const int32 PointFacadeOutStartIndex = PointFacadeOut.GetPositions().Num();
-	PointFacadeOut.SetPositionsFromIndex(TransformedPoints, PointFacadeOutStartIndex);
+	ForEachGraftPointExcludingRoot(NumOfPointsInSourceGraft, SourceGraftTrunkRootPointIndex,
+		[&](const int32 SourcePointIndex, const int32 TargetPointIndex)
+			{
+				TransformedPointScales[TargetPointIndex] =
+					SourceGraftPointScales[SourcePointIndex] * ParentToGraftRootPointScaleRatio;
+			});
+
+	PointFacadeOut.SetPointScalesFromIndex(TransformedPointScales, PointFacadeOutStartIndex);
 
 	const TManagedArray<float>& LFRs = PointFacadeSourceGraft.GetLengthFromRootsArray();
 	TArray<float> TransformedLFRs;
-	TransformedLFRs.SetNum(NumOfPoints);
-	for (int32 i = 0; i < NumOfPoints; ++i)
-	{
-		TransformedLFRs[i] = LFRs[i] + (InAttributes.LengthFromRoot * InAttributes.Scale);
-	}
+	TransformedLFRs.SetNum(NumberOfPointsToAddToSkeleton);
+
+	ForEachGraftPointExcludingRoot(NumOfPointsInSourceGraft, SourceGraftTrunkRootPointIndex,
+		[&](const int32 SourcePointIndex, const int32 TargetPointIndex)
+			{
+				TransformedLFRs[TargetPointIndex] = ((LFRs[SourcePointIndex] - LFRs[SourceGraftTrunkRootPointIndex]) * InAttributes.Scale)
+					+ InAttributes.LengthFromRoot;
+			});
 
 	PointFacadeOut.SetLFRsFromIndex(TransformedLFRs, PointFacadeOutStartIndex);
 
 	const TManagedArray<float>& LengthFromSeeds = PointFacadeSourceGraft.GetLengthFromSeedsArray();
 	TArray<float> TransformedLengthFromSeeds;
-	TransformedLengthFromSeeds.SetNum(LengthFromSeeds.Num());
-	for (int32 i = 0; i < LengthFromSeeds.Num(); ++i)
-	{
-		TransformedLengthFromSeeds[i] = LengthFromSeeds[i] + (InAttributes.LengthFromRoot * InAttributes.Scale);
-	}
+	TransformedLengthFromSeeds.SetNum(NumberOfPointsToAddToSkeleton);
+
+	ForEachGraftPointExcludingRoot(NumOfPointsInSourceGraft, SourceGraftTrunkRootPointIndex,
+		[&](const int32 SourcePointIndex, const int32 TargetPointIndex)
+			{
+				TransformedLengthFromSeeds[TargetPointIndex] = ((LengthFromSeeds[SourcePointIndex] - LengthFromSeeds[SourceGraftTrunkRootPointIndex])
+					* InAttributes.Scale) + InAttributes.LengthFromRoot;
+			});
 
 	PointFacadeOut.SetLengthFromSeedsArrayFromIndex(TransformedLengthFromSeeds, PointFacadeOutStartIndex);
 
 	const TManagedArray<TArray<float>>& BudLightDetectedArrays = PointFacadeSourceGraft.GetBudLightDetectedArrays();
-	PointFacadeOut.SetBudLightDetectedArraysFromIndex(BudLightDetectedArrays.GetConstArray(), PointFacadeOutStartIndex);
 
-	const int32 NumOfPointsInSourceGraft = PointFacadeSourceGraft.GetElementCount();
-	for (int32 i = 0; i < NumOfPointsInSourceGraft; ++i)
-	{
-		const TArray<int>& BudDevelopmentValuesInGraftPoint = PointFacadeSourceGraft.GetBudDevelopment(i);
-		const TArray<int>& BudDevelopmentValuesInAttachmentPoint = PointFacadeOut.GetBudDevelopment(InAttributes.PointIndex);
-		const int32 NumOfValuesInSource = BudDevelopmentValuesInGraftPoint.Num();
-		TArray<int> TransformedValues;
-		TransformedValues.SetNum(NumOfValuesInSource);
-		check(NumOfValuesInSource == 6);
-
-		// 0_Generation
-		TransformedValues[0] = BudDevelopmentValuesInGraftPoint[0] + BudDevelopmentValuesInAttachmentPoint[0];
-
-		// 1_BudAge
-		TransformedValues[1] = BudDevelopmentValuesInGraftPoint[1] + BudDevelopmentValuesInAttachmentPoint[1];
-		
-		// 2_BranchAge
-		TransformedValues[2] = BudDevelopmentValuesInGraftPoint[2] + BudDevelopmentValuesInAttachmentPoint[2];
-		
-		// 3_AgeSenescense
-		TransformedValues[3] = 0;
-		
-		// 4_LightSenescense
-		TransformedValues[4] = 0;
-		
-		// 5_RelativeBudAge
-		// TODO: recompute 0-1 for whole tree
-		TransformedValues[5] = BudDevelopmentValuesInGraftPoint[5];
-		
-		PointFacadeOut.SetBudDevelopment(i + PointFacadeOutStartIndex, TransformedValues);
-	}
-	
-	for (int32 i = 0; i < NumOfPointsInSourceGraft; ++i)
-	{
-		const TArray<float>& BudHormoneLevelsInGraftPoint = PointFacadeSourceGraft.GetBudHormoneLevels(i);
-		// keeping them the same as original graft for now
-		// TODO: recompute and update
-		PointFacadeOut.SetBudHormoneLevels(i + PointFacadeOutStartIndex, BudHormoneLevelsInGraftPoint);
-	}
-	
-	for (int32 i = 0; i < NumOfPointsInSourceGraft; ++i)
-	{
-		const TArray<FVector3f>& BudDirectionsInGraftPoint = PointFacadeSourceGraft.GetBudDirection(i);
-		const int32 NumOfValues = BudDirectionsInGraftPoint.Num();
-		TArray<FVector3f> TransformedValues;
-		TransformedValues.SetNum(NumOfValues);
-		
-		for (int32 j = 0; j < NumOfValues; ++j)
-		{
-			TransformedValues[j] = FVector3f(Transform.TransformVectorNoScale(FVector(BudDirectionsInGraftPoint[j])));
-		}
-		PointFacadeOut.SetBudDirections(i + PointFacadeOutStartIndex, TransformedValues);
-	}
+	ForEachGraftPointExcludingRoot(NumOfPointsInSourceGraft, SourceGraftTrunkRootPointIndex,
+		[&](const int32 SourcePointIndex, const int32 TargetPointIndex)
+			{
+				PointFacadeOut.SetBudLightDetected(TargetPointIndex + PointFacadeOutStartIndex, BudLightDetectedArrays[SourcePointIndex]);
+			});
+
+
+	ForEachGraftPointExcludingRoot(NumOfPointsInSourceGraft, SourceGraftTrunkRootPointIndex,
+		[&](const int32 SourcePointIndex, const int32 TargetPointIndex)
+			{
+				const TArray<int>& BudDevelopmentValuesInGraftPoint = PointFacadeSourceGraft.GetBudDevelopment(SourcePointIndex);
+				const TArray<int>& BudDevelopmentValuesInAttachmentPoint = PointFacadeOut.GetBudDevelopment(InAttributes.PointIndex);
+				const int32 NumOfValuesInSource = BudDevelopmentValuesInGraftPoint.Num();
+				TArray<int> TransformedValues;
+				TransformedValues.SetNum(NumOfValuesInSource);
+				check(NumOfValuesInSource == 6);
+
+				// 0_Generation
+				TransformedValues[0] = BudDevelopmentValuesInGraftPoint[0] + BudDevelopmentValuesInAttachmentPoint[0];
+
+				// 1_BudAge
+				// Max bud age from source graft has already been added to all the points in the source skeleton
+				TransformedValues[1] = BudDevelopmentValuesInGraftPoint[1];
+
+				// 2_BranchAge
+				TransformedValues[2] = BudDevelopmentValuesInGraftPoint[2] + BudDevelopmentValuesInAttachmentPoint[2];
+
+				// 3_AgeSenescence
+				TransformedValues[3] = 0;
+
+				// 4_LightSenescence
+				TransformedValues[4] = 0;
+
+				// 5_RelativeBudAge
+				TransformedValues[5] = TransformedValues[0];
+
+				PointFacadeOut.SetBudDevelopment(TargetPointIndex + PointFacadeOutStartIndex, TransformedValues);
+			});
+
+	ForEachGraftPointExcludingRoot(NumOfPointsInSourceGraft, SourceGraftTrunkRootPointIndex,
+		[&](const int32 SourcePointIndex, const int32 TargetPointIndex)
+			{
+				const TArray<float>& BudHormoneLevelsInGraftPoint = PointFacadeSourceGraft.GetBudHormoneLevels(SourcePointIndex);
+				// keeping them the same as original graft for now
+				// TODO: recompute and update
+				PointFacadeOut.SetBudHormoneLevels(TargetPointIndex + PointFacadeOutStartIndex, BudHormoneLevelsInGraftPoint);
+			});
+
+	ForEachGraftPointExcludingRoot(NumOfPointsInSourceGraft, SourceGraftTrunkRootPointIndex,
+		[&](const int32 SourcePointIndex, const int32 TargetPointIndex)
+			{
+				const TArray<FVector3f>& BudDirectionsInGraftPoint = PointFacadeSourceGraft.GetBudDirection(SourcePointIndex);
+				const int32 NumOfValues = BudDirectionsInGraftPoint.Num();
+				TArray<FVector3f> TransformedValues;
+				TransformedValues.SetNum(NumOfValues);
+
+				for (int32 j = 0; j < NumOfValues; ++j)
+				{
+					TransformedValues[j] = FVector3f(Transform.TransformVectorNoScale(FVector(BudDirectionsInGraftPoint[j])));
+				}
+				PointFacadeOut.SetBudDirections(TargetPointIndex + PointFacadeOutStartIndex, TransformedValues);
+			});
 
 	// Transform and set branches
 	PV::Facades::FBranchFacade BranchFacadeOut(OutCollection);
-	const PV::Facades::FBranchFacade BranchFacadeSourceGraft(SourceGraftCollection);
 
 	const int32 NumOfBranchesInSource = BranchFacadeSourceGraft.GetElementCount();
 	const int32 NumOfBranchesInTarget = BranchFacadeOut.GetElementCount();
-	BranchFacadeOut.AddElements(NumOfBranchesInSource);
-	for (int32 i = 0; i < NumOfBranchesInSource; ++i)
+
+	int32 NumOfBranchesToAdd = NumOfBranchesInSource - (InAttributes.bIsTipInstance
+		? 1
+		: 0);
+	BranchFacadeOut.AddElements(NumOfBranchesToAdd);
+
+	int32 SourceIndex = 0;
+	int32 TargetIndex = NumOfBranchesInTarget;
+	while (SourceIndex < NumOfBranchesInSource)
 	{
-		const TArray<int32>& SourceBranchPoints = BranchFacadeSourceGraft.GetPoints(i);
+		TArray<int32> SourceBranchPoints = BranchFacadeSourceGraft.GetPoints(SourceIndex);
+
 		const int32 NumOfPointsInSource = SourceBranchPoints.Num();
 		TArray<int32> TransformedBranchPoints;
-		TransformedBranchPoints.SetNum(NumOfPointsInSource);
+		TransformedBranchPoints.Reserve(NumOfPointsInSource);
 
 		for (int32 j = 0; j < NumOfPointsInSource; ++j)
 		{
-			TransformedBranchPoints[j] = SourceBranchPoints[j] + PointFacadeOutStartIndex;
+			if (const int32 BranchPoint = SourceBranchPoints[j];
+				BranchPoint < SourceGraftTrunkRootPointIndex)
+			{
+				TransformedBranchPoints.Add(PointFacadeOutStartIndex + BranchPoint);
+			}
+			else if (BranchPoint > SourceGraftTrunkRootPointIndex)
+			{
+				TransformedBranchPoints.Add(PointFacadeOutStartIndex + BranchPoint - 1);
+			}
 		}
 
-		BranchFacadeOut.SetPoints(NumOfBranchesInTarget + i, TransformedBranchPoints);
+		if (SourceIndex == SourceGraftTrunkBranchIndex)
+		[[unlikely]]
+		{
+			if (InAttributes.bIsTipInstance)
+			{
+				TArray<int32> TargetBranchPoints = BranchFacadeOut.GetPoints(InAttributes.BranchIndex);
+				TargetBranchPoints.Append(TransformedBranchPoints);
+				BranchFacadeOut.SetPoints(InAttributes.BranchIndex, TargetBranchPoints);
+
+				RescaleBranchForGraftingOnTip(OutCollection, InAttributes.BranchIndex, InAttributes.PointIndex, GraftTrunkLength);
+			}
+			else
+			{
+				TransformedBranchPoints.Insert(InAttributes.PointIndex, 0);
+				BranchFacadeOut.SetPoints(TargetIndex, TransformedBranchPoints);
+				TargetIndex++;
+			}
+		}
+		else
+		[[likely]]
+		{
+			BranchFacadeOut.SetPoints(TargetIndex, TransformedBranchPoints);
+			TargetIndex++;
+		}
+
+		SourceIndex++;
 	}
 
+
 	check(NumOfBranchesInTarget > 0);
 	const int32 BranchFacadeOutStartIndex = NumOfBranchesInTarget;
 	const int32 MaxBranchNumberInTarget = *Algo::MaxElement(BranchFacadeOut.GetBranchNumbers());
+	const int32 BranchNumberBeingAttachedToInTarget = BranchFacadeOut.GetBranchNumber(InAttributes.BranchIndex);
 	const TManagedArray<int32>& BranchNumbers = BranchFacadeSourceGraft.GetBranchNumbers();
-	TArray<int32> TransformedBranchNumbers;
-	TransformedBranchNumbers.SetNum(NumOfBranchesInSource);
 	TMap<int32, int32> SourceOldBranchNumbersToNewBranchNumbers;
 	SourceOldBranchNumbersToNewBranchNumbers.Reserve(NumOfBranchesInSource);
+
 	for (int32 i = 0; i < NumOfBranchesInSource; ++i)
 	{
-		TransformedBranchNumbers[i] = BranchNumbers[i] + MaxBranchNumberInTarget;
-		SourceOldBranchNumbersToNewBranchNumbers.Add(BranchNumbers[i], TransformedBranchNumbers[i]);
+		const int32 NewBranchNumber = (i == SourceGraftTrunkBranchIndex && InAttributes.bIsTipInstance)
+			? BranchNumberBeingAttachedToInTarget
+			: BranchNumbers[i] + MaxBranchNumberInTarget;
+		SourceOldBranchNumbersToNewBranchNumbers.Add(BranchNumbers[i], NewBranchNumber);
+	}
+
+	TArray<int32> TransformedBranchNumbers;
+	TransformedBranchNumbers.SetNum(NumOfBranchesToAdd);
+	SourceIndex = 0;
+	TargetIndex = 0;
+	while (SourceIndex < NumOfBranchesInSource)
+	{
+		if (!(InAttributes.bIsTipInstance && SourceIndex == SourceGraftTrunkBranchIndex))
+		[[likely]]
+		{
+			TransformedBranchNumbers[TargetIndex] = SourceOldBranchNumbersToNewBranchNumbers[BranchNumbers[SourceIndex]];
+			TargetIndex++;
+		}
+
+		SourceIndex++;
 	}
 
 	BranchFacadeOut.SetBranchNumbersFromIndex(TransformedBranchNumbers, BranchFacadeOutStartIndex);
 
 	const TManagedArray<int32>& ParentBranchNumbers = BranchFacadeSourceGraft.GetParentBranchNumbers();
 	TArray<int32> TransformedParentBranchNumbers;
-	TransformedParentBranchNumbers.SetNum(ParentBranchNumbers.Num());
-	const int32 BranchNumberBeingAttachedToInTarget = BranchFacadeOut.GetBranchNumber(InAttributes.BranchIndex);
-	for (int32 i = 0; i < ParentBranchNumbers.Num(); ++i)
+	TransformedParentBranchNumbers.SetNum(NumOfBranchesToAdd);
+	SourceIndex = 0;
+	TargetIndex = 0;
+	while (SourceIndex < NumOfBranchesInSource)
 	{
-		TransformedParentBranchNumbers[i] = ParentBranchNumbers[i] == 0
-			? BranchNumberBeingAttachedToInTarget
-			: SourceOldBranchNumbersToNewBranchNumbers[BranchNumbers[i]];
+		if (!(InAttributes.bIsTipInstance && SourceIndex == SourceGraftTrunkBranchIndex))
+		[[likely]]
+		{
+			TransformedParentBranchNumbers[TargetIndex] = ParentBranchNumbers[SourceIndex] == 0
+				? BranchNumberBeingAttachedToInTarget
+				: SourceOldBranchNumbersToNewBranchNumbers[BranchNumbers[SourceIndex]];
+			TargetIndex++;
+		}
+
+		SourceIndex++;
 	}
 
 	BranchFacadeOut.SetParentBranchNumbersFromIndex(TransformedParentBranchNumbers, BranchFacadeOutStartIndex);
 
 	const TManagedArray<int32>& BranchHierarchyNumbers = BranchFacadeSourceGraft.GetBranchHierarchyNumbers();
 	TArray<int32> TransformedBranchHierarchyNumbers;
-	TransformedBranchHierarchyNumbers.SetNum(BranchHierarchyNumbers.Num());
+	TransformedBranchHierarchyNumbers.SetNum(NumOfBranchesToAdd);
 	const int32 BranchHierarchyNumberOfTarget = BranchFacadeOut.GetBranchHierarchyNumber(InAttributes.BranchIndex);
-	for (int32 i = 0; i < BranchHierarchyNumbers.Num(); ++i)
+	SourceIndex = 0;
+	TargetIndex = 0;
+	while (SourceIndex < NumOfBranchesInSource)
 	{
-		TransformedBranchHierarchyNumbers[i] = BranchHierarchyNumbers[i] + BranchHierarchyNumberOfTarget;
+		if (!(InAttributes.bIsTipInstance && SourceIndex == SourceGraftTrunkBranchIndex))
+		[[likely]]
+		{
+			TransformedBranchHierarchyNumbers[TargetIndex] = BranchHierarchyNumbers[SourceIndex] + BranchHierarchyNumberOfTarget;
+			TargetIndex++;
+		}
+
+		SourceIndex++;
 	}
 
 	BranchFacadeOut.SetBranchHierarchyNumbersFromIndex(TransformedBranchHierarchyNumbers, BranchFacadeOutStartIndex);
 
-	// TODO: Add code for updating branchSourceBudNumber after adding code for updating budNumbers in points
-
-	for (int32 i = 0; i < NumOfBranchesInSource; ++i)
+	// Commented out for now since Grower node is currently generating inconsistent source bud numbers 
+	//
+	// const TManagedArray<int32>& BranchSourceBudNumbers = BranchFacadeSourceGraft.GetBranchSourceBudNumbers();
+	// TArray<int32> TransformedBranchSourceBudNumbers;
+	// TransformedBranchSourceBudNumbers.SetNum(NumOfBranchesToAdd);
+	//
+	// SourceIndex = 0;
+	// TargetIndex = 0;
+	// while (SourceIndex < NumOfBranchesInSource)
+	// {
+	// 	if (!(InAttributes.bIsTipInstance && SourceIndex == SourceTrunkBranchIndex))
+	// 	[[likely]]
+	// 	{
+	// 		const int32 ParentPointSourceBudNumber = PointFacadeOut.GetBudNumber(InAttributes.PointIndex);
+	// 		const int32 SourceBudNumber = BranchSourceBudNumbers[SourceIndex];
+	// 		const int32 NewSourceBudNumber = SourceBudNumber == 1
+	// 			? ParentPointSourceBudNumber
+	// 			: SourceGraftOldToNewBudNumbers[SourceBudNumber];
+	// 		TransformedBranchSourceBudNumbers[TargetIndex] = NewSourceBudNumber;
+	// 		TargetIndex++;
+	// 	}
+	//
+	// 	SourceIndex++;
+	// }
+	//
+	// BranchFacadeOut.SetBranchSourceBudNumbersFromIndex(TransformedBranchSourceBudNumbers, BranchFacadeOutStartIndex);
+
+
+	SourceIndex = 0;
+	TargetIndex = NumOfBranchesInTarget;
+	while (SourceIndex < NumOfBranchesInSource)
 	{
-		const TArray<int32>& SourceBranchChildren = BranchFacadeSourceGraft.GetChildren(i);
-		const int32 NumOfChildrenInSource = SourceBranchChildren.Num();
-		TArray<int32> TransformedChildren;
-		TransformedChildren.SetNum(NumOfChildrenInSource);
-
-		for (int32 j = 0; j < NumOfChildrenInSource; ++j)
+		if (!(InAttributes.bIsTipInstance && SourceIndex == SourceGraftTrunkBranchIndex))
+		[[likely]]
 		{
-			if (SourceOldBranchNumbersToNewBranchNumbers.Contains(SourceBranchChildren[j]))
+			const TArray<int32>& SourceBranchChildren = BranchFacadeSourceGraft.GetChildren(SourceIndex);
+			const int32 NumOfChildrenInSource = SourceBranchChildren.Num();
+			TArray<int32> TransformedChildren;
+			TransformedChildren.SetNum(NumOfChildrenInSource);
+
+			for (int32 j = 0; j < NumOfChildrenInSource; ++j)
 			{
-				TransformedChildren[j] = SourceOldBranchNumbersToNewBranchNumbers[SourceBranchChildren[j]];
+				if (SourceOldBranchNumbersToNewBranchNumbers.Contains(SourceBranchChildren[j]))
+				{
+					TransformedChildren[j] = SourceOldBranchNumbersToNewBranchNumbers[SourceBranchChildren[j]];
+				}
 			}
+
+			BranchFacadeOut.SetChildren(TargetIndex, TransformedChildren);
+			TargetIndex++;
 		}
 
-		BranchFacadeOut.SetChildren(NumOfBranchesInTarget + i, TransformedChildren);
+		SourceIndex++;
 	}
 
-	for (int32 i = 0; i < NumOfBranchesInSource; ++i)
-	{
-		const TArray<int32>& SourceBranchParents = BranchFacadeSourceGraft.GetParents(i);
-		const int32 NumOfParentsInSource = SourceBranchParents.Num();
-		TArray<int32> TransformedParents;
-		TransformedParents.SetNum(NumOfParentsInSource);
 
-		for (int32 j = 0; j < NumOfParentsInSource; ++j)
+	SourceIndex = 0;
+	TargetIndex = NumOfBranchesInTarget;
+	while (SourceIndex < NumOfBranchesInSource)
+	{
+		if (!(InAttributes.bIsTipInstance && SourceIndex == SourceGraftTrunkBranchIndex))
+		[[likely]]
 		{
-			if (SourceOldBranchNumbersToNewBranchNumbers.Contains(SourceBranchParents[j]))
+			const TArray<int32>& SourceBranchParents = BranchFacadeSourceGraft.GetParents(SourceIndex);
+			const int32 NumOfParentsInSource = SourceBranchParents.Num();
+			TArray<int32> TransformedParents;
+			TransformedParents.SetNum(NumOfParentsInSource);
+
+			for (int32 j = 0; j < NumOfParentsInSource; ++j)
 			{
-				TransformedParents[j] = SourceOldBranchNumbersToNewBranchNumbers[SourceBranchParents[j]];
+				if (SourceOldBranchNumbersToNewBranchNumbers.Contains(SourceBranchParents[j]))
+				{
+					TransformedParents[j] = SourceOldBranchNumbersToNewBranchNumbers[SourceBranchParents[j]];
+				}
 			}
+
+			BranchFacadeOut.SetParents(TargetIndex, TransformedParents);
+			TargetIndex++;
 		}
 
-		BranchFacadeOut.SetParents(NumOfBranchesInTarget + i, TransformedParents);
+		SourceIndex++;
 	}
 
 	const int32 PlantNumberOfTargetBranch = BranchFacadeOut.GetBranchPlantNumber(InAttributes.BranchIndex);
-	for (int32 i = 0; i < NumOfBranchesInSource; ++i)
+	SourceIndex = 0;
+	TargetIndex = NumOfBranchesInTarget;
+	while (SourceIndex < NumOfBranchesInSource)
 	{
-		BranchFacadeOut.SetBranchPlantNumber(NumOfBranchesInTarget + i, PlantNumberOfTargetBranch);
-	}
-
-	if (!SourceGraftCollection.HasGroup(PV::GroupNames::FoliageNamesGroup))
-	{
-		return;
-	}
-
-	PV::Facades::FFoliageFacade FoliageFacadeOut(OutCollection);
-
-	const int32 NumOfFoliageEntriesInTarget = FoliageFacadeOut.GetElementCount();
-
-	for (int32 i = 0; i < NumOfBranchesInSource; ++i)
-	{
-		const TArray<int32>& SourceBranchFoliageIDs = BranchFacadeSourceGraft.GetBranchFoliageIDs(i);
-		const int32 NumOfFoliageIDsInSource = SourceBranchFoliageIDs.Num();
-		TArray<int32> TransformedFoliageIDs;
-		TransformedFoliageIDs.SetNum(NumOfFoliageIDsInSource);
-
-		for (int32 j = 0; j < NumOfFoliageIDsInSource; ++j)
+		if (!(InAttributes.bIsTipInstance && SourceIndex == SourceGraftTrunkBranchIndex))
+		[[likely]]
 		{
-			TransformedFoliageIDs[j] = SourceBranchFoliageIDs[j] + NumOfFoliageEntriesInTarget;
+			BranchFacadeOut.SetBranchPlantNumber(TargetIndex, PlantNumberOfTargetBranch);
+			TargetIndex++;
 		}
 
-		BranchFacadeOut.SetBranchFoliageIDs(NumOfBranchesInTarget + i, TransformedFoliageIDs);
-	}
-
-	// Transform and set foliage
-	const PV::Facades::FFoliageFacade FoliageFacadeSourceGraft(SourceGraftCollection);
-	const TManagedArray<int32>& FoliageNameIDs = FoliageFacadeSourceGraft.GetFoliageNameIDs();
-	TArray<int32> TransformedFoliageNameIDs;
-	TransformedFoliageNameIDs.SetNum(FoliageNameIDs.Num());
-
-	for (int32 j = 0; j < FoliageNameIDs.Num(); ++j)
-	{
-		TransformedFoliageNameIDs[j] = FoliageNamesOldIDsToNewIDs[FoliageNameIDs[j]];
+		SourceIndex++;
 	}
+}
 
-	FoliageFacadeOut.SetFoliageNameIDsFromIndex(TransformedFoliageNameIDs, NumOfFoliageEntriesInTarget);
-
-	const TManagedArray<int32>& FoliageBranchIDs = FoliageFacadeSourceGraft.GetFoliageBranchIDs();
-	TArray<int32> TransformedFoliageBranchIDs;
-	TransformedFoliageBranchIDs.SetNum(FoliageBranchIDs.Num());
-
-	for (int32 j = 0; j < FoliageBranchIDs.Num(); ++j)
-	{
-		TransformedFoliageBranchIDs[j] = FoliageBranchIDs[j] + NumOfBranchesInTarget;
-	}
-
-	FoliageFacadeOut.SetFoliageBranchIDsFromIndex(TransformedFoliageBranchIDs, NumOfFoliageEntriesInTarget);
-
-	const TManagedArray<FVector3f>& FoliagePivotPositions = FoliageFacadeSourceGraft.GetPivotPositions();
-	TArray<FVector3f> TransformedFoliagePivotPositions;
-	TransformedFoliagePivotPositions.SetNum(FoliagePivotPositions.Num());
-
-	for (int32 j = 0; j < FoliagePivotPositions.Num(); ++j)
-	{
-		TransformedFoliagePivotPositions[j] = FVector3f(Transform.TransformPosition(FVector(FoliagePivotPositions[j])));
-	}
-
-	FoliageFacadeOut.SetPivotPositionsFromIndex(TransformedFoliagePivotPositions, NumOfFoliageEntriesInTarget);
-
-	const TManagedArray<FVector3f>& FoliageUpVectors = FoliageFacadeSourceGraft.GetUpVectors();
-	TArray<FVector3f> TransformedFoliageUpVectors;
-	TransformedFoliageUpVectors.SetNum(FoliageUpVectors.Num());
-
-	for (int32 j = 0; j < FoliageUpVectors.Num(); ++j)
-	{
-		TransformedFoliageUpVectors[j] = FVector3f(Transform.TransformVectorNoScale(FVector(FoliageUpVectors[j])));
-	}
-
-	FoliageFacadeOut.SetUpVectorsFromIndex(TransformedFoliageUpVectors, NumOfFoliageEntriesInTarget);
-
-	const TManagedArray<FVector3f>& FoliageNormalVectors = FoliageFacadeSourceGraft.GetNormalVectors();
-	TArray<FVector3f> TransformedFoliageNormalVectors;
-	TransformedFoliageNormalVectors.SetNum(FoliageNormalVectors.Num());
-
-	for (int32 j = 0; j < FoliageNormalVectors.Num(); ++j)
-	{
-		TransformedFoliageNormalVectors[j] = FVector3f(Transform.TransformVectorNoScale(FVector(FoliageNormalVectors[j])));
-	}
-
-	FoliageFacadeOut.SetNormalVectorsFromIndex(TransformedFoliageNormalVectors, NumOfFoliageEntriesInTarget);
-
-	const TManagedArray<float>& FoliageScales = FoliageFacadeSourceGraft.GetScales();
-	TArray<float> TransformedFoliageScales;
-	TransformedFoliageScales.SetNum(FoliageScales.Num());
+bool FPVGrafter::RemoveFoliageDataIfPresent(FManagedArrayCollection& OutCollection)
+{
+	bool bRemovedFoliageData = false;
+	PV::Facades::FBranchFacade BranchFacade(OutCollection);
 
-	for (int32 j = 0; j < FoliageScales.Num(); ++j)
+	if (OutCollection.HasGroup(PV::GroupNames::FoliageNamesGroup) || OutCollection.HasGroup(PV::GroupNames::FoliageGroup))
 	{
-		TransformedFoliageScales[j] = FoliageScales[j] * InAttributes.Scale;
-	}
+		if (OutCollection.HasGroup(PV::GroupNames::FoliageNamesGroup))
+			OutCollection.RemoveGroup(PV::GroupNames::FoliageNamesGroup);
 
-	FoliageFacadeOut.SetScalesFromIndex(TransformedFoliageScales, NumOfFoliageEntriesInTarget);
+		if (OutCollection.HasGroup(PV::GroupNames::FoliageGroup))
+			OutCollection.RemoveGroup(PV::GroupNames::FoliageGroup);
 
-	const TManagedArray<float>& FoliageLFRs = FoliageFacadeSourceGraft.GetLFRs();
-	TArray<float> TransformedFoliageLFRs;
-	TransformedFoliageLFRs.SetNum(FoliageLFRs.Num());
-
-	for (int32 j = 0; j < FoliageLFRs.Num(); ++j)
-	{
-		TransformedFoliageLFRs[j] = FoliageLFRs[j] + InAttributes.LengthFromRoot;
-	}
-
-	FoliageFacadeOut.SetLFRsFromIndex(TransformedFoliageLFRs, NumOfFoliageEntriesInTarget);
-}
+		for (int32 i = 0; i < BranchFacade.GetElementCount(); ++i)
+		{
+			BranchFacade.SetBranchFoliageIDs(i, {});
+		}
 
-TMap<int32, int32> FPVGrafter::SetFoliageNames(FManagedArrayCollection& OutCollection, const FManagedArrayCollection& SourceGraftCollection)
-{
-	PV::Facades::FFoliageFacade FoliageFacadeOut(OutCollection);
-	const PV::Facades::FFoliageFacade FoliageFacadeSourceGraft(SourceGraftCollection);
-	const TArray<FString> FoliageNames = FoliageFacadeSourceGraft.GetFoliageNames();
-	const int32 NumOfFoliageNamesInTarget = FoliageFacadeOut.NumFoliageNames();
-	FoliageFacadeOut.SetFoliageNamesFromIndex(FoliageNames, NumOfFoliageNamesInTarget);
-
-	TMap<int32, int32> FoliageNamesOldIDsToNewIDs;
-	FoliageNamesOldIDsToNewIDs.Reserve(FoliageNames.Num());
-	for (int32 i = 0; i < FoliageNames.Num(); ++i)
-	{
-		FoliageNamesOldIDsToNewIDs.Add(i, i + NumOfFoliageNamesInTarget);
+		bRemovedFoliageData = true;
 	}
 
-	return FoliageNamesOldIDsToNewIDs;
+	return bRemovedFoliageData;
 }
 
 void FPVGrafter::DistributeGraft(FManagedArrayCollection& OutCollection, const FManagedArrayCollection& SourceSkeletonCollection,
-                                 const FManagedArrayCollection& SourceGraftCollection, const FDistributionSettings& DistributionSettings,
+                                 FManagedArrayCollection& SourceGraftCollection, const FDistributionSettings& DistributionSettings,
                                  const FScaleSettings& ScaleSettings, const FVectorSettings& VectorSettings, const FMiscSettings& MiscSettings)
 {
 	TRACE_CPUPROFILER_EVENT_SCOPE(PVGrafter::DistributeGraft);
@@ -375,6 +749,23 @@ void FPVGrafter::DistributeGraft(FManagedArrayCollection& OutCollection, const F
 	PV::Facades::FBranchFacade BranchFacadeSourceSkeleton(SourceSkeletonCollection);
 	PV::Facades::FPointFacade PointFacadeSourceSkeleton(SourceSkeletonCollection);
 
+	auto PointBranchGradientAttribute = PV::FPointBranchGradientAttribute::GetAttribute(OutCollection);
+	auto BranchParentNumberAttribute = PV::FBranchParentNumberAttribute::GetAttribute(OutCollection);
+	auto BranchChildrenAttribute = PV::FBranchChildrenAttribute::GetAttribute(OutCollection);
+	auto BranchNumberAttribute = PV::FBranchNumberAttribute::GetAttribute(OutCollection);
+	auto BranchPointsAttribute = PV::FBranchPointsAttribute::GetAttribute(OutCollection);
+
+	// Only need to do this temporarily since legacy skeleton has inconsistent data in children
+	// TODO: Remove this once support for legacy skeletons is droppped
+	RecomputeBranchChildren(BranchNumberAttribute, BranchParentNumberAttribute, BranchChildrenAttribute);
+
+	ComputeBranchGradientAttribute(
+		PointBranchGradientAttribute,
+		BranchParentNumberAttribute,
+		BranchChildrenAttribute,
+		BranchNumberAttribute,
+		BranchPointsAttribute);
+
 	TManagedArrayAccessor<TArray<FVector3f>> BudDirectionAttribute(OutCollection, PV::AttributeNames::BudDirection,
 		PV::GroupNames::PointGroup);
 	TManagedArrayAccessor<float> BranchGradientsAttribute(OutCollection, PV::AttributeNames::BranchGradient,
@@ -390,13 +781,12 @@ void FPVGrafter::DistributeGraft(FManagedArrayCollection& OutCollection, const F
 	int32 ResetPhyllotaxyLeaf = static_cast<int>(LeafPhyllotaxyAttribute[0][6]);
 	float PhyllotaxyOffsetLeaf = LeafPhyllotaxyAttribute[0][7];
 
-	TMap<int32, int32> FoliageNamesOldIDsToNewIDs;
-	if (SourceGraftCollection.HasGroup(PV::GroupNames::FoliageNamesGroup))
-	{
-		FoliageNamesOldIDsToNewIDs = SetFoliageNames(OutCollection, SourceGraftCollection);
-	}
-
 	const int32 NumBranches = BranchFacadeSourceSkeleton.GetElementCount();
+	const TManagedArray<int32>& SourceSkeletonBudNumbers = PointFacadeSourceSkeleton.GetBudNumbersAttribute();
+
+	check(SourceSkeletonBudNumbers.Num() > 0);
+	const int32 MaxBudNumber = *Algo::MaxElement(SourceSkeletonBudNumbers.GetConstArray());
+	int32 CurrentBudNumber = MaxBudNumber + 1;
 
 	// Compute attachment points in terms of normalized distances
 	for (int32 BranchIndex = 0; BranchIndex < NumBranches; ++BranchIndex)
@@ -409,7 +799,7 @@ void FPVGrafter::DistributeGraft(FManagedArrayCollection& OutCollection, const F
 		TArray<FVector3f> AttachmentPointsNormalVectors;
 		TArray<float> AttachmentPointsScales;
 		TArray<int32> AttachmentPointsPointIndices;
-		TArray<int32> BranchPoints = BranchFacadeSourceSkeleton.GetPoints(BranchIndex);
+		const TArray<int32>& BranchPoints = BranchFacadeSourceSkeleton.GetPoints(BranchIndex);
 
 		if (BranchPoints.Num() < 2)
 		{
@@ -499,8 +889,6 @@ void FPVGrafter::DistributeGraft(FManagedArrayCollection& OutCollection, const F
 
 					const FVector3f Position = FMath::Lerp(CurrentPointPosition, NextPointPosition, DistanceAlpha);
 					AttachmentPointsPositions.Add(Position);
-					
-					AttachmentPointsPointIndices.Add(CurrentPointIndex);
 
 					const float AttachmentPointLfr = FMath::GetMappedRangeValueClamped(
 						FVector2f(0.0f, 1.0f),
@@ -508,6 +896,19 @@ void FPVGrafter::DistributeGraft(FManagedArrayCollection& OutCollection, const F
 						NormalizedAttachmentPoints[Idx]);
 					AttachmentPointsLengthsFromRoot.Add(AttachmentPointLfr);
 
+					int32 IndexAdded;
+					// Don't add a new point for a tip instance
+					if (Idx == 0 && FMath::IsNearlyEqual(AttachmentPointLfr, LastPointLengthFromRoot, 0.001f))
+					{
+						IndexAdded = NextPointIndex;
+					}
+					else
+					{
+						IndexAdded = AddPointToCollection(OutCollection, CurrentPointIndex, NextPointIndex, DistanceAlpha, CurrentBudNumber);
+						CurrentBudNumber++;
+					}
+					AttachmentPointsPointIndices.Add(IndexAdded);
+
 					// Compute Scale
 					const float CurrentPointBranchGradient = BranchGradientsAttribute[CurrentPointIndex];
 					const float NextPointBranchGradient = BranchGradientsAttribute[NextPointIndex];
@@ -599,8 +1000,9 @@ void FPVGrafter::DistributeGraft(FManagedArrayCollection& OutCollection, const F
 								FVector3f LastPositionAdded = AttachmentPointsPositions.Last();
 								float LastScaleAdded = AttachmentPointsScales.Last();
 								float LastLengthFromRootAdded = AttachmentPointsLengthsFromRoot.Last();
+								int32 LastPointIndexAdded = AttachmentPointsPointIndices.Last();
 								AttachmentPointsPositions.Add(LastPositionAdded);
-								AttachmentPointsPointIndices.Add(CurrentPointIndex);
+								AttachmentPointsPointIndices.Add(LastPointIndexAdded);
 								AttachmentPointsScales.Add(LastScaleAdded);
 								AttachmentPointsLengthsFromRoot.Add(LastLengthFromRootAdded);
 							}
@@ -630,6 +1032,9 @@ void FPVGrafter::DistributeGraft(FManagedArrayCollection& OutCollection, const F
 
 		for (int32 Idx = 0; Idx < AttachmentPointsLengthsFromRoot.Num(); Idx++)
 		{
+			bool bIsTipInstance =
+				Idx == 0 && FMath::IsNearlyEqual(AttachmentPointsLengthsFromRoot[Idx], LastPointLengthFromRoot, 0.001f);
+
 			FPVGraftAttributes Attributes{
 				.BranchIndex = BranchIndex,
 				.PointIndex = AttachmentPointsPointIndices[Idx],
@@ -637,10 +1042,11 @@ void FPVGrafter::DistributeGraft(FManagedArrayCollection& OutCollection, const F
 				.UpVector = AttachmentPointsUpVectors[Idx],
 				.NormalVector = AttachmentPointsNormalVectors[Idx],
 				.Scale = AttachmentPointsScales[Idx],
-				.LengthFromRoot = AttachmentPointsLengthsFromRoot[Idx]
+				.LengthFromRoot = AttachmentPointsLengthsFromRoot[Idx],
+				.bIsTipInstance = bIsTipInstance
 			};
 
-			SetGraft(OutCollection, SourceSkeletonCollection, SourceGraftCollection, Attributes, FoliageNamesOldIDsToNewIDs);
+			SetGraft(OutCollection, SourceSkeletonCollection, SourceGraftCollection, Attributes);
 		}
 	}
 }
diff --git a/Engine/Plugins/Experimental/ProceduralVegetationEditor/Source/ProceduralVegetation/Private/Implementations/PVGrafter.h b/Engine/Plugins/Experimental/ProceduralVegetationEditor/Source/ProceduralVegetation/Private/Implementations/PVGrafter.h
--- a/Engine/Plugins/Experimental/ProceduralVegetationEditor/Source/ProceduralVegetation/Private/Implementations/PVGrafter.h
+++ b/Engine/Plugins/Experimental/ProceduralVegetationEditor/Source/ProceduralVegetation/Private/Implementations/PVGrafter.h
@@ -3,13 +3,12 @@
 #pragma once
 
 #include "CoreMinimal.h"
+#include "Utils/PVAttributes.h"
 #include "GeometryCollection/ManagedArrayCollection.h"
 
-// #include "PVGrafter.generated.h"
 
 struct FPVFloatRamp;
 
-	
 
 struct FPVGrafter
 {
@@ -46,7 +45,7 @@ struct FPVGrafter
 	{
 		const int32 RandomSeed;
 	};
-	
+
 	struct FPVGraftAttributes
 	{
 		int32 BranchIndex;
@@ -56,18 +55,42 @@ struct FPVGrafter
 		FVector3f NormalVector;
 		float Scale;
 		float LengthFromRoot;
+		bool bIsTipInstance;
 	};
 
+	template <typename T>
+	static TArray<T> LerpArrayElements(const TArray<T>& InArray1, const TArray<T>& InArray2, const float Alpha);
+
+	static int32 AddPointToCollection(FManagedArrayCollection& OutCollection, const int32 CurrentPointIndex, const int32 NextPointIndex,
+	                                  const float Alpha, const int32 BudNumber);
+
 	static FTransform GetPointTransform(const FVector3f AttachmentPoint, const FVector3f UpVector, const FVector3f NormalVector,
 	                                    const float Scale);
 
+	static void RecomputeBranchChildren(
+		PV::FBranchNumberAttributeConstView BranchNumberAttribute,
+		PV::FBranchParentNumberAttributeConstView BranchParentNumberAttribute,
+		PV::FBranchChildrenAttributeView BranchChildrenAttribute);
+
+	static void ComputeBranchGradientAttribute(
+		PV::FPointBranchGradientAttributeView PointBranchGradientAttribute,
+		PV::FBranchParentNumberAttributeConstView BranchParentNumberAttribute,
+		PV::FBranchChildrenAttributeConstView BranchChildrenAttribute,
+		PV::FBranchNumberAttributeConstView BranchNumberAttribute,
+		PV::FBranchPointsAttributeConstView BranchPointsAttribute);
+
+	static void RescaleBranchForGraftingOnTip(FManagedArrayCollection& OutCollection, const int32 BranchIndex, const int32 BranchTipPointIndex,
+	                                          const float GraftTrunkLength);
+
+	template <typename Fn>
+	static void ForEachGraftPointExcludingRoot(const int32 NumPoints, const int32 RootIndex, Fn&& Func);
+
 	static void SetGraft(FManagedArrayCollection& OutCollection, const FManagedArrayCollection& SourceSkeletonCollection,
-	                     const FManagedArrayCollection& SourceGraftCollection, const FPVGraftAttributes& InAttributes,
-	                     const TMap<int32, int32>& FoliageNamesOldIDsToNewIDs);
-	
-	static TMap<int32, int32> SetFoliageNames(FManagedArrayCollection& OutCollection, const FManagedArrayCollection& SourceGraftCollection);
-	
+	                     const FManagedArrayCollection& SourceGraftCollection, const FPVGraftAttributes& InAttributes);
+
+	static bool RemoveFoliageDataIfPresent(FManagedArrayCollection& OutCollection);
+
 	static void DistributeGraft(FManagedArrayCollection& OutCollection, const FManagedArrayCollection& SourceSkeletonCollection,
-	                            const FManagedArrayCollection& SourceGraftCollection, const FDistributionSettings& DistributionSettings,
+	                            FManagedArrayCollection& SourceGraftCollection, const FDistributionSettings& DistributionSettings,
 	                            const FScaleSettings& ScaleSettings, const FVectorSettings& VectorSettings, const FMiscSettings& MiscSettings);
 };
diff --git a/Engine/Plugins/Experimental/ProceduralVegetationEditor/Source/ProceduralVegetation/Private/Nodes/PVGrafterSettings.cpp b/Engine/Plugins/Experimental/ProceduralVegetationEditor/Source/ProceduralVegetation/Private/Nodes/PVGrafterSettings.cpp
--- a/Engine/Plugins/Experimental/ProceduralVegetationEditor/Source/ProceduralVegetation/Private/Nodes/PVGrafterSettings.cpp
+++ b/Engine/Plugins/Experimental/ProceduralVegetationEditor/Source/ProceduralVegetation/Private/Nodes/PVGrafterSettings.cpp
@@ -5,6 +5,7 @@
 #include "ProceduralVegetationModule.h"
 #include "DataTypes/PVGrowthData.h"
 #include "ProceduralVegetationPreset.h"
+#include "Facades/PVAttributesNames.h"
 #include "Facades/PVBranchFacade.h"
 #include "Facades/PVPointFacade.h"
 #include "Implementations/PVCarve.h"
@@ -116,6 +117,16 @@ bool FPVGrafterElement::ExecuteInternal(FPCGContext* InContext) const
 		.RandomSeed = Settings->RandomSeed
 	};
 
+	if (FPVGrafter::RemoveFoliageDataIfPresent(OutCollection))
+	{
+		PCGLog::LogWarningOnGraph(LOCTEXT("FoliageDataRemovedFromSourceSkeleton", "Foliage data removed from source skeleton"), InContext);
+	}
+
+	if (FPVGrafter::RemoveFoliageDataIfPresent(SourceGraftCollection))
+	{
+		PCGLog::LogWarningOnGraph(LOCTEXT("FoliageDataRemovedFromGraft", "Foliage data removed from source graft"), InContext);
+	}
+
 	FPVGrafter::DistributeGraft(OutCollection, SourceSkeletonCollection, SourceGraftCollection,
 		DistributionSettings, ScaleSettings, VectorSettings, MiscSettings);
 
diff --git a/Engine/Plugins/Experimental/ProceduralVegetationEditor/Source/ProceduralVegetation/Public/Facades/PVBranchFacade.h b/Engine/Plugins/Experimental/ProceduralVegetationEditor/Source/ProceduralVegetation/Public/Facades/PVBranchFacade.h
--- a/Engine/Plugins/Experimental/ProceduralVegetationEditor/Source/ProceduralVegetation/Public/Facades/PVBranchFacade.h
+++ b/Engine/Plugins/Experimental/ProceduralVegetationEditor/Source/ProceduralVegetation/Public/Facades/PVBranchFacade.h
@@ -139,6 +139,8 @@ namespace PV::Facades
 			OutValue = DefaultValue;
 			return false;
 		}
+		
+		bool IsTrunk(int32 Index) const;
 
 		void GetParents(int32 Index, TArray<int32>& OutValue) const;
 
@@ -150,6 +152,10 @@ namespace PV::Facades
 		
 		void GetBranchSourceBudNumber(int32 Index, int32& OutValue) const;
 		
+		const TManagedArray<int32>& GetBranchSourceBudNumbers() const;
+		
+		void SetBranchSourceBudNumbersFromIndex(const TArray<int32>& InputArray, const int32 StartIndex);
+		
 		void GetBranchHierarchyNumber(int32 Index, int32& OutValue) const;
 		
 		void GetBranchParentNumber(int32 Index, int32& OutValue) const;
diff --git a/Engine/Plugins/Experimental/ProceduralVegetationEditor/Source/ProceduralVegetation/Public/Facades/PVPointFacade.h b/Engine/Plugins/Experimental/ProceduralVegetationEditor/Source/ProceduralVegetation/Public/Facades/PVPointFacade.h
--- a/Engine/Plugins/Experimental/ProceduralVegetationEditor/Source/ProceduralVegetation/Public/Facades/PVPointFacade.h
+++ b/Engine/Plugins/Experimental/ProceduralVegetationEditor/Source/ProceduralVegetation/Public/Facades/PVPointFacade.h
@@ -32,16 +32,24 @@ Accessor.ModifyAt(i, InputArray[i]); \
 	public:
 		FPointFacade(FManagedArrayCollection& InCollection);
 		FPointFacade(const FManagedArrayCollection& InCollection);
+		
+		bool IsConst() const { return Collection == nullptr; }
 
 		bool IsValid() const;
 
 		virtual int32 GetElementCount() const override;
 
 		const FVector3f& GetPosition(int32 Index) const;
+		
+		void SetPosition(int32 Index, const FVector3f& InPosition);
 
 		float GetLengthFromRoot(int32 Index) const;
+		
+		void SetLengthFromRoot(int32 Index, float InLength);
 
 		float GetLengthFromSeed(int32 Index) const;
+		
+		void SetLengthFromSeed(int32 Index, float InLength);
 
 		float GetPointScaleGradient(int32 Index) const;
 
@@ -61,27 +69,33 @@ Accessor.ModifyAt(i, InputArray[i]); \
 
 		const TArray<float>& GetBudLightDetected(int32 Index) const;
 
-		void SetBudLightDetected(int32 Index, TArray<float> InBudLightDetected);
+		void SetBudLightDetected(int32 Index, const TArray<float>& InBudLightDetected);
 
 		const TArray<int>& GetBudDevelopment(int32 Index) const;
 		
-		void SetBudDevelopment(int32 Index, TArray<int> InBudDevelopment);
+		void SetBudDevelopment(int32 Index, const TArray<int>& InBudDevelopment);
 
 		int32 GetBudGeneratation(int32 Index) const;
 
 		int32 GetBudAge(int32 Index) const;
+		
+		int32 GetMaxBudAge() const;
+		
+		void AddToBudAgeForAllPoints(const int32 Age);
 
 		int32 GetBudBranchAge(int32 Index) const;
 
 		const TArray<FVector3f>& GetBudDirection(int32 Index) const;
 
-		void SetBudDirections(int32 Index, TArray<FVector3f> InBudDirections);
+		void SetBudDirections(int32 Index, const TArray<FVector3f>& InBudDirections);
 
 		const TArray<float>& GetBudHormoneLevels(int32 Index) const;
 
-		void SetBudHormoneLevels(int32 Index, TArray<float> InBudHormoneLevels);
+		void SetBudHormoneLevels(int32 Index, const TArray<float>& InBudHormoneLevels);
 
 		const TArray<float>& GetBudLateralMeristem(int32 Index) const;
+		
+		void SetBudLateralMeristem(int32 Index, const TArray<float>& InBudLateralMeristem);
 
 		float GetBranchGradient(int32 Index) const;
 
@@ -163,6 +177,8 @@ Accessor.ModifyAt(i, InputArray[i]); \
 		
 		void SetPositionsFromIndex(const TArray<FVector3f>& InputArray, const int32 StartIndex);
 		
+		void SetPointScalesFromIndex(const TArray<float>& InputArray, const int32 StartIndex);
+		
 		void SetLFRsFromIndex(const TArray<float>& InputArray, const int32 StartIndex);
 		
 		const TManagedArray<float>& GetLengthFromSeedsArray() const;
@@ -211,6 +227,10 @@ Accessor.ModifyAt(i, InputArray[i]); \
 		void GetPlantGradients(int32 Index, float& OutValue) const;
 
 		void GetBudNumbers(int32 Index, int32& OutValue) const;
+		
+		const TManagedArray<int32>& GetBudNumbersAttribute() const;
+		
+		void SetBudNumbersFromIndex(const TArray<int32>& InputArray, const int32 StartIndex);
 
 		void GetBudLightDetected(int32 Index, TArray<float>& OutValue) const;
 		
@@ -224,6 +244,8 @@ Accessor.ModifyAt(i, InputArray[i]); \
 
 		void GetBudStatus(int32 Index, TArray<int>& OutValue) const;
 		
+		void SetBudStatus(int32 Index, const TArray<int>& InValue);
+		
 		void GetBudLateralMeristem(int32 Index, TArray<float>& OutValue) const;
 		
 		void GetSeedPScale(int32 Index, float& OutValue) const;
@@ -231,6 +253,8 @@ Accessor.ModifyAt(i, InputArray[i]); \
 		void GetSeedPScaleRatio(int32 Index, float& OutValue) const;
 
 	private:
+		FManagedArrayCollection* Collection = nullptr;
+		TManagedArrayAccessor<TArray<int32>> BranchPoints;
 		TManagedArrayAccessor<FVector3f> Positions;
 		TManagedArrayAccessor<float> LengthFromRoot;
 		TManagedArrayAccessor<float> PointScaleGradient;
