[PVE] Foliage palette refactored to make it work with the grower node.
- input pin and override functionality removed
- User can now add/remove the foliage items
- Foliage palette now outputs a new data type for FoliageData
- New pin for assigning foliage palette added in the Foliage distributor
- Broke the link of the Foliage palette from the MegaPlant data asset
- Moved Node pin labels to a common place.
- Updated the node links to handle backwards compatibility.

Major changes in files:
- FoliagePaletteSettings - Added the Foliage information in the collection and handled deprecation
- FoliageDistributorSettings - Handled incoming foliage info coming from a separate pin.

#Synced-CL 50858990

diff --git a/Engine/Plugins/Experimental/ProceduralVegetationEditor/Source/ProceduralVegetation/Private/DataTypes/PVFoliageData.cpp b/Engine/Plugins/Experimental/ProceduralVegetationEditor/Source/ProceduralVegetation/Private/DataTypes/PVFoliageData.cpp
new file mode 100644
--- /dev/null
+++ b/Engine/Plugins/Experimental/ProceduralVegetationEditor/Source/ProceduralVegetation/Private/DataTypes/PVFoliageData.cpp
@@ -0,0 +1,5 @@
+// Copyright Epic Games, Inc. All Rights Reserved.
+
+#include "DataTypes/PVFoliageData.h"
+
+PCG_DEFINE_TYPE_INFO(FPVDataTypeInfoFoliage, UPVFoliageData);
\ No newline at end of file
diff --git a/Engine/Plugins/Experimental/ProceduralVegetationEditor/Source/ProceduralVegetation/Private/Implementations/PVFoliage.cpp b/Engine/Plugins/Experimental/ProceduralVegetationEditor/Source/ProceduralVegetation/Private/Implementations/PVFoliage.cpp
--- a/Engine/Plugins/Experimental/ProceduralVegetationEditor/Source/ProceduralVegetation/Private/Implementations/PVFoliage.cpp
+++ b/Engine/Plugins/Experimental/ProceduralVegetationEditor/Source/ProceduralVegetation/Private/Implementations/PVFoliage.cpp
@@ -69,27 +69,30 @@ struct FFoliageWeightInfo
 	float Weight = 0;
 };
  
-void FPVFoliage::DistributeFoliage(FManagedArrayCollection& OutCollection, const FManagedArrayCollection& SourceCollection,
+void FPVFoliage::DistributeFoliage(FManagedArrayCollection& OutCollection, const FManagedArrayCollection& SourceCollection, const FManagedArrayCollection& FoliageCollection,
                                    const FDistributionSettings& DistributionSettings, const FScaleSettings& ScaleSettings,
                                    const FVectorSettings& VectorSettings, const FPhyllotaxySettings& PhyllotaxySettings,
                                    const FMiscSettings& MiscSettings, const FPVFoliageDistributionConditions& DistributionConditions)
 {
 	PV::Facades::FFoliageFacade FoliageFacadeOutput(OutCollection);
+	PV::Facades::FFoliageFacade FoliagePaletteFacade(FoliageCollection);
  
-	if (const int32 NumFoliageMeshes = FoliageFacadeOutput.NumFoliageInfo(); NumFoliageMeshes == 0)
+	if (const int32 NumFoliageMeshes = FoliagePaletteFacade.NumFoliageInfo(); NumFoliageMeshes == 0)
 	{
 		UE_LOG(LogProceduralVegetation, Warning, TEXT("There are no foliage meshes available in input for distribution"));
 		return;
 	}
+
+	FoliageFacadeOutput.SetFoliageInfos(FoliagePaletteFacade.GetFoliageInfos());
  
 	FRandomStream RandomStream;
 	RandomStream.Initialize(MiscSettings.RandomSeed);
  
 	PV::Facades::FBranchFacade BranchFacadeSource(SourceCollection);
 	PV::Facades::FPointFacade PointFacadeSource(SourceCollection);
- 
+
 	TManagedArrayAccessor<TArray<FVector3f>> BudDirectionAttribute(OutCollection, PV::AttributeNames::BudDirection,
-		PV::GroupNames::PointGroup);
+			PV::GroupNames::PointGroup);
 	
 	TManagedArrayAccessor<float> BranchGradientsAttribute(OutCollection, PV::AttributeNames::BranchGradient,
 		PV::GroupNames::PointGroup);
@@ -102,12 +105,12 @@ void FPVFoliage::DistributeFoliage(FManagedArrayCollection& OutCollection, const
  
 	TManagedArrayAccessor<TArray<int>> BudDevelopment(OutCollection, PV::AttributeNames::BudDevelopment,
 		PV::GroupNames::PointGroup);
- 
-	float PhyllotaxyFormationLeaf = LeafPhyllotaxyAttribute[0][1];
-	float MinBudsLeaf = LeafPhyllotaxyAttribute[0][3];
-	float MaxBudsLeaf = LeafPhyllotaxyAttribute[0][4];
-	int32 ResetPhyllotaxyLeaf = static_cast<int>(LeafPhyllotaxyAttribute[0][6]);
-	float PhyllotaxyOffsetLeaf = LeafPhyllotaxyAttribute[0][7];
+	
+	float PhyllotaxyFormationLeaf = 180;
+	float MinBudsLeaf = 1;
+	float MaxBudsLeaf = 1;
+	int32 ResetPhyllotaxyLeaf = 0;
+	float PhyllotaxyOffsetLeaf = 0;
  
 	if (PhyllotaxySettings.OverridePhyllotaxy)
 	{
@@ -148,6 +151,14 @@ void FPVFoliage::DistributeFoliage(FManagedArrayCollection& OutCollection, const
 			}
 		}
 	}
+	else
+	{
+		PhyllotaxyFormationLeaf = LeafPhyllotaxyAttribute[0][1];
+		MinBudsLeaf = LeafPhyllotaxyAttribute[0][3];
+		MaxBudsLeaf = LeafPhyllotaxyAttribute[0][4];
+		ResetPhyllotaxyLeaf = static_cast<int>(LeafPhyllotaxyAttribute[0][6]);
+		PhyllotaxyOffsetLeaf = LeafPhyllotaxyAttribute[0][7];
+	}
  
 	const int32 NumBranches = BranchFacadeSource.GetElementCount();
  
@@ -227,7 +238,7 @@ void FPVFoliage::DistributeFoliage(FManagedArrayCollection& OutCollection, const
 		OutCollection.EmptyGroup(PV::GroupNames::FoliageGroup);
  
 		// Collect non-empty mesh name indices
-		const TArray<FString>& FoliageMeshNames = FoliageFacadeOutput.GetFoliageNames();
+		const TArray<FString>& FoliageMeshNames = FoliagePaletteFacade.GetFoliageNames();
 		TArray<int32> NonEmptyFoliageMeshNameIndices;
 		NonEmptyFoliageMeshNameIndices.Reserve(FoliageMeshNames.Num());
  
@@ -255,7 +266,7 @@ void FPVFoliage::DistributeFoliage(FManagedArrayCollection& OutCollection, const
 			TArray<FVector3f> AttachmentPointsNormalVectors;
 			TArray<float> AttachmentPointsScales;
 			
-			const int32 NumFoliageInfo = FoliageFacadeOutput.NumFoliageInfo();
+			const int32 NumFoliageInfo = FoliagePaletteFacade.NumFoliageInfo();
 			
 			TArray<int32> BranchPoints = BranchFacadeSource.GetPoints(BranchIndex);
  
@@ -541,7 +552,7 @@ void FPVFoliage::DistributeFoliage(FManagedArrayCollection& OutCollection, const
 			TManagedArrayAccessor<FVector3f> PivotPointsAttribute(OutCollection, PV::AttributeNames::FoliagePivotPoint,
 				PV::GroupNames::FoliageGroup);
  
-			const TArray<FPVFoliageInfo> FoliageInfos = FoliageFacadeOutput.GetFoliageInfos();
+			const TArray<FPVFoliageInfo> FoliageInfos = FoliagePaletteFacade.GetFoliageInfos();
 			
  
 			const int32 BranchNumber = BranchFacadeSource.GetBranchNumber(BranchIndex);
@@ -667,13 +678,6 @@ void FPVFoliage::DistributeFoliage(FManagedArrayCollection& OutCollection, const
 				});
 			}
 		}
- 
-		// TArray<float> Scales;
-		//
-		// for (int32 i = 0; i < FoliageFacadeOutput.NumFoliageEntries(); ++i)
-		// {
-		// 	FoliageFacadeOutput.GetFoliageEntry(i).Scale;
-		// }
 		
 		// Re-calculate Foliage IDs for each branch
 		TMap<int32, TArray<int32>> BranchIndexToFoliageIDs;
diff --git a/Engine/Plugins/Experimental/ProceduralVegetationEditor/Source/ProceduralVegetation/Private/Implementations/PVFoliage.h b/Engine/Plugins/Experimental/ProceduralVegetationEditor/Source/ProceduralVegetation/Private/Implementations/PVFoliage.h
--- a/Engine/Plugins/Experimental/ProceduralVegetationEditor/Source/ProceduralVegetation/Private/Implementations/PVFoliage.h
+++ b/Engine/Plugins/Experimental/ProceduralVegetationEditor/Source/ProceduralVegetation/Private/Implementations/PVFoliage.h
@@ -351,7 +351,7 @@ struct FPVFoliageInfo
  
 struct FPVFoliage
 {
-	static void DistributeFoliage(FManagedArrayCollection& OutCollection, const FManagedArrayCollection& SourceCollection,
+	static void DistributeFoliage(FManagedArrayCollection& OutCollection, const FManagedArrayCollection& SourceCollection, const FManagedArrayCollection& FoliageCollection,  
 	                              const FDistributionSettings& DistributionSettings, const FScaleSettings& ScaleSettings,
 	                              const FVectorSettings& VectorSettings, const FPhyllotaxySettings& PhyllotaxySettings,
 	                              const FMiscSettings& MiscSettings, const FPVFoliageDistributionConditions& DistributionConditions);
diff --git a/Engine/Plugins/Experimental/ProceduralVegetationEditor/Source/ProceduralVegetation/Private/Nodes/PVFoliageDistributorSettings.cpp b/Engine/Plugins/Experimental/ProceduralVegetationEditor/Source/ProceduralVegetation/Private/Nodes/PVFoliageDistributorSettings.cpp
--- a/Engine/Plugins/Experimental/ProceduralVegetationEditor/Source/ProceduralVegetation/Private/Nodes/PVFoliageDistributorSettings.cpp
+++ b/Engine/Plugins/Experimental/ProceduralVegetationEditor/Source/ProceduralVegetation/Private/Nodes/PVFoliageDistributorSettings.cpp
@@ -1,10 +1,12 @@
 ﻿// Copyright Epic Games, Inc. All Rights Reserved.
  
 #include "PVFoliageDistributorSettings.h"
+#include "PVCommon.h"
 #include "DataTypes/PVData.h"
 #include "DataTypes/PVMeshData.h"
 #include "DataTypes/PVFoliageMeshData.h"
 #include "Data/PCGBasePointData.h"
+#include "DataTypes/PVFoliageData.h"
 #include "Facades/PVFoliageConditionFacade.h"
 #include "GeometryCollection/ManagedArrayCollection.h"
 #include "Helpers/PCGHelpers.h"
@@ -75,7 +77,18 @@ UPVFoliageDistributorSettings::UPVFoliageDistributorSettings()
 		RichCurve->SetKeyInterpMode(Handle5, RCIM_Linear);
 	}
 }
- 
+
+TArray<FPCGPinProperties> UPVFoliageDistributorSettings::InputPinProperties() const
+{
+	TArray<FPCGPinProperties> Properties = Super::InputPinProperties();
+	
+	FPCGPinProperties& PhyllotaxyPin = Properties.Emplace_GetRef(PV::Pins::FoliageDistributorFoliageInputLabel, FPCGDataTypeIdentifier{ FPVDataTypeInfoFoliage::AsId()});
+	PhyllotaxyPin.SetAllowMultipleConnections(false);
+	PhyllotaxyPin.SetRequiredPin();
+	
+	return Properties;
+}
+
 FPCGDataTypeIdentifier UPVFoliageDistributorSettings::GetInputPinTypeIdentifier() const
 {
 	return FPCGDataTypeIdentifier{ FPVDataTypeInfoMesh::AsId() };
@@ -102,7 +115,17 @@ bool FPVFoliageDistributorElement::ExecuteInternal(FPCGContext* InContext) const
  
 	UPVFoliageDistributorSettings* MutableSettings = const_cast<UPVFoliageDistributorSettings*>(Settings);
 	check(MutableSettings);
- 
+
+	FManagedArrayCollection FoliageCollection;
+	
+	for (const FPCGTaggedData& Input : InContext->InputData.GetInputsByPin(PV::Pins::FoliageDistributorFoliageInputLabel))
+	{
+		if(const UPVFoliageData* FoliageData = Cast<UPVFoliageData>(Input.Data))
+		{
+			FoliageCollection = FoliageData->GetCollection();
+		}		
+	}
+	
 	for (const FPCGTaggedData& Input : InContext->InputData.GetInputsByPin(PCGPinConstants::DefaultInputLabel))
 	{
 		if(const UPVMeshData* InputData = Cast<UPVMeshData>(Input.Data))
@@ -190,7 +213,7 @@ bool FPVFoliageDistributorElement::ExecuteInternal(FPCGContext* InContext) const
 				MutableSettings->CustomConditions = MutableSettings->Conditions;
 			}	
  
-			FPVFoliage::DistributeFoliage(Collection, Collection, DistributionSettings, ScaleSettings, VectorSettings,
+			FPVFoliage::DistributeFoliage(Collection, Collection, FoliageCollection, DistributionSettings, ScaleSettings, VectorSettings,
 												 PhyllotaxySettings, MiscSettings, Settings->Conditions);
  
 			OutManagedArrayCollectionData->Initialize(MoveTemp(Collection));
diff --git a/Engine/Plugins/Experimental/ProceduralVegetationEditor/Source/ProceduralVegetation/Private/Nodes/PVFoliageDistributorSettings.h b/Engine/Plugins/Experimental/ProceduralVegetationEditor/Source/ProceduralVegetation/Private/Nodes/PVFoliageDistributorSettings.h
--- a/Engine/Plugins/Experimental/ProceduralVegetationEditor/Source/ProceduralVegetation/Private/Nodes/PVFoliageDistributorSettings.h
+++ b/Engine/Plugins/Experimental/ProceduralVegetationEditor/Source/ProceduralVegetation/Private/Nodes/PVFoliageDistributorSettings.h
@@ -34,7 +34,7 @@ public:
 	#endif
 
 protected:
-
+	virtual TArray<FPCGPinProperties> InputPinProperties() const override;
 	virtual FPCGDataTypeIdentifier GetInputPinTypeIdentifier() const override;
 	virtual FPCGDataTypeIdentifier GetOutputPinTypeIdentifier() const override;
 
@@ -46,7 +46,7 @@ private:
 	float EthyleneThreshold = 0.85f;
 
 	UPROPERTY(EditAnywhere, Category="Distribution Settings", meta=(Tooltip="Override default placement with custom distribution rules.\n\nOverides the distribution settings loaded from the procedural vegetation preset."))
-	bool OverrideDistribution = false;
+	bool OverrideDistribution = true;
 
 	UPROPERTY(EditAnywhere, Category="Distribution Settings", meta=(EditCondition="OverrideDistribution", ClampMin=0.0f, ClampMax=1.0f, UIMin=0.0f, UIMax=1.0f, Tooltip="Minimum distance between placed instances.\n\nHigher spacing yields a more open distribution; lower spacing allows denser placement."))
 	float InstanceSpacing = 0.3f;
@@ -97,7 +97,7 @@ private:
 	float AxilAngleRampEffect = 0.0f;
 
 	UPROPERTY(EditAnywhere, Category="Phyllotaxy Settings", meta=(EditCondition="OverrideDistribution", Tooltip="Use custom phyllotaxy settings instead of defaults.\n\nOveride the phyllotaxy settings loaded from the procedural vegetation preset."))
-	bool OverridePhyllotaxy = false;
+	bool OverridePhyllotaxy = true;
 
 	UPROPERTY(EditAnywhere, Category="Phyllotaxy Settings", meta=(EditCondition="OverrideDistribution", Tooltip="Pattern for arranging buds/leaves around the stem.\n\nSelects the phyllotactic pattern used for placement (e.g., Alternate, Opposite, Whorled, Spiral). This defines the sequence and angular spacing of buds/leaves around the stem and along its length."))
 	EPhyllotaxyType PhyllotaxyType = EPhyllotaxyType::Alternate;
diff --git a/Engine/Plugins/Experimental/ProceduralVegetationEditor/Source/ProceduralVegetation/Private/Nodes/PVFoliagePaletteSettings.cpp b/Engine/Plugins/Experimental/ProceduralVegetationEditor/Source/ProceduralVegetation/Private/Nodes/PVFoliagePaletteSettings.cpp
--- a/Engine/Plugins/Experimental/ProceduralVegetationEditor/Source/ProceduralVegetation/Private/Nodes/PVFoliagePaletteSettings.cpp
+++ b/Engine/Plugins/Experimental/ProceduralVegetationEditor/Source/ProceduralVegetation/Private/Nodes/PVFoliagePaletteSettings.cpp
@@ -1,12 +1,15 @@
 ﻿// Copyright Epic Games, Inc. All Rights Reserved.
  
 #include "PVFoliagePaletteSettings.h"
- 
+
+#include "PCGEdge.h"
 #include "ProceduralVegetationModule.h"
+#include "PVCommon.h"
 #include "DataTypes/PVData.h"
 #include "DataTypes/PVMeshData.h"
 #include "Data/PCGBasePointData.h"
 #include "Data/PCGStaticMeshResourceData.h"
+#include "DataTypes/PVFoliageData.h"
 #include "Facades/PVFoliageFacade.h"
 #include "GeometryCollection/ManagedArrayCollection.h"
 #include "Helpers/PCGHelpers.h"
@@ -26,23 +29,11 @@ FText UPVFoliagePaletteSettings::GetNodeTooltipText() const
 { 
 	return LOCTEXT("NodeTooltip", 
 		"Allows the user to visualize and optionally change the foliage meshes. "
-		"Foliage mesh data is loaded from the Procedrual Vegetation preset. "
-		"The user can replace or remove foliage meshes. The placement of folliage meshes is driven by the procedural vegetaiton preset and cannot be modifed"
+		"The user can add or remove foliage meshes. The placement of foliage meshes is driven by the procedural vegetation preset and cannot be modified"
 		"\n\nPress Ctrl + L to lock/unlock node output"
 	);
 }
  
-void UPVFoliagePaletteSettings::PostEditChangeProperty(FPropertyChangedEvent& PropertyChangedEvent)
-{
-	Super::PostEditChangeProperty(PropertyChangedEvent);
- 
-	if (PropertyChangedEvent.Property
-		&& PropertyChangedEvent.Property->GetFName() == GET_MEMBER_NAME_CHECKED(UPVFoliagePaletteSettings, bOverrideFoliage))
-	{
-		DirtyCache();
-	}
-}
- 
 void UPVFoliagePaletteSettings::PostEditChangeChainProperty(FPropertyChangedChainEvent& PropertyChangedEvent)
 {
 	Super::PostEditChangeChainProperty(PropertyChangedEvent);
@@ -72,7 +63,8 @@ void UPVFoliagePaletteSettings::PostEditChangeChainProperty(FPropertyChangedChai
 					PV::Analytics::SendFoliageMeshChangeEvent(MeshPath);	
 				}
 			}
- 
+
+			PRAGMA_DISABLE_DEPRECATION_WARNINGS
 			// For backwards compatibility 
 			FoliageMeshes.Empty();
  
@@ -80,20 +72,77 @@ void UPVFoliagePaletteSettings::PostEditChangeChainProperty(FPropertyChangedChai
 			{
 				FoliageMeshes.Add(FoliageInfos[i].Mesh);
 			}
+			PRAGMA_ENABLE_DEPRECATION_WARNINGS
+			
 		}
 	}
 }
- 
-#endif
- 
-FPCGDataTypeIdentifier UPVFoliagePaletteSettings::GetInputPinTypeIdentifier() const
+
+void UPVFoliagePaletteSettings::ApplyDeprecation(UPCGNode* InOutNode)
+{
+	Super::ApplyDeprecation(InOutNode);
+	
+	for (const auto& Connection : PinConnectionsToUpdate)
+	{
+		if (Connection.Node.IsValid() && Connection.TargetNode.IsValid())
+		{
+			Connection.Node->AddEdgeTo(Connection.FromPinLabel, Connection.TargetNode.Get(), Connection.ToPinLabel);
+		}
+	}
+}
+
+void UPVFoliagePaletteSettings::ApplyDeprecationBeforeUpdatePins(UPCGNode* InOutNode, TArray<TObjectPtr<UPCGPin>>& InputPins,
+                                                                 TArray<TObjectPtr<UPCGPin>>& OutputPins)
+{
+	Super::ApplyDeprecationBeforeUpdatePins(InOutNode, InputPins, OutputPins);
+
+	if (InputPins.Num() > 0 && OutputPins.Num() > 0)
+	{
+		auto InputPin = InputPins[0];
+		auto OutputPin = OutputPins[0];
+
+		if (InputPin->Edges.Num() > 0 && OutputPin->Edges.Num() > 0)
+		{
+			auto InputEdge = InputPin->Edges[0];
+			auto OutputEdge = OutputPin->Edges[0];
+			
+			{
+				// Foliage palette connection to Foliage distributor
+
+				FPinConnectionInfo ConnectionInfo;
+				ConnectionInfo.Node = MakeWeakObjectPtr(InOutNode);
+				ConnectionInfo.TargetNode = MakeWeakObjectPtr(OutputEdge->OutputPin->Node);
+				ConnectionInfo.FromPinLabel = OutputEdge->GetInputPinLabel();
+				ConnectionInfo.ToPinLabel = PV::Pins::FoliageDistributorFoliageInputLabel;
+
+				PinConnectionsToUpdate.Add(ConnectionInfo);
+			}
+
+			{
+				FPinConnectionInfo ConnectionInfo;
+				
+				// Foliage palette Input pin gets connection to distributor
+				ConnectionInfo.TargetNode = MakeWeakObjectPtr(OutputEdge->OutputPin->Node);
+				ConnectionInfo.Node = MakeWeakObjectPtr(InputEdge->InputPin->Node);
+				ConnectionInfo.FromPinLabel = InputEdge->GetInputPinLabel();
+				ConnectionInfo.ToPinLabel = OutputEdge->GetOutputPinLabel();
+				
+				PinConnectionsToUpdate.Add(ConnectionInfo);
+			}
+		}
+	}
+}
+
+TArray<FPCGPinProperties> UPVFoliagePaletteSettings::InputPinProperties() const
 {
-	return FPCGDataTypeIdentifier{ FPVDataTypeInfoMesh::AsId() };
+	return TArray<FPCGPinProperties>();
 }
+
+#endif
  
 FPCGDataTypeIdentifier UPVFoliagePaletteSettings::GetOutputPinTypeIdentifier() const
 {
-	return FPCGDataTypeIdentifier{ FPVDataTypeInfoMesh::AsId() };
+	return FPCGDataTypeIdentifier{ FPVDataTypeInfoFoliage::AsId() };
 }
  
 FPCGElementPtr UPVFoliagePaletteSettings::CreateElement() const
@@ -104,6 +153,8 @@ FPCGElementPtr UPVFoliagePaletteSettings::CreateElement() const
 void UPVFoliagePaletteSettings::PostLoad()
 {
 	Super::PostLoad();
+
+	PRAGMA_DISABLE_DEPRECATION_WARNINGS
  
 	if (bOverrideFoliage)
 	{
@@ -115,6 +166,9 @@ void UPVFoliagePaletteSettings::PostLoad()
 			}
 		}
 	}
+
+	PRAGMA_ENABLE_DEPRECATION_WARNINGS
+	
 }
  
 bool FPVFoliageElement::ExecuteInternal(FPCGContext* InContext) const
@@ -125,74 +179,17 @@ bool FPVFoliageElement::ExecuteInternal(FPCGContext* InContext) const
  
 	const UPVFoliagePaletteSettings* Settings = InContext->GetInputSettings<UPVFoliagePaletteSettings>();
 	check(Settings);
+
+	FManagedArrayCollection FoliageCollection;
  
-	UPVFoliagePaletteSettings* const MutableSettings = const_cast<UPVFoliagePaletteSettings*>(Settings);
-	check(MutableSettings);
- 
-	const TArray<FPCGTaggedData>& Inputs = InContext->InputData.GetInputsByPin(PCGPinConstants::DefaultInputLabel);
-	if (!Inputs.IsEmpty())
-	{
-		if (const UPVMeshData* InputData = Cast<UPVMeshData>(Inputs[0].Data))
-		{
-			FManagedArrayCollection Collection = InputData->GetCollection();
- 
-			UPVMeshData* OutManagedArrayCollectionData = FPCGContext::NewObject_AnyThread<UPVMeshData>(InContext);
- 
-			PV::Facades::FFoliageFacade Facade(Collection);
- 
-			if (!Settings->bOverrideFoliage)
-			{
-				MutableSettings->FoliageInfos.Empty();
-				MutableSettings->bOverrideFoliage = false;
-				for (int32 FoliageIndex = 0; FoliageIndex < Facade.NumFoliageInfo(); ++FoliageIndex)
-				{
-					FPVFoliageInfo FoliageInfo = Facade.GetFoliageInfo(FoliageIndex);
-					MutableSettings->FoliageInfos.Emplace(FoliageInfo);
-				}
-#if WITH_EDITOR
-				MutableSettings->PostEditChange();
-#endif
-			}
- 
-			const int32 IterationLimit = FMath::Min(Settings->FoliageInfos.Num(), Facade.NumFoliageInfo());
-			TArray<FPVFoliageInfo> FoliageInfos;
-			
-			for (int32 i = 0; i < IterationLimit; ++i)
-			{
-				FoliageInfos.Emplace(Settings->FoliageInfos[i]);
-			}
- 
-			if (Settings->FoliageInfos.Num() > Facade.NumFoliageInfo())
-			{
-				UE_LOG(LogProceduralVegetation, Warning,
-				       TEXT("Num of Foliage meshes specified is greater than the number of available entries, ignoring the last %d entries"),
-				       Settings->FoliageInfos.Num() - Facade.NumFoliageInfo());
-			}
-			else if (Settings->FoliageInfos.Num() < Facade.NumFoliageInfo())
-			{
-				if (Settings->FoliageInfos.Num() > 0)
-				{
-					UE_LOG(LogProceduralVegetation, Warning,
-					       TEXT("Num of Foliage meshes specified is less than the number of available entries, repeating the last entry"));
+	UPVFoliageData* OutManagedArrayCollectionData = FPCGContext::NewObject_AnyThread<UPVFoliageData>(InContext);
  
-					for (int32 i = Settings->FoliageInfos.Num(); i < Facade.NumFoliageInfo(); ++i)
-					{
-						FoliageInfos.Emplace(Settings->FoliageInfos.Last());
-					}
-				}
-			}
+	PV::Facades::FFoliageFacade Facade(FoliageCollection);
  
-			Facade.SetFoliageInfos(FoliageInfos);
+	Facade.SetFoliageInfos(Settings->FoliageInfos);
  
-			OutManagedArrayCollectionData->Initialize(MoveTemp(Collection));
-			InContext->OutputData.TaggedData.Emplace(OutManagedArrayCollectionData);
-		}
-		else
-		{
-			PCGLog::InputOutput::LogInvalidInputDataError(InContext);
-			return true;
-		}
-	}
+	OutManagedArrayCollectionData->Initialize(MoveTemp(FoliageCollection));
+	InContext->OutputData.TaggedData.Emplace(OutManagedArrayCollectionData);
  
 	return true;
 }
diff --git a/Engine/Plugins/Experimental/ProceduralVegetationEditor/Source/ProceduralVegetation/Private/Nodes/PVFoliagePaletteSettings.h b/Engine/Plugins/Experimental/ProceduralVegetationEditor/Source/ProceduralVegetation/Private/Nodes/PVFoliagePaletteSettings.h
--- a/Engine/Plugins/Experimental/ProceduralVegetationEditor/Source/ProceduralVegetation/Private/Nodes/PVFoliagePaletteSettings.h
+++ b/Engine/Plugins/Experimental/ProceduralVegetationEditor/Source/ProceduralVegetation/Private/Nodes/PVFoliagePaletteSettings.h
@@ -5,12 +5,20 @@
 #include "Nodes/PVBaseSettings.h"
 #include "DataTypes/PVData.h"
 #include "PVFoliagePaletteSettings.generated.h"
- 
+
+struct FPinConnectionInfo
+{
+	TWeakObjectPtr<UPCGNode> Node = nullptr;
+	TWeakObjectPtr<UPCGNode> TargetNode = nullptr;
+	FName FromPinLabel = NAME_None;
+	FName ToPinLabel = NAME_None;
+};
+
 UCLASS(BlueprintType, ClassGroup = (Procedural))
 class UPVFoliagePaletteSettings : public UPVBaseSettings
 {
 	GENERATED_BODY()
- 
+
 public:
 #if WITH_EDITOR
 	virtual FName GetDefaultNodeName() const override { return FName(TEXT("ProceduralVegetationFoliage")); }
@@ -20,29 +28,35 @@ public:
  
 	virtual TArray<EPVRenderType> GetDefaultRenderType() const override { return TArray{EPVRenderType::FoliageGrid}; }
 	virtual TArray<EPVRenderType> GetSupportedRenderTypes() const override { return TArray{EPVRenderType::FoliageGrid, EPVRenderType::Foliage, EPVRenderType::FoliageAttachments, EPVRenderType::Mesh}; }
-	
-	virtual void PostEditChangeProperty(FPropertyChangedEvent& PropertyChangedEvent) override;
 	virtual void PostEditChangeChainProperty(FPropertyChangedChainEvent& PropertyChangedEvent) override;
+
+	virtual void ApplyDeprecation(UPCGNode* InOutNode) override;
+	virtual void ApplyDeprecationBeforeUpdatePins(UPCGNode* InOutNode, TArray<TObjectPtr<UPCGPin>>& InInputPins, TArray<TObjectPtr<UPCGPin>>& InOutputPins) override;
 #endif
  
 protected:
-	virtual FPCGDataTypeIdentifier GetInputPinTypeIdentifier() const override;
+	virtual TArray<FPCGPinProperties> InputPinProperties() const override;
 	virtual FPCGDataTypeIdentifier GetOutputPinTypeIdentifier() const override;
 	virtual FPCGElementPtr CreateElement() const override;
  
 public:
 	virtual void PostLoad() override;
- 
-	UPROPERTY(EditAnywhere, Category="Meshes", meta=(PinHiddenByDefault, InlineEditConditionToggle))
-	bool bOverrideFoliage = false;
- 
-	UPROPERTY()
+
+	UE_DEPRECATED(5.8, "FoliageInfos can now be updated without Override foliage")
+	UPROPERTY(EditAnywhere, Category="Meshes", meta=(PinHiddenByDefault, InlineEditConditionToggle, DeprecatedProperty, DeprecationMessage = "FoliageInfos can now be updated without Override foliage"))
+	bool bOverrideFoliage = true;
+
+	UE_DEPRECATED(5.8, "Use FoliageInfos instead")
+	UPROPERTY(meta = (DeprecatedProperty, DeprecationMessage = "Use FoliageInfos instead"))
 	TArray<TSoftObjectPtr<UObject>> FoliageMeshes;
-	
-	UPROPERTY(EditAnywhere, Category = "Meshes", EditFixedSize,
-		Meta = (NoResetToDefault, EditCondition="bOverrideFoliage", AllowedClasses = "/Script/Engine.StaticMesh,/Script/Engine.SkeletalMesh", ShowOnlyInnerProperties, FullyExpand="true", Tooltip=
+
+	UPROPERTY(EditAnywhere, Category = "Meshes",
+		Meta = (AllowedClasses = "/Script/Engine.StaticMesh,/Script/Engine.SkeletalMesh", ShowOnlyInnerProperties, FullyExpand="true", Tooltip=
 			"List of Foliage meshes used for this plant\n\nList of all foliage meshes applied based on procedural vegetation preset."))
 	TArray<struct FPVFoliageInfo> FoliageInfos;
+
+private:
+	TArray<FPinConnectionInfo> PinConnectionsToUpdate;
 };
  
 class FPVFoliageElement : public FPVBaseElement
diff --git a/Engine/Plugins/Experimental/ProceduralVegetationEditor/Source/ProceduralVegetation/Private/Nodes/PVGrowerSettings.cpp b/Engine/Plugins/Experimental/ProceduralVegetationEditor/Source/ProceduralVegetation/Private/Nodes/PVGrowerSettings.cpp
--- a/Engine/Plugins/Experimental/ProceduralVegetationEditor/Source/ProceduralVegetation/Private/Nodes/PVGrowerSettings.cpp
+++ b/Engine/Plugins/Experimental/ProceduralVegetationEditor/Source/ProceduralVegetation/Private/Nodes/PVGrowerSettings.cpp
@@ -4,6 +4,7 @@
 
 #include "PCGAssetExporterUtils.h"
 #include "ProceduralVegetationModule.h"
+#include "PVCommon.h"
 #include "GrowerSettings/PVGrowerPhyllotaxySettings.h"
 #include "DataTypes/PVData.h"
 #include "DataTypes/PVGrowthData.h"
@@ -17,22 +18,6 @@
 
 #define LOCTEXT_NAMESPACE "PVGrowerSettings"
 
-namespace GrowerPinsLabels
-{
-	const FName PhyllotaxyInputLabel = TEXT("Phyllotaxy");
-	const FName AxialInputLabel = TEXT("Axial");
-	const FName LateralInputLabel = TEXT("Lateral");
-	const FName PhototropismInputLabel = TEXT("Phototropism");
-	const FName AgeSenescenceInputLabel = TEXT("AgeSenescence");
-	const FName GuideInputLabel = TEXT("Guide");
-	const FName GravityInputLabel = TEXT("Gravity");
-	const FName LeafGrowthInputLabel = TEXT("LeafGrowth");
-	const FName TrunkGrowthInputLabel = TEXT("TrunkGrowth");
-	const FName BranchingInputLabel = TEXT("Branching");
-	const FName RandomAngleInputLabel = TEXT("RandomAngle");
-}
-
-
 #if WITH_EDITOR
 FText UPVGrowerSettings::GetDefaultNodeTitle() const
 { 
@@ -83,43 +68,43 @@ TArray<FPCGPinProperties> UPVGrowerSettings::InputPinProperties() const
 	Pin.SetOverrideOrUserParamPin();
 	Pin.bAllowMultipleData = false;
 
-	FPCGPinProperties& PhyllotaxyPin = Properties.Emplace_GetRef(GrowerPinsLabels::PhyllotaxyInputLabel, FPCGDataTypeIdentifier{ FPVDataTypeInfoGrowerPhyllotaxy::AsId()});
+	FPCGPinProperties& PhyllotaxyPin = Properties.Emplace_GetRef(PV::Pins::GrowerPhyllotaxyInputLabel, FPCGDataTypeIdentifier{ FPVDataTypeInfoGrowerPhyllotaxy::AsId()});
 	PhyllotaxyPin.SetAllowMultipleConnections(true);
 	PhyllotaxyPin.SetAdvancedPin();
 
-	FPCGPinProperties& AxialPin = Properties.Emplace_GetRef(GrowerPinsLabels::AxialInputLabel, FPCGDataTypeIdentifier{ FPVDataTypeInfoGrowerAxial::AsId()});
+	FPCGPinProperties& AxialPin = Properties.Emplace_GetRef(PV::Pins::GrowerAxialInputLabel, FPCGDataTypeIdentifier{ FPVDataTypeInfoGrowerAxial::AsId()});
 	AxialPin.SetAllowMultipleConnections(false);
 	AxialPin.SetAdvancedPin();
 
-	FPCGPinProperties& PhototropismPin = Properties.Emplace_GetRef(GrowerPinsLabels::PhototropismInputLabel, FPCGDataTypeIdentifier{ FPVDataTypeInfoGrowerPhototropism::AsId()});
+	FPCGPinProperties& PhototropismPin = Properties.Emplace_GetRef(PV::Pins::GrowerPhototropismInputLabel, FPCGDataTypeIdentifier{ FPVDataTypeInfoGrowerPhototropism::AsId()});
 	PhototropismPin.SetAllowMultipleConnections(true);
 	PhototropismPin.SetAdvancedPin();
 
-	FPCGPinProperties& LateralPin = Properties.Emplace_GetRef(GrowerPinsLabels::LateralInputLabel, FPCGDataTypeIdentifier{ FPVDataTypeInfoGrowerLateral::AsId()});
+	FPCGPinProperties& LateralPin = Properties.Emplace_GetRef(PV::Pins::GrowerLateralInputLabel, FPCGDataTypeIdentifier{ FPVDataTypeInfoGrowerLateral::AsId()});
 	LateralPin.SetAllowMultipleConnections(false);
 	LateralPin.SetAdvancedPin();
 
-	FPCGPinProperties& BranchingPin = Properties.Emplace_GetRef(GrowerPinsLabels::BranchingInputLabel, FPCGDataTypeIdentifier{ FPVDataTypeInfoGrowerBranching::AsId()});
+	FPCGPinProperties& BranchingPin = Properties.Emplace_GetRef(PV::Pins::GrowerBranchingInputLabel, FPCGDataTypeIdentifier{ FPVDataTypeInfoGrowerBranching::AsId()});
 	BranchingPin.SetAllowMultipleConnections(true);
 	BranchingPin.SetAdvancedPin();
 
-	FPCGPinProperties& AgeSenescencePin = Properties.Emplace_GetRef(GrowerPinsLabels::AgeSenescenceInputLabel, FPCGDataTypeIdentifier{ FPVDataTypeInfoGrowerAgeSenescence::AsId()});
+	FPCGPinProperties& AgeSenescencePin = Properties.Emplace_GetRef(PV::Pins::GrowerAgeSenescenceInputLabel, FPCGDataTypeIdentifier{ FPVDataTypeInfoGrowerAgeSenescence::AsId()});
 	AgeSenescencePin.SetAllowMultipleConnections(false);
 	AgeSenescencePin.SetAdvancedPin();
 
-	FPCGPinProperties& RandomAnglePin = Properties.Emplace_GetRef(GrowerPinsLabels::RandomAngleInputLabel, FPCGDataTypeIdentifier{ FPVDataTypeInfoGrowerRandomAngle::AsId()});
+	FPCGPinProperties& RandomAnglePin = Properties.Emplace_GetRef(PV::Pins::GrowerRandomAngleInputLabel, FPCGDataTypeIdentifier{ FPVDataTypeInfoGrowerRandomAngle::AsId()});
 	RandomAnglePin.SetAllowMultipleConnections(true);
 	RandomAnglePin.SetAdvancedPin();
 
-	FPCGPinProperties& LeafGrowthPin = Properties.Emplace_GetRef(GrowerPinsLabels::LeafGrowthInputLabel, FPCGDataTypeIdentifier{ FPVDataTypeInfoGrowerLeafGrowth::AsId()});
+	FPCGPinProperties& LeafGrowthPin = Properties.Emplace_GetRef(PV::Pins::GrowerLeafGrowthInputLabel, FPCGDataTypeIdentifier{ FPVDataTypeInfoGrowerLeafGrowth::AsId()});
 	LeafGrowthPin.SetAllowMultipleConnections(false);
 	LeafGrowthPin.SetAdvancedPin();
 
-	FPCGPinProperties& GravityPin = Properties.Emplace_GetRef(GrowerPinsLabels::GravityInputLabel, FPCGDataTypeIdentifier{ FPVDataTypeInfoGrowerGravity::AsId()});
+	FPCGPinProperties& GravityPin = Properties.Emplace_GetRef(PV::Pins::GrowerGravityInputLabel, FPCGDataTypeIdentifier{ FPVDataTypeInfoGrowerGravity::AsId()});
 	GravityPin.SetAllowMultipleConnections(false);
 	GravityPin.SetAdvancedPin();
 	
-	FPCGPinProperties& TrunkGrowthPin = Properties.Emplace_GetRef(GrowerPinsLabels::TrunkGrowthInputLabel, FPCGDataTypeIdentifier{ FPVDataTypeInfoGrowerTrunkGrowth::AsId()});
+	FPCGPinProperties& TrunkGrowthPin = Properties.Emplace_GetRef(PV::Pins::GrowerTrunkGrowthInputLabel, FPCGDataTypeIdentifier{ FPVDataTypeInfoGrowerTrunkGrowth::AsId()});
 	TrunkGrowthPin.SetAllowMultipleConnections(false);
 	TrunkGrowthPin.SetAdvancedPin();
 	
@@ -177,7 +162,7 @@ bool FPVGrowerElement::ExecuteInternal(FPCGContext* InContext) const
 	{
 		FPVGrowerParams GrowerParams = Settings->GrowerParams;
 		
-		for (const FPCGTaggedData& TaggedData : InContext->InputData.GetInputsByPin(GrowerPinsLabels::PhyllotaxyInputLabel))
+		for (const FPCGTaggedData& TaggedData : InContext->InputData.GetInputsByPin(PV::Pins::GrowerPhyllotaxyInputLabel))
 		{
 			if (const UPVGrowerPhyllotaxyData* PhyllotaxyData = Cast<UPVGrowerPhyllotaxyData>(TaggedData.Data))
 			{
@@ -194,7 +179,7 @@ bool FPVGrowerElement::ExecuteInternal(FPCGContext* InContext) const
 			}
 		}
 
-		for (const FPCGTaggedData& TaggedData : InContext->InputData.GetInputsByPin(GrowerPinsLabels::BranchingInputLabel))
+		for (const FPCGTaggedData& TaggedData : InContext->InputData.GetInputsByPin(PV::Pins::GrowerBranchingInputLabel))
 		{
 			if (const UPVGrowerBranchingData* Data = Cast<UPVGrowerBranchingData>(TaggedData.Data))
 			{
@@ -211,7 +196,7 @@ bool FPVGrowerElement::ExecuteInternal(FPCGContext* InContext) const
 			}
 		}
 
-		for (const FPCGTaggedData& TaggedData : InContext->InputData.GetInputsByPin(GrowerPinsLabels::LateralInputLabel))
+		for (const FPCGTaggedData& TaggedData : InContext->InputData.GetInputsByPin(PV::Pins::GrowerLateralInputLabel))
 		{
 			if (const UPVGrowerLateralData* Data = Cast<UPVGrowerLateralData>(TaggedData.Data))
 			{
@@ -219,7 +204,7 @@ bool FPVGrowerElement::ExecuteInternal(FPCGContext* InContext) const
 			}
 		}
 
-		for (const FPCGTaggedData& TaggedData : InContext->InputData.GetInputsByPin(GrowerPinsLabels::TrunkGrowthInputLabel))
+		for (const FPCGTaggedData& TaggedData : InContext->InputData.GetInputsByPin(PV::Pins::GrowerTrunkGrowthInputLabel))
 		{
 			if (const UPVGrowerTrunkGrowthData* Data = Cast<UPVGrowerTrunkGrowthData>(TaggedData.Data))
 			{
@@ -227,7 +212,7 @@ bool FPVGrowerElement::ExecuteInternal(FPCGContext* InContext) const
 			}
 		}
 
-		for (const FPCGTaggedData& TaggedData : InContext->InputData.GetInputsByPin(GrowerPinsLabels::AgeSenescenceInputLabel))
+		for (const FPCGTaggedData& TaggedData : InContext->InputData.GetInputsByPin(PV::Pins::GrowerAgeSenescenceInputLabel))
 		{
 			if (const UPVGrowerAgeSenescenceData* Data = Cast<UPVGrowerAgeSenescenceData>(TaggedData.Data))
 			{
@@ -235,7 +220,7 @@ bool FPVGrowerElement::ExecuteInternal(FPCGContext* InContext) const
 			}
 		}
 
-		for (const FPCGTaggedData& TaggedData : InContext->InputData.GetInputsByPin(GrowerPinsLabels::RandomAngleInputLabel))
+		for (const FPCGTaggedData& TaggedData : InContext->InputData.GetInputsByPin(PV::Pins::GrowerRandomAngleInputLabel))
 		{
 			if (const UPVGrowerRandomAngleData* Data = Cast<UPVGrowerRandomAngleData>(TaggedData.Data))
 			{
@@ -252,7 +237,7 @@ bool FPVGrowerElement::ExecuteInternal(FPCGContext* InContext) const
 			}
 		}
 
-		for (const FPCGTaggedData& TaggedData : InContext->InputData.GetInputsByPin(GrowerPinsLabels::GravityInputLabel))
+		for (const FPCGTaggedData& TaggedData : InContext->InputData.GetInputsByPin(PV::Pins::GrowerGravityInputLabel))
 		{
 			if (const UPVGrowerGravityData* Data = Cast<UPVGrowerGravityData>(TaggedData.Data))
 			{
@@ -260,7 +245,7 @@ bool FPVGrowerElement::ExecuteInternal(FPCGContext* InContext) const
 			}
 		}
 
-		for (const FPCGTaggedData& TaggedData : InContext->InputData.GetInputsByPin(GrowerPinsLabels::AxialInputLabel))
+		for (const FPCGTaggedData& TaggedData : InContext->InputData.GetInputsByPin(PV::Pins::GrowerAxialInputLabel))
 		{
 			if (const UPVGrowerAxialData* Data = Cast<UPVGrowerAxialData>(TaggedData.Data))
 			{
@@ -268,7 +253,7 @@ bool FPVGrowerElement::ExecuteInternal(FPCGContext* InContext) const
 			}
 		}
 
-		for (const FPCGTaggedData& TaggedData : InContext->InputData.GetInputsByPin(GrowerPinsLabels::PhototropismInputLabel))
+		for (const FPCGTaggedData& TaggedData : InContext->InputData.GetInputsByPin(PV::Pins::GrowerPhototropismInputLabel))
 		{
 			if (const UPVGrowerPhototropismData* Data = Cast<UPVGrowerPhototropismData>(TaggedData.Data))
 			{
@@ -285,7 +270,7 @@ bool FPVGrowerElement::ExecuteInternal(FPCGContext* InContext) const
 			}
 		}
 
-		for (const FPCGTaggedData& TaggedData : InContext->InputData.GetInputsByPin(GrowerPinsLabels::LeafGrowthInputLabel))
+		for (const FPCGTaggedData& TaggedData : InContext->InputData.GetInputsByPin(PV::Pins::GrowerLeafGrowthInputLabel))
 		{
 			if (const UPVGrowerLeafGrowthData* Data = Cast<UPVGrowerLeafGrowthData>(TaggedData.Data))
 			{
diff --git a/Engine/Plugins/Experimental/ProceduralVegetationEditor/Source/ProceduralVegetation/Private/PVCommon.h b/Engine/Plugins/Experimental/ProceduralVegetationEditor/Source/ProceduralVegetation/Private/PVCommon.h
new file mode 100644
--- /dev/null
+++ b/Engine/Plugins/Experimental/ProceduralVegetationEditor/Source/ProceduralVegetation/Private/PVCommon.h
@@ -0,0 +1,21 @@
+﻿// Copyright Epic Games, Inc. All Rights Reserved.
+#pragma once
+
+#include "CoreMinimal.h"
+
+namespace PV::Pins
+{
+	const FName GrowerPhyllotaxyInputLabel = TEXT("Phyllotaxy");
+	const FName GrowerAxialInputLabel = TEXT("Axial");
+	const FName GrowerLateralInputLabel = TEXT("Lateral");
+	const FName GrowerPhototropismInputLabel = TEXT("Phototropism");
+	const FName GrowerAgeSenescenceInputLabel = TEXT("AgeSenescence");
+	const FName GrowerGuideInputLabel = TEXT("Guide");
+	const FName GrowerGravityInputLabel = TEXT("Gravity");
+	const FName GrowerLeafGrowthInputLabel = TEXT("LeafGrowth");
+	const FName GrowerTrunkGrowthInputLabel = TEXT("TrunkGrowth");
+	const FName GrowerBranchingInputLabel = TEXT("Branching");
+	const FName GrowerRandomAngleInputLabel = TEXT("RandomAngle");
+	
+	const FName FoliageDistributorFoliageInputLabel = TEXT("Foliage");
+}
\ No newline at end of file
diff --git a/Engine/Plugins/Experimental/ProceduralVegetationEditor/Source/ProceduralVegetation/Public/DataTypes/PVFoliageData.h b/Engine/Plugins/Experimental/ProceduralVegetationEditor/Source/ProceduralVegetation/Public/DataTypes/PVFoliageData.h
new file mode 100644
--- /dev/null
+++ b/Engine/Plugins/Experimental/ProceduralVegetationEditor/Source/ProceduralVegetation/Public/DataTypes/PVFoliageData.h
@@ -0,0 +1,30 @@
+// Copyright Epic Games, Inc. All Rights Reserved.
+
+#pragma once
+
+#include "PVData.h"
+#include "PVFoliageData.generated.h"
+
+USTRUCT()
+struct FPVDataTypeInfoFoliage : public FPCGDataTypeInfo
+{
+	GENERATED_BODY()
+
+	PCG_DECLARE_TYPE_INFO(PROCEDURALVEGETATION_API);
+
+#if WITH_EDITOR
+	// If the type should be hidden to the user.
+	virtual bool Hidden() const override { return true; };
+#endif // WITH_EDITOR
+};
+
+UCLASS(BlueprintType, ClassGroup = (Procedural))
+class UPVFoliageData : public UPVData
+{
+	GENERATED_BODY()
+
+public:
+	//~ Begin UPCGData interface
+	PCG_ASSIGN_DEFAULT_TYPE_INFO(FPVDataTypeInfoFoliage)
+	//~ End UPCGData interface
+};
\ No newline at end of file
