Fix FLoopDebugData::Pivot not being initialized properly

#Synced-CL 50901855

diff --git a/Engine/Plugins/Experimental/ProceduralVegetationEditor/Source/ProceduralVegetation/Public/DataTypes/PVData.h b/Engine/Plugins/Experimental/ProceduralVegetationEditor/Source/ProceduralVegetation/Public/DataTypes/PVData.h
--- a/Engine/Plugins/Experimental/ProceduralVegetationEditor/Source/ProceduralVegetation/Public/DataTypes/PVData.h
+++ b/Engine/Plugins/Experimental/ProceduralVegetationEditor/Source/ProceduralVegetation/Public/DataTypes/PVData.h
@@ -131,7 +131,7 @@ struct FLoopDebugData
 	GENERATED_BODY()
 	
 	UPROPERTY(EditAnywhere, Category = DebugData)
-	FVector Pivot;
+	FVector Pivot = FVector::ZeroVector;
 
 	UPROPERTY(EditAnywhere, Category = DebugData)
 	TOptional<FVector> Direction;
