[PVE] Update Mesher node for compatibility with the Grafting workflow
- Mesher now removes non-bud points and adds new points throughout the skeleton
- Branch children are recomputed to compensate for inconsistencies in the legacy skeleton data
- NjordPixelIndex is recomputed for the whole skeleton before meshing
- All the gradients are recomputed before meshing the skeleton
- Refactored some code to reuse in the Mesher node

#Synced-CL 50967823

diff --git a/Engine/Plugins/Experimental/ProceduralVegetationEditor/Source/ProceduralVegetation/Private/Facades/PVBranchFacade.cpp b/Engine/Plugins/Experimental/ProceduralVegetationEditor/Source/ProceduralVegetation/Private/Facades/PVBranchFacade.cpp
--- a/Engine/Plugins/Experimental/ProceduralVegetationEditor/Source/ProceduralVegetation/Private/Facades/PVBranchFacade.cpp
+++ b/Engine/Plugins/Experimental/ProceduralVegetationEditor/Source/ProceduralVegetation/Private/Facades/PVBranchFacade.cpp
@@ -494,6 +494,33 @@ namespace PV::Facades
 		FILL_COLLECTION_ATTRIBUTE(Parents, InputArray);
 	}
 
+	void FBranchFacade::RecomputeBranchChildren()
+	{
+		const int32 NumberOfElements = BranchParentNumbers.Num();
+		TMap<int32, TArray<int32>> BranchNumbersToBranchChildren;
+		BranchNumbersToBranchChildren.Reserve(NumberOfElements);
+
+		for (int32 i = 0; i < NumberOfElements; i++)
+		{
+			const int32 BranchNumber = BranchNumbers[i];
+			const int32 BranchParentNumber = BranchParentNumbers[i];
+
+			if (BranchParentNumber != 0)
+				[[likely]]
+				{
+					BranchNumbersToBranchChildren.FindOrAdd(BranchParentNumber).Add(BranchNumber);
+				}
+		}
+
+		for (int32 i = 0; i < NumberOfElements; i++)
+		{
+			const int32 BranchNumber = BranchNumbers[i];
+			Children.ModifyAt(i, BranchNumbersToBranchChildren.Contains(BranchNumber)
+				? BranchNumbersToBranchChildren[BranchNumber]
+				: TArray<int32>());
+		}
+	}
+
 	void FBranchFacade::FillChildren(const TArray<TArray<int32>>& InputArray)
 	{
 		FILL_COLLECTION_ATTRIBUTE(Children, InputArray);
diff --git a/Engine/Plugins/Experimental/ProceduralVegetationEditor/Source/ProceduralVegetation/Private/Helpers/PVUtilities.cpp b/Engine/Plugins/Experimental/ProceduralVegetationEditor/Source/ProceduralVegetation/Private/Helpers/PVUtilities.cpp
--- a/Engine/Plugins/Experimental/ProceduralVegetationEditor/Source/ProceduralVegetation/Private/Helpers/PVUtilities.cpp
+++ b/Engine/Plugins/Experimental/ProceduralVegetationEditor/Source/ProceduralVegetation/Private/Helpers/PVUtilities.cpp
@@ -290,4 +290,127 @@ namespace PV::Utilities
 	
 		return PointFacade.IsValid() && BranchFacade.IsValid();
 	}
-}
\ No newline at end of file
+
+	template <typename T>
+	static FORCEINLINE T LerpValue(const T& A, const T& B, float Alpha)
+	{
+		if constexpr (TIsIntegral<T>::Value)
+		{
+			const float Val = FMath::Lerp(static_cast<float>(A), static_cast<float>(B), Alpha);
+			return static_cast<T>(FMath::RoundToInt(Val));
+		}
+		else
+		{
+			return FMath::Lerp(A, B, Alpha);
+		}
+	}
+
+	template <typename T>
+	TArray<T> LerpArrayElements(const TArray<T>& InArray1, const TArray<T>& InArray2, const float Alpha)
+	{
+		TArray<T> OutArray;
+		OutArray.SetNumUninitialized(FMath::Min(InArray1.Num(), InArray2.Num()));
+
+		for (int32 i = 0; i < FMath::Min(InArray1.Num(), InArray2.Num()); i++)
+		{
+			OutArray[i] = LerpValue(InArray1[i], InArray2[i], Alpha);
+		}
+
+		return OutArray;
+	}
+
+	int32 AddInterpolatedPointToCollection(FManagedArrayCollection& OutCollection, const int32 CurrentPointIndex, const int32 NextPointIndex,
+	                           const float Alpha, const int32 BudNumber, float NjordPixelIdx)
+	{
+		Facades::FPointFacade PointFacade(OutCollection);
+
+		const FVector3f CurrentPointPosition = PointFacade.GetPosition(CurrentPointIndex);
+		const FVector3f NextPointPosition = PointFacade.GetPosition(NextPointIndex);
+		const FVector3f Position = FMath::Lerp(CurrentPointPosition, NextPointPosition, Alpha);
+
+		const float CurrentPointLengthFromRoot = PointFacade.GetLengthFromRoot(CurrentPointIndex);
+		const float NextPointLengthFromRoot = PointFacade.GetLengthFromRoot(NextPointIndex);
+		const float LengthFromRoot = FMath::Lerp(CurrentPointLengthFromRoot, NextPointLengthFromRoot, Alpha);
+
+		const float CurrentPointScale = PointFacade.GetPointScale(CurrentPointIndex);
+		const float NextPointScale = PointFacade.GetPointScale(NextPointIndex);
+		const float PointScale = FMath::Lerp(CurrentPointScale, NextPointScale, Alpha);
+
+		const float CurrentPointLengthFromSeed = PointFacade.GetLengthFromSeed(CurrentPointIndex);
+		const float NextPointLengthFromSeed = PointFacade.GetLengthFromSeed(NextPointIndex);
+		const float LengthFromSeed = FMath::Lerp(CurrentPointLengthFromSeed, NextPointLengthFromSeed, Alpha);
+		
+		if (NjordPixelIdx == -1.0f)
+		{
+			const float CurrentPointNjordPixelIdx = PointFacade.GetNjordPixelIndex(CurrentPointIndex);
+			const float NextPointNjordPixelIdx = PointFacade.GetNjordPixelIndex(NextPointIndex);
+			NjordPixelIdx = FMath::Lerp(CurrentPointNjordPixelIdx, NextPointNjordPixelIdx, Alpha);
+		}
+
+		const TArray<float>& CurrentPointBudLightDetected = PointFacade.GetBudLightDetected(CurrentPointIndex);
+		const TArray<float>& NextPointBudLightDetected = PointFacade.GetBudLightDetected(NextPointIndex);
+		const TArray<float> BudLightDetected = LerpArrayElements(CurrentPointBudLightDetected, NextPointBudLightDetected, Alpha);
+
+		const TArray<int>& CurrentPointBudDevelopment = PointFacade.GetBudDevelopment(CurrentPointIndex);
+		const TArray<int>& NextPointBudDevelopment = PointFacade.GetBudDevelopment(NextPointIndex);
+		const TArray<int> BudDevelopment = LerpArrayElements(CurrentPointBudDevelopment, NextPointBudDevelopment, Alpha);
+
+		const TArray<float>& CurrentPointBudLateralMeristem = PointFacade.GetBudLateralMeristem(CurrentPointIndex);
+		const TArray<float>& NextPointBudLateralMeristem = PointFacade.GetBudLateralMeristem(NextPointIndex);
+		const TArray<float> BudLateralMeristem = LerpArrayElements(CurrentPointBudLateralMeristem, NextPointBudLateralMeristem, Alpha);
+
+		const TArray<float>& CurrentPointBudHormoneLevels = PointFacade.GetBudHormoneLevels(CurrentPointIndex);
+		const TArray<float>& NextPointBudHormoneLevels = PointFacade.GetBudHormoneLevels(NextPointIndex);
+		const TArray<float> BudHormoneLevels = LerpArrayElements(CurrentPointBudHormoneLevels, NextPointBudHormoneLevels, Alpha);
+
+		const TArray<FVector3f>& CurrentPointBudDirections = PointFacade.GetBudDirection(CurrentPointIndex);
+		const TArray<FVector3f>& NextPointBudDirections = PointFacade.GetBudDirection(NextPointIndex);
+		const TArray<FVector3f> BudDirections = LerpArrayElements(CurrentPointBudDirections, NextPointBudDirections, Alpha);
+
+		TArray<int> CurrentPointBudStatus;
+		PointFacade.GetBudStatus(CurrentPointIndex, CurrentPointBudStatus);
+		TArray<int> NextPointBudStatus;
+		PointFacade.GetBudStatus(NextPointIndex, NextPointBudStatus);
+		TArray<int> BudStatus;
+		BudStatus.SetNum(CurrentPointBudStatus.Num());
+
+		check(BudStatus.Num() == 10);
+
+		// 0_ApicalMeristem
+		BudStatus[0] = 0;
+		// 1_Codominant
+		BudStatus[1] = 0;
+		// 2_Axillary
+		BudStatus[2] = 1;
+		// 3_Seed
+		BudStatus[3] = 0;
+		// 4_Dormant
+		BudStatus[4] = 0;
+		// 5_Triggered
+		BudStatus[5] = 1;
+		// 6_NumTriggered
+		BudStatus[6] = 1;
+		// 7_Inactive
+		BudStatus[7] = 0;
+		// 8_BrokenTip
+		BudStatus[8] = 0;
+		// 9_Broken
+		BudStatus[9] = 0;
+
+		const int32 NewIndex = PointFacade.AddElements(1);
+		PointFacade.SetPosition(NewIndex, Position);
+		PointFacade.SetLengthFromRoot(NewIndex, LengthFromRoot);
+		PointFacade.SetPointScale(NewIndex, PointScale);
+		PointFacade.SetLengthFromSeed(NewIndex, LengthFromSeed);
+		PointFacade.SetNjordPixelIndex(NewIndex, NjordPixelIdx);
+		PointFacade.SetBudLightDetected(NewIndex, BudLightDetected);
+		PointFacade.SetBudDevelopment(NewIndex, BudDevelopment);
+		PointFacade.SetBudLateralMeristem(NewIndex, BudLateralMeristem);
+		PointFacade.SetBudNumber(NewIndex, BudNumber);
+		PointFacade.SetBudHormoneLevels(NewIndex, BudHormoneLevels);
+		PointFacade.SetBudDirections(NewIndex, BudDirections);
+		PointFacade.SetBudStatus(NewIndex, BudStatus);
+
+		return NewIndex;
+	}
+}
diff --git a/Engine/Plugins/Experimental/ProceduralVegetationEditor/Source/ProceduralVegetation/Private/Implementations/PVGrafter.cpp b/Engine/Plugins/Experimental/ProceduralVegetationEditor/Source/ProceduralVegetation/Private/Implementations/PVGrafter.cpp
--- a/Engine/Plugins/Experimental/ProceduralVegetationEditor/Source/ProceduralVegetation/Private/Implementations/PVGrafter.cpp
+++ b/Engine/Plugins/Experimental/ProceduralVegetationEditor/Source/ProceduralVegetation/Private/Implementations/PVGrafter.cpp
@@ -12,37 +12,9 @@
 #include "Facades/PVPointFacade.h"
 #include "Helpers/PVAttributesHelper.h"
 #include "Helpers/PVFoliageJSONHelper.h"
+#include "Helpers/PVUtilities.h"
 
 
-template <typename T>
-static FORCEINLINE T LerpValue(const T& A, const T& B, float Alpha)
-{
-	if constexpr (TIsIntegral<T>::Value)
-	{
-		const float Val = FMath::Lerp(static_cast<float>(A), static_cast<float>(B), Alpha);
-		return static_cast<T>(FMath::RoundToInt(Val));
-	}
-	else
-	{
-		return FMath::Lerp(A, B, Alpha);
-	}
-}
-
-template <typename T>
-TArray<T> FPVGrafter::LerpArrayElements(const TArray<T>& InArray1, const TArray<T>& InArray2, const float Alpha)
-{
-	check(InArray1.Num() == InArray2.Num());
-	TArray<T> OutArray;
-	OutArray.SetNumUninitialized(InArray1.Num());
-
-	for (int32 i = 0; i < InArray1.Num(); i++)
-	{
-		OutArray[i] = LerpValue(InArray1[i], InArray2[i], Alpha);
-	}
-
-	return OutArray;
-}
-
 template <typename Fn>
 void FPVGrafter::ForEachGraftPointExcludingRoot(const int32 NumPoints, const int32 RootIndex, Fn&& Func)
 {
@@ -59,92 +31,6 @@ void FPVGrafter::ForEachGraftPointExcludingRoot(const int32 NumPoints, const int
 	}
 }
 
-int32 FPVGrafter::AddPointToCollection(FManagedArrayCollection& OutCollection, const int32 CurrentPointIndex, const int32 NextPointIndex,
-                                       const float Alpha, const int32 BudNumber)
-{
-	PV::Facades::FPointFacade PointFacade(OutCollection);
-
-	const FVector3f CurrentPointPosition = PointFacade.GetPosition(CurrentPointIndex);
-	const FVector3f NextPointPosition = PointFacade.GetPosition(NextPointIndex);
-	const FVector3f Position = FMath::Lerp(CurrentPointPosition, NextPointPosition, Alpha);
-
-	const float CurrentPointLengthFromRoot = PointFacade.GetLengthFromRoot(CurrentPointIndex);
-	const float NextPointLengthFromRoot = PointFacade.GetLengthFromRoot(NextPointIndex);
-	const float LengthFromRoot = FMath::Lerp(CurrentPointLengthFromRoot, NextPointLengthFromRoot, Alpha);
-
-	const float CurrentPointScale = PointFacade.GetPointScale(CurrentPointIndex);
-	const float NextPointScale = PointFacade.GetPointScale(NextPointIndex);
-	const float PointScale = FMath::Lerp(CurrentPointScale, NextPointScale, Alpha);
-
-	const float CurrentPointLengthFromSeed = PointFacade.GetLengthFromSeed(CurrentPointIndex);
-	const float NextPointLengthFromSeed = PointFacade.GetLengthFromSeed(NextPointIndex);
-	const float LengthFromSeed = FMath::Lerp(CurrentPointLengthFromSeed, NextPointLengthFromSeed, Alpha);
-
-	const TArray<float>& CurrentPointBudLightDetected = PointFacade.GetBudLightDetected(CurrentPointIndex);
-	const TArray<float>& NextPointBudLightDetected = PointFacade.GetBudLightDetected(NextPointIndex);
-	const TArray<float> BudLightDetected = LerpArrayElements(CurrentPointBudLightDetected, NextPointBudLightDetected, Alpha);
-
-	const TArray<int>& CurrentPointBudDevelopment = PointFacade.GetBudDevelopment(CurrentPointIndex);
-	const TArray<int>& NextPointBudDevelopment = PointFacade.GetBudDevelopment(NextPointIndex);
-	const TArray<int> BudDevelopment = LerpArrayElements(CurrentPointBudDevelopment, NextPointBudDevelopment, Alpha);
-
-	const TArray<float>& CurrentPointBudLateralMeristem = PointFacade.GetBudLateralMeristem(CurrentPointIndex);
-	const TArray<float>& NextPointBudLateralMeristem = PointFacade.GetBudLateralMeristem(NextPointIndex);
-	const TArray<float> BudLateralMeristem = LerpArrayElements(CurrentPointBudLateralMeristem, NextPointBudLateralMeristem, Alpha);
-
-	const TArray<float>& CurrentPointBudHormoneLevels = PointFacade.GetBudHormoneLevels(CurrentPointIndex);
-	const TArray<float>& NextPointBudHormoneLevels = PointFacade.GetBudHormoneLevels(NextPointIndex);
-	const TArray<float> BudHormoneLevels = LerpArrayElements(CurrentPointBudHormoneLevels, NextPointBudHormoneLevels, Alpha);
-
-	const TArray<FVector3f>& CurrentPointBudDirections = PointFacade.GetBudDirection(CurrentPointIndex);
-	const TArray<FVector3f>& NextPointBudDirections = PointFacade.GetBudDirection(NextPointIndex);
-	const TArray<FVector3f> BudDirections = LerpArrayElements(CurrentPointBudDirections, NextPointBudDirections, Alpha);
-
-	TArray<int> CurrentPointBudStatus;
-	PointFacade.GetBudStatus(CurrentPointIndex, CurrentPointBudStatus);
-	TArray<int> NextPointBudStatus;
-	PointFacade.GetBudStatus(NextPointIndex, NextPointBudStatus);
-	TArray<int> BudStatus;
-	BudStatus.SetNum(CurrentPointBudStatus.Num());
-
-	check(BudStatus.Num() == 10);
-
-	// 0_ApicalMeristem
-	BudStatus[0] = 0;
-	// 1_Codominant
-	BudStatus[1] = 0;
-	// 2_Axillary
-	BudStatus[2] = 1;
-	// 3_Seed
-	BudStatus[3] = 0;
-	// 4_Dormant
-	BudStatus[4] = 0;
-	// 5_Triggered
-	BudStatus[5] = 1;
-	// 6_NumTriggered
-	BudStatus[6] = 1;
-	// 7_Inactive
-	BudStatus[7] = 0;
-	// 8_BrokenTip
-	BudStatus[8] = 0;
-	// 9_Broken
-	BudStatus[9] = 0;
-
-	const int32 NewIndex = PointFacade.AddElements(1);
-	PointFacade.SetPosition(NewIndex, Position);
-	PointFacade.SetLengthFromRoot(NewIndex, LengthFromRoot);
-	PointFacade.SetPointScale(NewIndex, PointScale);
-	PointFacade.SetLengthFromSeed(NewIndex, LengthFromSeed);
-	PointFacade.SetBudLightDetected(NewIndex, BudLightDetected);
-	PointFacade.SetBudDevelopment(NewIndex, BudDevelopment);
-	PointFacade.SetBudLateralMeristem(NewIndex, BudLateralMeristem);
-	PointFacade.SetBudNumber(NewIndex, BudNumber);
-	PointFacade.SetBudHormoneLevels(NewIndex, BudHormoneLevels);
-	PointFacade.SetBudDirections(NewIndex, BudDirections);
-	PointFacade.SetBudStatus(NewIndex, BudStatus);
-
-	return NewIndex;
-}
 
 FTransform FPVGrafter::GetPointTransform(const FVector3f AttachmentPoint, const FVector3f UpVector, const FVector3f NormalVector, const float Scale)
 {
@@ -168,35 +54,6 @@ FTransform FPVGrafter::GetPointTransform(const FVector3f AttachmentPoint, const
 	return FTransform(RotationQuat, FVector(AttachmentPoint), FVector(Scale));
 }
 
-void FPVGrafter::RecomputeBranchChildren(
-	PV::FBranchNumberAttributeConstView BranchNumberAttribute,
-	PV::FBranchParentNumberAttributeConstView BranchParentNumberAttribute,
-	PV::FBranchChildrenAttributeView BranchChildrenAttribute)
-{
-	const int32 NumberOfElements = BranchParentNumberAttribute.Num();
-	TMap<int32, TArray<int32>> BranchNumbersToBranchChildren;
-	BranchNumbersToBranchChildren.Reserve(NumberOfElements);
-
-	for (int32 i = 0; i < NumberOfElements; i++)
-	{
-		const int32 BranchNumber = BranchNumberAttribute[i];
-		const int32 BranchParentNumber = BranchParentNumberAttribute[i];
-
-		if (BranchParentNumber != 0)
-		[[likely]]
-		{
-			BranchNumbersToBranchChildren.FindOrAdd(BranchParentNumber).Add(BranchNumber);
-		}
-	}
-
-	for (int32 i = 0; i < NumberOfElements; i++)
-	{
-		const int32 BranchNumber = BranchNumberAttribute[i];
-		BranchChildrenAttribute[i] = BranchNumbersToBranchChildren.Contains(BranchNumber)
-			? BranchNumbersToBranchChildren[BranchNumber]
-			: TArray<int32>();
-	}
-}
 
 void FPVGrafter::ComputeBranchGradientAttribute(
 	PV::FPointBranchGradientAttributeView PointBranchGradientAttribute,
@@ -475,6 +332,14 @@ void FPVGrafter::SetGraft(FManagedArrayCollection& OutCollection, const FManaged
 				}
 				PointFacadeOut.SetBudDirections(TargetPointIndex + PointFacadeOutStartIndex, TransformedValues);
 			});
+	
+	ForEachGraftPointExcludingRoot(NumOfPointsInSourceGraft, SourceGraftTrunkRootPointIndex,
+		[&](const int32 SourcePointIndex, const int32 TargetPointIndex)
+			{
+				TArray<int> BudStatusInGraftPoint;
+				PointFacadeSourceGraft.GetBudStatus(SourcePointIndex, BudStatusInGraftPoint);
+				PointFacadeOut.SetBudStatus(TargetPointIndex + PointFacadeOutStartIndex, BudStatusInGraftPoint);
+			});
 
 	// Transform and set branches
 	PV::Facades::FBranchFacade BranchFacadeOut(OutCollection);
@@ -757,7 +622,8 @@ void FPVGrafter::DistributeGraft(FManagedArrayCollection& OutCollection, const F
 
 	// Only need to do this temporarily since legacy skeleton has inconsistent data in children
 	// TODO: Remove this once support for legacy skeletons is droppped
-	RecomputeBranchChildren(BranchNumberAttribute, BranchParentNumberAttribute, BranchChildrenAttribute);
+	PV::Facades::FBranchFacade BranchFacadeOut(OutCollection);
+	BranchFacadeOut.RecomputeBranchChildren();
 
 	ComputeBranchGradientAttribute(
 		PointBranchGradientAttribute,
@@ -904,7 +770,16 @@ void FPVGrafter::DistributeGraft(FManagedArrayCollection& OutCollection, const F
 					}
 					else
 					{
-						IndexAdded = AddPointToCollection(OutCollection, CurrentPointIndex, NextPointIndex, DistanceAlpha, CurrentBudNumber);
+						// Setting NjordPixelIdx to an integer value so that it is not removed
+						// in the mesher when it removes non bud points
+						IndexAdded = PV::Utilities::AddInterpolatedPointToCollection(
+							OutCollection,
+							CurrentPointIndex,
+							NextPointIndex,
+							DistanceAlpha,
+							CurrentBudNumber,
+							CurrentBudNumber);
+						
 						CurrentBudNumber++;
 					}
 					AttachmentPointsPointIndices.Add(IndexAdded);
diff --git a/Engine/Plugins/Experimental/ProceduralVegetationEditor/Source/ProceduralVegetation/Private/Implementations/PVGrafter.h b/Engine/Plugins/Experimental/ProceduralVegetationEditor/Source/ProceduralVegetation/Private/Implementations/PVGrafter.h
--- a/Engine/Plugins/Experimental/ProceduralVegetationEditor/Source/ProceduralVegetation/Private/Implementations/PVGrafter.h
+++ b/Engine/Plugins/Experimental/ProceduralVegetationEditor/Source/ProceduralVegetation/Private/Implementations/PVGrafter.h
@@ -58,20 +58,9 @@ struct FPVGrafter
 		bool bIsTipInstance;
 	};
 
-	template <typename T>
-	static TArray<T> LerpArrayElements(const TArray<T>& InArray1, const TArray<T>& InArray2, const float Alpha);
-
-	static int32 AddPointToCollection(FManagedArrayCollection& OutCollection, const int32 CurrentPointIndex, const int32 NextPointIndex,
-	                                  const float Alpha, const int32 BudNumber);
-
 	static FTransform GetPointTransform(const FVector3f AttachmentPoint, const FVector3f UpVector, const FVector3f NormalVector,
 	                                    const float Scale);
 
-	static void RecomputeBranchChildren(
-		PV::FBranchNumberAttributeConstView BranchNumberAttribute,
-		PV::FBranchParentNumberAttributeConstView BranchParentNumberAttribute,
-		PV::FBranchChildrenAttributeView BranchChildrenAttribute);
-
 	static void ComputeBranchGradientAttribute(
 		PV::FPointBranchGradientAttributeView PointBranchGradientAttribute,
 		PV::FBranchParentNumberAttributeConstView BranchParentNumberAttribute,
diff --git a/Engine/Plugins/Experimental/ProceduralVegetationEditor/Source/ProceduralVegetation/Private/Implementations/PVMeshBuilder.cpp b/Engine/Plugins/Experimental/ProceduralVegetationEditor/Source/ProceduralVegetation/Private/Implementations/PVMeshBuilder.cpp
--- a/Engine/Plugins/Experimental/ProceduralVegetationEditor/Source/ProceduralVegetation/Private/Implementations/PVMeshBuilder.cpp
+++ b/Engine/Plugins/Experimental/ProceduralVegetationEditor/Source/ProceduralVegetation/Private/Implementations/PVMeshBuilder.cpp
@@ -13,6 +13,7 @@
 #include "Facades/PVPlantFacade.h"
 #include "Facades/PVPointFacade.h"
 #include "Facades/PVTreeFacade.h"
+#include "Helpers/PVAttributesHelper.h"
 
 #include "Helpers/PVUtilities.h"
 
@@ -54,7 +55,8 @@ void FPVMeshBuilder::GenerateGeometryCollection(const FManagedArrayCollection& I
 	PV::Facades::FBoneFacade::DefineSchema(OutGeometryCollection);
 	PV::Facades::FBoneFacade BoneFacade = PV::Facades::FBoneFacade(OutGeometryCollection);
 
-	PV::Facades::FBranchFacade BranchFacade = PV::Facades::FBranchFacade(OutGeometryCollection);
+	PV::Facades::FBranchFacade BranchFacade(OutGeometryCollection);
+	PV::Facades::FPointFacade PointFacade(OutGeometryCollection);
 
 	if (BranchFacade.GetElementCount() == 0)
 	{
@@ -62,7 +64,45 @@ void FPVMeshBuilder::GenerateGeometryCollection(const FManagedArrayCollection& I
 		return;
 	}
 
-	SetNjordPixelID(OutGeometryCollection);
+	RemoveNonBudPoints(BranchFacade, PointFacade);
+	
+	auto PointNjordPixelIndexAttribute = PV::FPointNjordPixelIndexAttribute::GetAttribute(OutGeometryCollection);
+	auto BranchPointsAttribute = PV::FBranchPointsAttribute::GetAttribute(OutGeometryCollection);
+	auto BranchParentNumberAttribute = PV::FBranchParentNumberAttribute::GetAttribute(OutGeometryCollection);
+	auto BranchNumberAttribute = PV::FBranchNumberAttribute::GetAttribute(OutGeometryCollection);
+	auto BranchChildrenAttribute = PV::FBranchChildrenAttribute::GetAttribute(OutGeometryCollection);
+
+	BranchFacade.RecomputeBranchChildren();
+	
+	PV::AttributesHelper::ComputeNjordPixelIndex(
+		PointNjordPixelIndexAttribute,
+		BranchPointsAttribute,
+		BranchParentNumberAttribute, 
+		BranchNumberAttribute,
+		BranchChildrenAttribute
+	);
+	
+	AddPointsToSkeleton(OutGeometryCollection);
+	
+	auto GroundGradientAttribute = PV::FPointGroundGradientAttribute::GetAttribute(OutGeometryCollection);
+	auto PointScaleGradientAttribute = PV::FPointScaleGradientAttribute::GetAttribute(OutGeometryCollection);
+	auto PointPlantGradientAttribute = PV::FPointPlantGradientAttribute::GetAttribute(OutGeometryCollection);
+	auto PointBranchGradientAttribute = PV::FPointBranchGradientAttribute::GetAttribute(OutGeometryCollection);
+	auto PointLengthFromRootAttribute = PV::FPointLengthFromRootAttribute::GetAttribute(OutGeometryCollection);
+	auto PointScaleAttribute = PV::FPointScaleAttribute::GetAttribute(OutGeometryCollection);
+	
+	PV::AttributesHelper::ComputeGradients(
+		GroundGradientAttribute,
+		PointScaleGradientAttribute,
+		PointPlantGradientAttribute,
+		PointBranchGradientAttribute,
+		PointLengthFromRootAttribute,
+		PointScaleAttribute,
+		BranchPointsAttribute,
+		BranchParentNumberAttribute,
+		BranchNumberAttribute,
+		BranchChildrenAttribute
+	);
 
 	MeshBuilderParams.MaterialSettings.ApplyMaterialSettings(OutGeometryCollection);
 	GetLocalDynamicMeshData(OutGeometryCollection, MeshBuilderParams, MeshesDataCollection);
@@ -215,28 +255,143 @@ void FPVMeshBuilder::GenerateGeometryCollection(const FManagedArrayCollection& I
 	OutGeometryCollection.UpdateBoundingBox();
 }
 
-void FPVMeshBuilder::SetNjordPixelID(FGeometryCollection& OutGeometryCollection)
+void FPVMeshBuilder::AddPointsToSkeleton(FManagedArrayCollection& OutCollection)
 {
-	PV::Facades::FPointFacade PointFacade(OutGeometryCollection);
-	TManagedArray<float>& NjordPixelIds = PointFacade.ModifyNjordPixelIDs();
+	PV::Facades::FBranchFacade BranchFacade(OutCollection);
+	PV::Facades::FPointFacade PointFacade(OutCollection);
+	
+	const int32 NumOfBranches = BranchFacade.GetElementCount();
+	const int32 NumOfPoints = PointFacade.GetElementCount();
 
-	TMap<int, TArray<int>> BudNumberPointIndicesMap;
-	for (int PointIndex = 0; PointIndex < PointFacade.GetElementCount(); PointIndex++)
+	int32 LongestBranchIndex = -1;
+	float LongestBranchLength = TNumericLimits<float>::Lowest();
+	TMap<int32, TArray<int32>> BranchIndexToBranchPoints;
+	BranchIndexToBranchPoints.Reserve(NumOfPoints);
+	
+	for (int32 BranchIndex = 0; BranchIndex < NumOfBranches; BranchIndex++)
 	{
-		int BudNumber = PointFacade.GetBudNumber(PointIndex);
-		BudNumberPointIndicesMap.FindOrAdd(BudNumber).Add(PointIndex);
+		const TArray<int32>& BranchPoints = BranchFacade.GetPoints(BranchIndex);
+		TArray<int32> SortedPointIndices = BranchPoints;
+		SortedPointIndices.Sort([&PointFacade](const int32& A, const int32& B)
+			{
+				return PointFacade.GetLengthFromRoot(A) < PointFacade.GetLengthFromRoot(B);
+			});
+		BranchIndexToBranchPoints.Add(BranchIndex, SortedPointIndices);
+		
+		if (BranchPoints.Num() < 2)
+		[[unlikely]]
+		{
+			continue;
+		}
+
+		const float BranchLength = PointFacade.GetLengthFromRoot(SortedPointIndices.Last()) - PointFacade.GetLengthFromRoot(SortedPointIndices[0]);
+		if (BranchLength > LongestBranchLength)
+		{
+			LongestBranchLength = BranchLength;
+			LongestBranchIndex = BranchIndex;
+		}
+	}
+	
+	if (LongestBranchIndex == -1)
+	{
+		return;
+	}
+	
+	const float DivisionUnitLength = LongestBranchLength / 500.0f;
+	
+	for (int32 BranchIndex = 0; BranchIndex < NumOfBranches; BranchIndex++)
+	{
+		const TArray<int32>& BranchPoints = BranchIndexToBranchPoints[BranchIndex];
+		TArray<int32> BranchPointsUpdated(BranchPoints);
+		if (BranchPoints.Num() < 2)
+		[[unlikely]]
+		{
+			continue;
+		}
+		
+		for (int32 i = 0; i < BranchPoints.Num() - 1; i++)
+		{
+			const int32 CurrentPointIndex = BranchPoints[i];
+			const int32 NextPointIndex = BranchPoints[i + 1];
+			const int32 CurrentPointBudNumber = PointFacade.GetBudNumber(CurrentPointIndex);
+
+			const float DistanceBetweenPoints = PointFacade.GetLengthFromRoot(NextPointIndex) - PointFacade.GetLengthFromRoot(CurrentPointIndex);
+			const int32 NumOfSections = FMath::RoundToInt32(DistanceBetweenPoints / DivisionUnitLength);
+		
+			if (NumOfSections <= 1)
+			[[unlikely]]
+			{
+				continue;
+			}
+		
+			const float SectionLength = DistanceBetweenPoints / NumOfSections;
+			const int32 NumberOfPointsToAdd = NumOfSections - 1;
+		
+			for (int32 j = 0; j < NumberOfPointsToAdd; j++)
+			{
+				const float Alpha = (SectionLength * (j + 1)) / DistanceBetweenPoints;
+				const int32 NewPointIndex = PV::Utilities::AddInterpolatedPointToCollection(
+					OutCollection,
+					CurrentPointIndex,
+					NextPointIndex,
+					Alpha,
+					CurrentPointBudNumber);
+				
+				BranchPointsUpdated.Add(NewPointIndex);
+			}
+		}
+		
+		BranchPointsUpdated.Sort([&PointFacade](const int32& A, const int32& B)
+			{
+				return PointFacade.GetLengthFromRoot(A) < PointFacade.GetLengthFromRoot(B);
+			});
+		BranchFacade.SetPoints(BranchIndex, MoveTemp(BranchPointsUpdated));
 	}
+}
 
-	for (const auto& KeyValuePair : BudNumberPointIndicesMap)
+void FPVMeshBuilder::RemoveNonBudPoints(PV::Facades::FBranchFacade& BranchFacade, PV::Facades::FPointFacade& PointFacade)
+{
+	// Note that we only remove points and do not care about removing branches or 
+	// Foliage instances since we assume that each branch will contain at least 
+	// two bud points
+	const int32 NumOfPoints = PointFacade.GetElementCount();
+	const int32 NumOfBranches = BranchFacade.GetElementCount();
+	
+	if (NumOfPoints == 0 || NumOfBranches == 0)
+	{
+		return;
+	}
+	
+	TArray<bool> PointsToRemove;
+	PointsToRemove.Init(false, NumOfPoints);
+	
+	for (int32 PointIndex = 0; PointIndex < NumOfPoints; PointIndex++)
+	{
+		const float NjordPixelIndex = PointFacade.GetNjordPixelIndex(PointIndex);
+		if (NjordPixelIndex != FMath::FloorToFloat(NjordPixelIndex))
+		{
+			PointsToRemove[PointIndex] = true;
+		}
+	}
+	
+	TMap<int32, int32> PointsOldIDsToNewIDs = PV::Facades::IShrinkable::RemoveEntries(PointFacade, PointsToRemove);
+	
+	for (int32 BranchIndex = 0; BranchIndex < NumOfBranches; ++BranchIndex)
 	{
-		int ValueIndex = 0;
-		for (const auto Value : KeyValuePair.Value)
+		TArray<int32> UpdatedPoints;
+		const TArray<int32>& BranchPoints = BranchFacade.GetPoints(BranchIndex);
+		for (const int32& PointIndex : BranchPoints)
 		{
-			//Modify NjordPixelID
-			float NjordPixelID = KeyValuePair.Key + static_cast<float>(ValueIndex) / KeyValuePair.Value.Num();
-			NjordPixelIds[Value] = NjordPixelID;
-			ValueIndex++;
+			if (PointsOldIDsToNewIDs.Contains(PointIndex))
+			{
+				UpdatedPoints.Add(PointsOldIDsToNewIDs[PointIndex]);
+			}
+			else if (!PointsToRemove[PointIndex])
+			{
+				UpdatedPoints.Add(PointIndex);
+			}
 		}
+		BranchFacade.SetPoints(BranchIndex, UpdatedPoints);
 	}
 }
 
diff --git a/Engine/Plugins/Experimental/ProceduralVegetationEditor/Source/ProceduralVegetation/Private/Implementations/PVMeshBuilder.h b/Engine/Plugins/Experimental/ProceduralVegetationEditor/Source/ProceduralVegetation/Private/Implementations/PVMeshBuilder.h
--- a/Engine/Plugins/Experimental/ProceduralVegetationEditor/Source/ProceduralVegetation/Private/Implementations/PVMeshBuilder.h
+++ b/Engine/Plugins/Experimental/ProceduralVegetationEditor/Source/ProceduralVegetation/Private/Implementations/PVMeshBuilder.h
@@ -22,6 +22,7 @@
 
 #include "Materials/MaterialInstance.h"
 #include "Materials/MaterialInterface.h"
+#include "Utils/PVAttributes.h"
 
 #include "PVMeshBuilder.generated.h"
 
@@ -310,7 +311,9 @@ struct FPVMeshBuilder
 	static void GenerateGeometryCollection(const FManagedArrayCollection& InSkeletonCollection, const FPVMeshBuilderParams& MeshBuilderParams,
 	                                       FGeometryCollection& OutGeometryCollection);
 
-	static void SetNjordPixelID(FGeometryCollection& OutGeometryCollection);
+	static void AddPointsToSkeleton(FManagedArrayCollection& OutCollection);
+	
+	static void RemoveNonBudPoints(PV::Facades::FBranchFacade& BranchFacade, PV::Facades::FPointFacade& PointFacade);
 
 	static void GenerateDynamicMesh(FManagedArrayCollection& Collection, const FPVMeshBuilderParams& MeshBuilderParams,
 	                                TObjectPtr<UDynamicMesh>& OutMesh);
diff --git a/Engine/Plugins/Experimental/ProceduralVegetationEditor/Source/ProceduralVegetation/Public/Facades/PVBranchFacade.h b/Engine/Plugins/Experimental/ProceduralVegetationEditor/Source/ProceduralVegetation/Public/Facades/PVBranchFacade.h
--- a/Engine/Plugins/Experimental/ProceduralVegetationEditor/Source/ProceduralVegetation/Public/Facades/PVBranchFacade.h
+++ b/Engine/Plugins/Experimental/ProceduralVegetationEditor/Source/ProceduralVegetation/Public/Facades/PVBranchFacade.h
@@ -107,6 +107,8 @@ namespace PV::Facades
 
 		void FillParents(const TArray<TArray<int32>>& InputArray);
 		
+		void RecomputeBranchChildren();
+		
 		void FillChildren(const TArray<TArray<int32>>& InputArray);
 		
 		void FillBranchPoints(const TArray<TArray<int32>>& InputArray);
diff --git a/Engine/Plugins/Experimental/ProceduralVegetationEditor/Source/ProceduralVegetation/Public/Helpers/PVUtilities.h b/Engine/Plugins/Experimental/ProceduralVegetationEditor/Source/ProceduralVegetation/Public/Helpers/PVUtilities.h
--- a/Engine/Plugins/Experimental/ProceduralVegetationEditor/Source/ProceduralVegetation/Public/Helpers/PVUtilities.h
+++ b/Engine/Plugins/Experimental/ProceduralVegetationEditor/Source/ProceduralVegetation/Public/Helpers/PVUtilities.h
@@ -35,6 +35,12 @@ namespace PV::Utilities
 	FLinearColor PROCEDURALVEGETATION_API GetRandomHueColor(float Alpha);
 
 	bool IsValidPVData(const FManagedArrayCollection& Collection);
+
+	template <typename T>
+	TArray<T> LerpArrayElements(const TArray<T>& InArray1, const TArray<T>& InArray2, const float Alpha);
+
+	int32 AddInterpolatedPointToCollection(FManagedArrayCollection& OutCollection, const int32 CurrentPointIndex, const int32 NextPointIndex,
+	                                  const float Alpha, const int32 BudNumber, float NjordPixelIdx=-1.0f);
 }
 
 USTRUCT()
