[PVE] Multiple Changes
- Refactor PVEditor for ITF
- Overhaul for Export Window in PVE
- Remove some redundant warnings

#Synced-CL 50970612

diff --git a/Engine/Plugins/Experimental/ProceduralVegetationEditor/Source/ProceduralVegetation/Private/DataTypes/PVData.cpp b/Engine/Plugins/Experimental/ProceduralVegetationEditor/Source/ProceduralVegetation/Private/DataTypes/PVData.cpp
--- a/Engine/Plugins/Experimental/ProceduralVegetationEditor/Source/ProceduralVegetation/Private/DataTypes/PVData.cpp
+++ b/Engine/Plugins/Experimental/ProceduralVegetationEditor/Source/ProceduralVegetation/Private/DataTypes/PVData.cpp
@@ -11,12 +11,12 @@
 UPVData::UPVData(const FObjectInitializer& ObjectInitializer)
 	: UPCGSpatialData(ObjectInitializer)
 {
-	
+	Collection = MakeShared<FManagedArrayCollection>();
 }
 
 void UPVData::Initialize(FManagedArrayCollection&& InCollection, bool bCanTakeOwnership)
 {
-	Collection = InCollection;
+	Collection = MakeShared<FManagedArrayCollection>(InCollection);
 }
 
 void UPVData::AddToCrc(FArchiveCrc32& Ar, bool bFullDataCrc) const
@@ -30,7 +30,7 @@ void UPVData::AddToCrc(FArchiveCrc32& Ar, bool bFullDataCrc) const
 		FString ClassName = StaticClass()->GetPathName();
 
 		Ar << ClassName;
-		Ar << Collection;
+		Ar << *Collection;
 #if WITH_EDITOR
 		Ar << DebugSettings;
 #endif
@@ -64,7 +64,7 @@ UPCGSpatialData* UPVData::CopyInternal(FPCGContext* Context) const
 	UPVData* NewProceduralVegetationData = FPCGContext::NewObject_AnyThread<UPVData>(Context);
 
 	FManagedArrayCollection NewCollection;
-	Collection.CopyTo(&NewCollection);
+	Collection->CopyTo(&NewCollection);
 	
 	NewProceduralVegetationData->Initialize(MoveTemp(NewCollection), /*bCanTakeOwnership=*/false);
 
diff --git a/Engine/Plugins/Experimental/ProceduralVegetationEditor/Source/ProceduralVegetation/Private/Helpers/PVUtilities.cpp b/Engine/Plugins/Experimental/ProceduralVegetationEditor/Source/ProceduralVegetation/Private/Helpers/PVUtilities.cpp
--- a/Engine/Plugins/Experimental/ProceduralVegetationEditor/Source/ProceduralVegetation/Private/Helpers/PVUtilities.cpp
+++ b/Engine/Plugins/Experimental/ProceduralVegetationEditor/Source/ProceduralVegetation/Private/Helpers/PVUtilities.cpp
@@ -74,8 +74,8 @@ namespace PV::Utilities
 				FString FoliageName = FoliageInfo.Mesh.ToString();
 				FSoftObjectPath MeshPath(FoliageName);
  
-				UStaticMesh* StaticFoliageMesh = PackageExists(MeshPath.GetLongPackageName(), UStaticMesh::StaticClass()) ? LoadObject<UStaticMesh>(nullptr, *FoliageName) : nullptr;
-				USkeletalMesh* SkeletalFoliageMesh = PackageExists(MeshPath.GetLongPackageName(), USkeletalMesh::StaticClass()) ? LoadObject<USkeletalMesh>(nullptr, *FoliageName) : nullptr;
+				UStaticMesh* StaticFoliageMesh = PackageExists(MeshPath.GetLongPackageName(), UStaticMesh::StaticClass()) ? LoadObject<UStaticMesh>(nullptr, *FoliageName, {}, LOAD_NoWarn | LOAD_Quiet) : nullptr;
+				USkeletalMesh* SkeletalFoliageMesh = PackageExists(MeshPath.GetLongPackageName(), USkeletalMesh::StaticClass()) ? LoadObject<USkeletalMesh>(nullptr, *FoliageName, {}, LOAD_NoWarn | LOAD_Quiet) : nullptr;
  
 				if(StaticFoliageMesh)
 				{
@@ -252,8 +252,8 @@ namespace PV::Utilities
  
 	int32 GetMeshTriangles(const FString InMeshPath)
 	{
-		UStaticMesh* StaticFoliageMesh = LoadObject<UStaticMesh>(nullptr, *InMeshPath);
-		USkeletalMesh* SkeletalFoliageMesh = LoadObject<USkeletalMesh>(nullptr, *InMeshPath);
+		UStaticMesh* StaticFoliageMesh = LoadObject<UStaticMesh>(nullptr, *InMeshPath, {}, LOAD_NoWarn | LOAD_Quiet);
+		USkeletalMesh* SkeletalFoliageMesh = LoadObject<USkeletalMesh>(nullptr, *InMeshPath, {}, LOAD_NoWarn | LOAD_Quiet);
  
 		int32 NumTriangles = 0;
 		
diff --git a/Engine/Plugins/Experimental/ProceduralVegetationEditor/Source/ProceduralVegetation/Private/Implementations/PVMaterialSettings.cpp b/Engine/Plugins/Experimental/ProceduralVegetationEditor/Source/ProceduralVegetation/Private/Implementations/PVMaterialSettings.cpp
--- a/Engine/Plugins/Experimental/ProceduralVegetationEditor/Source/ProceduralVegetation/Private/Implementations/PVMaterialSettings.cpp
+++ b/Engine/Plugins/Experimental/ProceduralVegetationEditor/Source/ProceduralVegetation/Private/Implementations/PVMaterialSettings.cpp
@@ -41,7 +41,7 @@ void FPVMaterialSettings::ApplyMaterialSettings(FManagedArrayCollection& Collect
 void FPVMaterialSettings::SetMaterial(FString TrunkMaterial, FVector2f UMinMax, int32 Index)
 {
 	FTrunkGenerationMaterialSetup MaterialSetup;
-	MaterialSetup.Material = LoadObject<UMaterialInterface>(nullptr, TrunkMaterial);
+	MaterialSetup.Material = LoadObject<UMaterialInterface>(nullptr, TrunkMaterial, {}, LOAD_NoWarn | LOAD_Quiet);
 	MaterialSetup.URange = FFloatRange(UMinMax.X, UMinMax.Y);
 
 	if (!MaterialSetups.IsValidIndex(Index))
diff --git a/Engine/Plugins/Experimental/ProceduralVegetationEditor/Source/ProceduralVegetation/Public/DataTypes/PVData.h b/Engine/Plugins/Experimental/ProceduralVegetationEditor/Source/ProceduralVegetation/Public/DataTypes/PVData.h
--- a/Engine/Plugins/Experimental/ProceduralVegetationEditor/Source/ProceduralVegetation/Public/DataTypes/PVData.h
+++ b/Engine/Plugins/Experimental/ProceduralVegetationEditor/Source/ProceduralVegetation/Public/DataTypes/PVData.h
@@ -196,8 +196,10 @@ public:
 	virtual bool SamplePoint(const FTransform& Transform, const FBox& Bounds, FPCGPoint& OutPoint, UPCGMetadata* OutMetadata) const override;
 	//~End UPCGSpatialData interface
 
-	const FManagedArrayCollection& GetCollection() const { return Collection; }
-	FManagedArrayCollection& GetCollection() { return Collection; }
+	const FManagedArrayCollection& GetCollection() const { return *Collection; }
+	FManagedArrayCollection& GetCollection() { return *Collection; }
+	TSharedPtr<FManagedArrayCollection>& GetSharedCollection() { return Collection; }
+	const TSharedPtr<FManagedArrayCollection>& GetSharedCollection() const { return Collection; }
 
 	void AddManagedResource(const UObject* InResource) 
 	{ 
@@ -234,7 +236,7 @@ private:
 
 protected:
 	
-	FManagedArrayCollection Collection;
+	TSharedPtr<FManagedArrayCollection> Collection;
 
 	UPROPERTY(Transient)
 	TArray<TObjectPtr<const UObject>> ManagedResources;
diff --git a/Engine/Plugins/Experimental/ProceduralVegetationEditor/Source/ProceduralVegetation/Public/Nodes/PVBaseSettings.h b/Engine/Plugins/Experimental/ProceduralVegetationEditor/Source/ProceduralVegetation/Public/Nodes/PVBaseSettings.h
--- a/Engine/Plugins/Experimental/ProceduralVegetationEditor/Source/ProceduralVegetation/Public/Nodes/PVBaseSettings.h
+++ b/Engine/Plugins/Experimental/ProceduralVegetationEditor/Source/ProceduralVegetation/Public/Nodes/PVBaseSettings.h
@@ -7,7 +7,7 @@
 #include "Helpers/PVUtilities.h"
 #include "PVBaseSettings.generated.h"
 
-UCLASS(BlueprintType, Abstract, HideCategories=(Debug), ClassGroup = (Procedural))
+UCLASS(BlueprintType, Abstract, HideCategories=(Debug, AssetInfo), ClassGroup = (Procedural))
 class UPVBaseSettings : public UPCGSettings
 {
 	GENERATED_BODY()
@@ -25,7 +25,7 @@ public:
 	virtual TMap<UObject*, FTransform> GetViewportObjects() const { return TMap<UObject*, FTransform>(); }
 	virtual void SetCurrentRenderType(TArray<EPVRenderType> InRenderTypes);
 	
-	FPVDebugSettings GetDebugVisualizationSettings() const { return DebugVisualizationSettings; }
+	const FPVDebugSettings& GetDebugVisualizationSettings() const { return DebugVisualizationSettings; }
 
 	virtual void PostEditChangeProperty(FPropertyChangedEvent& PropertyChangedEvent) override;
 #endif
diff --git a/Engine/Plugins/Experimental/ProceduralVegetationEditor/Source/ProceduralVegetation/Public/Params/PVExportParams.h b/Engine/Plugins/Experimental/ProceduralVegetationEditor/Source/ProceduralVegetation/Public/Params/PVExportParams.h
--- a/Engine/Plugins/Experimental/ProceduralVegetationEditor/Source/ProceduralVegetation/Public/Params/PVExportParams.h
+++ b/Engine/Plugins/Experimental/ProceduralVegetationEditor/Source/ProceduralVegetation/Public/Params/PVExportParams.h
@@ -45,12 +45,6 @@ struct PROCEDURALVEGETATION_API FPVExportParams
 {
 	GENERATED_BODY()
 
-	UPROPERTY()
-	bool bShouldExport = true;
-
-	UPROPERTY()
-	bool bShouldExportFoliage = true;
-
 	UPROPERTY(EditAnywhere, Category = "Asset", Meta = (ContentDir, DisplayName = "Content Browser Folder", Tooltip="Path to the folder where exported mesh will be saved\n\nPath to the folder where exported mesh will be saved."))
 	FDirectoryPath ContentBrowserFolder;
 
diff --git a/Engine/Plugins/Experimental/ProceduralVegetationEditor/Source/ProceduralVegetationEditor/Private/Exporter/PVExporter.cpp b/Engine/Plugins/Experimental/ProceduralVegetationEditor/Source/ProceduralVegetationEditor/Private/Exporter/PVExporter.cpp
new file mode 100644
--- /dev/null
+++ b/Engine/Plugins/Experimental/ProceduralVegetationEditor/Source/ProceduralVegetationEditor/Private/Exporter/PVExporter.cpp
@@ -0,0 +1,388 @@
+// Copyright Epic Games, Inc. All Rights Reserved.
+
+#include "PVExporter.h"
+
+#include "IMessageLogListing.h"
+#include "MessageLogModule.h"
+#include "PackagesDialog.h"
+#include "ProceduralVegetationEditorModule.h"
+
+#include "Facades/PVAttributesNames.h"
+
+#include "GeometryCollection/GeometryCollection.h"
+
+#include "Helpers/PVExportHelper.h"
+
+#include "Misc/MessageDialog.h"
+#include "Misc/ScopedSlowTask.h"
+#include "Misc/UObjectToken.h"
+
+#include "Styling/SlateTypes.h"
+
+#define LOCTEXT_NAMESPACE "PVExporter"
+
+void FPVOutputEntryStats::SetupStats(const FManagedArrayCollection& InCollection)
+{
+	if (InCollection.HasGroup(PV::GroupNames::PointGroup))
+	{
+		NumPoints = InCollection.NumElements(PV::GroupNames::PointGroup);
+	}
+	if (InCollection.HasGroup(PV::GroupNames::BranchGroup))
+	{
+		NumBranches = InCollection.NumElements(PV::GroupNames::BranchGroup);
+	}
+	if (InCollection.HasGroup(FGeometryCollection::VerticesGroup))
+	{
+		NumVertices = InCollection.NumElements(FGeometryCollection::VerticesGroup);
+	}
+	if (InCollection.HasGroup(FGeometryCollection::FacesGroup))
+	{
+		NumTris = InCollection.NumElements(FGeometryCollection::FacesGroup);
+	}
+	if (InCollection.HasGroup(PV::GroupNames::BonesGroup))
+	{
+		NumBones = InCollection.NumElements(PV::GroupNames::BonesGroup);
+	}
+	if (InCollection.HasGroup(PV::GroupNames::FoliageNamesGroup))
+	{
+		NumFoliage = InCollection.NumElements(PV::GroupNames::FoliageNamesGroup);
+	}
+	if (InCollection.HasGroup(PV::GroupNames::FoliageGroup))
+	{
+		NumFoliageInstances = InCollection.NumElements(PV::GroupNames::FoliageGroup);
+	}
+}
+
+void UPVExportEntry::Initialize(
+	const TSharedPtr<FManagedArrayCollection>& InCollection,
+	const TObjectPtr<UPCGNode> InNode,
+	const bool InIsNodeSelected,
+	const bool InHasFoliage
+)
+{
+	check(InNode);
+
+	bIsNodeSelected = InIsNodeSelected;
+	Node = InNode;
+	bExportFoliage = InHasFoliage;
+	Settings = CastChecked<UPVOutputSettings>(InNode->GetSettings());
+	OutputCollection = InCollection;
+
+	check(Settings && OutputCollection);
+	
+	Stats.SetupStats(*OutputCollection);
+}
+
+TSharedRef<FTokenizedMessage> FPVExporter::FScopedMessages::Error(const FText& ErrorMsg)
+{
+	TSharedRef<FTokenizedMessage> Message = FTokenizedMessage::Create(EMessageSeverity::Error, ErrorMsg);
+	Messages.Add(Message);
+	return Message;
+}
+
+TSharedRef<FTokenizedMessage> FPVExporter::FScopedMessages::Warning(const FText& ErrorMsg)
+{
+	TSharedRef<FTokenizedMessage> Message = FTokenizedMessage::Create(EMessageSeverity::Warning, ErrorMsg);
+	Messages.Add(Message);
+	return Message;
+}
+
+TSharedRef<FTokenizedMessage> FPVExporter::FScopedMessages::Info(const FText& ErrorMsg)
+{
+	TSharedRef<FTokenizedMessage> Message = FTokenizedMessage::Create(EMessageSeverity::Info, ErrorMsg);
+	Messages.Add(Message);
+	return Message;
+}
+
+bool FPVExporter::FScopedMessages::ContainsErrors() const
+{
+	for (TSharedRef<FTokenizedMessage> Message : Messages)
+	{
+		if (Message->GetSeverity() == EMessageSeverity::Error)
+		{
+			return true;
+		}
+	}
+
+	return false;
+}
+
+FPVExporter::FScopedMessages::~FScopedMessages()
+{
+	if (Messages.Num() > 0)
+	{
+		FMessageLogModule& MessageLogModule = FModuleManager::LoadModuleChecked<FMessageLogModule>("MessageLog");
+		const TSharedRef<IMessageLogListing> PVEditorMessageLogListing = MessageLogModule.GetLogListing(PVEditor::MessageLogName);
+
+		PVEditorMessageLogListing->ClearMessages();
+		PVEditorMessageLogListing->AddMessages(Messages);
+		PVEditorMessageLogListing->NotifyIfAnyMessages(
+			ContainsErrors()
+			? LOCTEXT("PVEMeshExportFailed", "Mesh Export Failed")
+			: LOCTEXT("PVEMeshExportComplete", "Mesh Export Completed"),
+			EMessageSeverity::Info,
+			true
+		);
+	}
+}
+
+
+FPVExporter::FPVExporter(const TObjectPtr<UProceduralVegetation> InProceduralVegetation, const TArray<TObjectPtr<UPVExportEntry>>& InExportEntries)
+	: ProceduralVegetation(InProceduralVegetation)
+{
+	ExportEntries = InExportEntries.FilterByPredicate(
+		[](const TObjectPtr<UPVExportEntry>& Entry)
+			{
+				return Entry->bExport;
+			}
+	);
+}
+
+bool FPVExporter::Export()
+{
+	if (!ValidateExportEntries())
+	{
+		return false;
+	}
+
+	if (!ShouldOverwriteAssets())
+	{
+		return false;
+	}
+
+	if (!FoliageValidation())
+	{
+		return false;
+	}
+
+	FScopedSlowTask ProgressSlowTask(ExportEntries.Num(), LOCTEXT("ExportAsset_SlowTask", "Exporting Asset(s) ..."));
+	ProgressSlowTask.MakeDialog(false);
+
+	bool bExportFailed = false;
+	for (int32 NodeIndex = 0; NodeIndex < ExportEntries.Num(); NodeIndex++)
+	{
+		const TObjectPtr<UPVExportEntry>& ExportEntry = ExportEntries[NodeIndex];
+		if (ProgressSlowTask.ShouldCancel())
+		{
+			break;
+		}
+
+		const FText AssetCountText = ExportEntries.Num() > 1
+			? FText::FromString(FString::Printf(TEXT(" [%d/%d]"), NodeIndex + 1, ExportEntries.Num()))
+			: FText::FromString("");
+
+		ProgressSlowTask.EnterProgressFrame(
+			1.f,
+			FText::Format(
+				LOCTEXT("ExportingAsset", "Exporting \"{0}\"{1}"),
+				FText::FromName(ExportEntry->Settings->ExportSettings.MeshName), AssetCountText
+			)
+		);
+
+		FScopedSlowTask ExportCollectionSlowTask(1.f, LOCTEXT("GeneratingMesh", "Generating Mesh ..."));
+		float PrevProgress = 0.f;
+
+		const auto OnStatusUpdated = [&](const FText& Stage, const float Progress)-> bool
+			{
+				const float WorkThisFrame = Progress - PrevProgress;
+				PrevProgress = Progress;
+
+				ExportCollectionSlowTask.FrameMessage = Stage;
+
+				ExportCollectionSlowTask.CompletedWork += ExportCollectionSlowTask.CurrentFrameScope;
+
+				const float WorkRemaining = ExportCollectionSlowTask.TotalAmountOfWork - ExportCollectionSlowTask.CompletedWork;
+				ExportCollectionSlowTask.CurrentFrameScope = FMath::Min(WorkRemaining, WorkThisFrame);
+				ExportCollectionSlowTask.ForceRefresh();
+
+				return !ProgressSlowTask.ShouldCancel();
+			};
+
+		TArray<FString> CreatedAssets;
+		const PV::Export::EExportResult ExportResult = PV::Export::ExportCollectionAsMesh(
+			ProceduralVegetation,
+			*ExportEntry->OutputCollection,
+			ExportEntry,
+			CreatedAssets,
+			OnStatusUpdated
+		);
+
+		if (ExportResult == PV::Export::EExportResult::Fail)
+		{
+			Messages.Error(FText::Format(
+				LOCTEXT("ExportFailed_UnknownError", "Failed to export mesh \"{0}\" due to unknown error, export aborted."),
+				FText::FromName(ExportEntry->Settings->ExportSettings.MeshName)
+			));
+			bExportFailed = true;
+			break;
+		}
+		else if (ExportResult == PV::Export::EExportResult::Cancelled)
+		{
+			FString MeshNames = "{ ";
+			for (int32 SkippedEntryIndex = NodeIndex; SkippedEntryIndex < ExportEntries.Num(); SkippedEntryIndex++)
+			{
+				MeshNames += ExportEntries[SkippedEntryIndex]->Settings->ExportSettings.MeshName.ToString();
+				if (SkippedEntryIndex != (ExportEntries.Num() - 1))
+				{
+					MeshNames += ", ";
+				}
+			}
+			MeshNames += " }";
+
+			Messages.Warning(
+				FText::Format(
+					LOCTEXT("ExportFailed_ExportCanceled", "Failed to export meshes \"{0}\", operation canceled by user."),
+					FText::FromString(MeshNames)
+				));
+			bExportFailed = true;
+			break;
+		}
+		else if (ExportResult == PV::Export::EExportResult::Skipped)
+		{
+			Messages.Info(FText::Format(
+				LOCTEXT("ExportSkipped", "Export of mesh \"{0}\" skipped due to conflicting asset at output location"),
+				FText::FromName(ExportEntry->Settings->ExportSettings.MeshName)
+			));
+		}
+		else
+		{
+			const TSharedRef<FTokenizedMessage> Message = Messages.Info(FText::Format(
+				LOCTEXT("ExportComplete", "Mesh exported successfully \"{0}\""),
+				FText::FromName(ExportEntry->Settings->ExportSettings.MeshName)
+			));
+
+			for (const FString& AssetPath : CreatedAssets)
+			{
+				Message->AddToken(FAssetNameToken::Create(AssetPath));
+			}
+		}
+	}
+
+	return !bExportFailed;
+}
+
+bool FPVExporter::ValidateExportEntries()
+{
+	bool bHasInvalidPath = false;
+	for (const TObjectPtr<UPVExportEntry>& ExportEntry : ExportEntries)
+	{
+		if (FString PathError; !ExportEntry->Settings->ExportSettings.Validate(PathError))
+		{
+			TSharedRef<FTokenizedMessage> Error = Messages.Error(
+				FText::Format(
+					LOCTEXT("ExportFailed_InvalidPath", "Invalid export path for node \"{0}\": \"{1}\""),
+					ExportEntry->Node->GetDefaultTitle(),
+					FText::FromString(PathError)
+				)
+			);
+			Error->AddToken(FUObjectToken::Create(ExportEntry->Settings));
+			bHasInvalidPath = true;
+		}
+	}
+
+	return !bHasInvalidPath;
+}
+
+bool FPVExporter::ShouldOverwriteAssets()
+{
+	TArray<FString> ExportPaths;
+	ExportPaths.Reserve(ExportEntries.Num() * 2);
+
+	for (const TObjectPtr<UPVExportEntry>& ExportEntry : ExportEntries)
+	{
+		const FPVExportParams& ExportParams = ExportEntry->Settings->ExportSettings;
+		if (ExportParams.ReplacePolicy != EPVAssetReplacePolicy::Replace)
+		{
+			continue;
+		}
+		ExportPaths.Add(ExportParams.GetOutputMeshPackagePath());
+		if (ExportParams.ExportMeshType == EPVExportMeshType::SkeletalMesh)
+		{
+			ExportPaths.Add(ExportParams.GetOutputSkeletonPackagePath());
+			if (ExportParams.IsCollisionEnable())
+			{
+				ExportPaths.Add(ExportParams.GetOutputPhysicsAssetPackagePath());
+			}
+		}
+	}
+
+	TArray<UPackage*> ExistingPackages;
+	ExistingPackages.Reserve(ExportPaths.Num());
+	for (const FString& AssetPath : ExportPaths)
+	{
+		if (UPackage* ExistingPackage = LoadPackage(nullptr, *AssetPath, LOAD_NoWarn | LOAD_Quiet))
+		{
+			ExistingPackages.Add(ExistingPackage);
+		}
+	}
+
+	if (ExistingPackages.Num() == 0)
+	{
+		return true;
+	}
+
+	FPackagesDialogModule& PackagesDialogModule = FModuleManager::LoadModuleChecked<FPackagesDialogModule>(TEXT("PackagesDialog"));
+	PackagesDialogModule.CreatePackagesDialog(
+		LOCTEXT("PackagesDialog_OverwriteFiles_Title", "WARNING: Assets will be overwritten"),
+		LOCTEXT("PackagesDialog_OverwriteFiles_Message", "The following assets will be overwritten. Would you like to continue?"),
+		/*InReadOnly=*/true
+	);
+	PackagesDialogModule.AddButton(
+		DRT_Save,
+		LOCTEXT("PackagesDialog_OverwriteFiles_ContinueButton", "Continue"),
+		LOCTEXT("PackagesDialog_OverwriteFiles_ButtonTip", "Continues the export, overwrites any existing assets")
+	);
+	PackagesDialogModule.AddButton(
+		DRT_Cancel,
+		LOCTEXT("PackagesDialog_OverwriteFiles_CancelButton", "Cancel"),
+		LOCTEXT("CancelDeleteButtonTip", "Cancel export")
+	);
+
+	for (UPackage* Package : ExistingPackages)
+	{
+		PackagesDialogModule.AddPackageItem(Package, ECheckBoxState::Checked);
+	}
+
+	const EDialogReturnType UserResponse = PackagesDialogModule.ShowPackagesDialog();
+	return UserResponse == DRT_Save;
+}
+
+bool FPVExporter::FoliageValidation()
+{
+	FString NodeNames;
+	for (const TObjectPtr<UPVExportEntry>& ExportEntry : ExportEntries)
+	{
+		if (!ExportEntry->bExportFoliage)
+		{
+			if (!NodeNames.IsEmpty())
+			{
+				NodeNames += TEXT(", ");
+			}
+			NodeNames += "'" + ExportEntry->Node->GetAuthoredTitleName().ToString() + "'";
+		}
+	}
+
+	if (!NodeNames.IsEmpty())
+	{
+		const FString Message = NodeNames +
+			TEXT(" does not use input from foliage distributor node. Do you want to export with default foliage distribution?");
+		const EAppReturnType::Type Ret = FMessageDialog::Open(EAppMsgType::YesNoCancel, FText::FromString(Message));
+		if (Ret != EAppReturnType::Cancel)
+		{
+			const bool bShouldExportFoliage = Ret == EAppReturnType::Yes;
+			for (const TObjectPtr<UPVExportEntry>& ExportEntry : ExportEntries)
+			{
+				ExportEntry->bExportFoliage = !ExportEntry->bExportFoliage
+					? bShouldExportFoliage
+					: true;
+			}
+		}
+		else
+		{
+			return false;
+		}
+	}
+	return true;
+}
+
+#undef LOCTEXT_NAMESPACE
diff --git a/Engine/Plugins/Experimental/ProceduralVegetationEditor/Source/ProceduralVegetationEditor/Private/Exporter/PVExporter.h b/Engine/Plugins/Experimental/ProceduralVegetationEditor/Source/ProceduralVegetationEditor/Private/Exporter/PVExporter.h
new file mode 100644
--- /dev/null
+++ b/Engine/Plugins/Experimental/ProceduralVegetationEditor/Source/ProceduralVegetationEditor/Private/Exporter/PVExporter.h
@@ -0,0 +1,110 @@
+// Copyright Epic Games, Inc. All Rights Reserved.
+
+#pragma once
+
+#include "Nodes/PVBaseSettings.h"
+#include "Nodes/PVOutputSettings.h"
+
+#include "PVExporter.generated.h"
+
+class UProceduralVegetation;
+
+
+
+USTRUCT()
+struct FPVOutputEntryStats
+{
+	GENERATED_BODY()
+	
+	UPROPERTY(VisibleAnywhere, Category="Stats", meta=(EditCondition="NumPoints != -1"))
+	int32 NumPoints = -1;
+
+	UPROPERTY(VisibleAnywhere, Category="Stats", meta=(EditCondition="NumBranches != -1"))
+	int32 NumBranches = -1;
+
+	UPROPERTY(VisibleAnywhere, Category="Stats", DisplayName="Num Vertices (Trunk)", meta=(EditCondition="NumVertices != -1"))
+	int32 NumVertices = -1;
+
+	UPROPERTY(VisibleAnywhere, Category="Stats", DisplayName="Num Triangles (Trunk)", meta=(EditCondition="NumTris != -1"))
+	int32 NumTris = -1;
+
+	UPROPERTY(VisibleAnywhere, Category="Stats", meta=(EditCondition="NumBones != -1"))
+	int32 NumBones = -1;
+	
+	UPROPERTY(VisibleAnywhere, Category="Stats", meta=(EditCondition="NumFoliage != -1"))
+	int32 NumFoliage = -1;
+
+	UPROPERTY(VisibleAnywhere, Category="Stats", meta=(EditCondition="NumFoliageInstances != -1"))
+	int32 NumFoliageInstances = -1;
+	
+	void SetupStats(const FManagedArrayCollection& InCollection);
+};
+
+UCLASS()
+class UPVExportEntry : public UObject
+{
+	GENERATED_BODY()
+
+public:
+	UPROPERTY(Transient)
+	bool bExport = true;
+
+	UPROPERTY(Transient)
+	bool bIsNodeSelected = false;
+
+	UPROPERTY(Transient)
+	TObjectPtr<UPCGNode> Node = nullptr;
+
+	UPROPERTY(VisibleAnywhere, Category="Stats", DisplayName="Has Foliage")
+	bool bExportFoliage = true;
+	
+	UPROPERTY(VisibleAnywhere, Category="Stats", meta=(ShowOnlyInnerProperties))
+	FPVOutputEntryStats Stats;
+
+	UPROPERTY(EditAnywhere, Category="Export Settings", meta=(ShowOnlyInnerProperties))
+	TObjectPtr<UPVOutputSettings> Settings = nullptr;
+
+	TSharedPtr<FManagedArrayCollection> OutputCollection;
+
+	void Initialize(
+		const TSharedPtr<FManagedArrayCollection>& InCollection,
+		const TObjectPtr<UPCGNode> InNode,
+		bool InIsNodeSelected = false,
+		bool InHasFoliage = true
+	);
+};
+
+class FPVExporter
+{
+	struct FScopedMessages
+	{
+		TSharedRef<FTokenizedMessage> Error(const FText& ErrorMsg);
+		TSharedRef<FTokenizedMessage> Warning(const FText& ErrorMsg);
+		TSharedRef<FTokenizedMessage> Info(const FText& ErrorMsg);
+
+		~FScopedMessages();
+
+	private:
+		bool ContainsErrors() const;
+
+		TArray<TSharedRef<FTokenizedMessage>> Messages;
+	};
+
+public:
+	FPVExporter(TObjectPtr<UProceduralVegetation> InProceduralVegetation, const TArray<TObjectPtr<UPVExportEntry>>& InExportEntries);
+
+	bool Export();
+
+private:
+	bool ValidateExportEntries();
+	bool ShouldOverwriteAssets();
+
+	bool FoliageValidation();
+
+private:
+	TObjectPtr<UProceduralVegetation> ProceduralVegetation = nullptr;
+
+	TArray<TObjectPtr<UPVExportEntry>> ExportEntries;
+
+	FScopedMessages Messages;
+};
diff --git a/Engine/Plugins/Experimental/ProceduralVegetationEditor/Source/ProceduralVegetationEditor/Private/Helpers/PVExportHelper.cpp b/Engine/Plugins/Experimental/ProceduralVegetationEditor/Source/ProceduralVegetationEditor/Private/Helpers/PVExportHelper.cpp
--- a/Engine/Plugins/Experimental/ProceduralVegetationEditor/Source/ProceduralVegetationEditor/Private/Helpers/PVExportHelper.cpp
+++ b/Engine/Plugins/Experimental/ProceduralVegetationEditor/Source/ProceduralVegetationEditor/Private/Helpers/PVExportHelper.cpp
@@ -34,7 +34,9 @@
  
 #include "Engine/SkeletalMesh.h"
 #include "Engine/StaticMesh.h"
- 
+
+#include "Exporter/PVExporter.h"
+
 #include "Facades/PVBranchFacade.h"
 #include "Facades/PVPlantFacade.h"
 #include "Facades/PVPointFacade.h"
@@ -309,7 +311,7 @@ namespace PV::Export::Internal
 		constexpr bool bCenterPivot = false;
 		FTransform TransformOut;
 		FDynamicMesh3 CombinedMesh;
-		ConvertGeometryCollectionToDynamicMesh(CombinedMesh, TransformOut, bCenterPivot, *GeomCollection, true, BoneTransforms.GetConstArray(), true, TransformIndices);
+		ConvertGeometryCollectionToDynamicMesh(CombinedMesh, TransformOut, bCenterPivot, *GeomCollection, false, BoneTransforms.GetConstArray(), true, TransformIndices);
  
 		NewMesh->SetMesh(CombinedMesh);
  
@@ -357,7 +359,7 @@ namespace PV::Export::Internal
 			const TManagedArray<FString>& MaterialPaths = Collection.GetAttribute<FString>("MaterialPath", FGeometryCollection::MaterialGroup);
 			for (const FString& Path : MaterialPaths)
 			{
-				Materials.Add(LoadObject<UMaterialInterface>(nullptr, *Path));
+				Materials.Add(LoadObject<UMaterialInterface>(nullptr, *Path, {}, LOAD_NoWarn | LOAD_Quiet));
 			}
 		}
 		return Materials;
@@ -442,7 +444,7 @@ namespace PV::Export::Internal
  
 				const FString& FoliageName = Result.FoliageName;
  
-				if (const TObjectPtr<UStaticMesh> StaticMesh = LoadObject<UStaticMesh>(nullptr, *FoliageName))
+				if (const TObjectPtr<UStaticMesh> StaticMesh = LoadObject<UStaticMesh>(nullptr, *FoliageName, {}, LOAD_NoWarn | LOAD_Quiet))
 				{
 					FText OutErrorMessage;
 					UE::Conversion::FStaticMeshConversionOptions Options;
@@ -459,7 +461,7 @@ namespace PV::Export::Internal
  
 					Result.OutMeshPath = StaticMesh->GetFullName();
 				}
-				else if (const TObjectPtr<USkeletalMesh> SkeletalMesh = LoadObject<USkeletalMesh>(nullptr, *FoliageName))
+				else if (const TObjectPtr<USkeletalMesh> SkeletalMesh = LoadObject<USkeletalMesh>(nullptr, *FoliageName, {}, LOAD_NoWarn | LOAD_Quiet))
 				{
 					if (!SkeletalMeshToDynamicMesh(SkeletalMesh, 0, Result.OutDynamicMesh))
 					{
@@ -689,7 +691,7 @@ namespace PV::Export::Internal
 				continue;
 			}
 			
-			TObjectPtr<UStaticMesh> PartStaticMesh = LoadObject<UStaticMesh>(nullptr, *FoliageInstanceObjectPath);
+			TObjectPtr<UStaticMesh> PartStaticMesh = LoadObject<UStaticMesh>(nullptr, *FoliageInstanceObjectPath, {}, LOAD_NoWarn | LOAD_Quiet);
 			if (PartStaticMesh)
 			{
 				AddStaticMeshToAssembly(FoliageInstanceObjectPath, PartStaticMesh);
@@ -704,14 +706,14 @@ namespace PV::Export::Internal
 			const FString StaticMeshPackagePath = FPaths::Combine(LongFoliagePackagePath, StaticMeshName);
 			const FString StaticMeshObjectPath = FString::Printf(TEXT("%s.%s"), *StaticMeshPackagePath, *StaticMeshName);
  
-			PartStaticMesh = LoadObject<UStaticMesh>(nullptr, StaticMeshObjectPath);
+			PartStaticMesh = LoadObject<UStaticMesh>(nullptr, StaticMeshObjectPath, {}, LOAD_NoWarn | LOAD_Quiet);
 			if (PartStaticMesh)
 			{
 				AddStaticMeshToAssembly(FoliageInstanceObjectPath, PartStaticMesh);
 				continue;
 			}
  
-			const TObjectPtr<USkeletalMesh> PartSkeletalMesh = LoadObject<USkeletalMesh>(nullptr, *FoliageInstanceObjectPath);
+			const TObjectPtr<USkeletalMesh> PartSkeletalMesh = LoadObject<USkeletalMesh>(nullptr, *FoliageInstanceObjectPath, {}, LOAD_NoWarn | LOAD_Quiet);
 			if (!PartSkeletalMesh)
 			{
 				UE_LOG(LogProceduralVegetationEditor, Warning, TEXT("Could not find foliage instance %s"), *FoliageInstanceObjectPath);
@@ -793,7 +795,7 @@ namespace PV::Export::Internal
 				continue;
 			}
 			
-			TObjectPtr<USkeletalMesh> PartSkeletalMesh = LoadObject<USkeletalMesh>(nullptr, *FoliageInstanceObjectPath);
+			TObjectPtr<USkeletalMesh> PartSkeletalMesh = LoadObject<USkeletalMesh>(nullptr, *FoliageInstanceObjectPath, {}, LOAD_NoWarn | LOAD_Quiet);
 			if (PartSkeletalMesh)
 			{
 				AddSkeletalMeshToAssembly(FoliageInstanceObjectPath, PartSkeletalMesh);
@@ -812,8 +814,8 @@ namespace PV::Export::Internal
 			const FString SkeletonPackagePath = FPaths::Combine(LongFoliagePackagePath, SkeletonName);
 			const FString SkeletonObjectPath = FString::Printf(TEXT("%s.%s"), *SkeletonPackagePath, *SkeletonName);
  
-			PartSkeletalMesh = LoadObject<USkeletalMesh>(nullptr, SkeletalMeshObjectPath);
-			TObjectPtr<USkeleton> Skeleton = LoadObject<USkeleton>(nullptr, SkeletonObjectPath);
+			PartSkeletalMesh = LoadObject<USkeletalMesh>(nullptr, SkeletalMeshObjectPath, {}, LOAD_NoWarn | LOAD_Quiet);
+			TObjectPtr<USkeleton> Skeleton = LoadObject<USkeleton>(nullptr, SkeletonObjectPath, {}, LOAD_NoWarn | LOAD_Quiet);
 			if (PartSkeletalMesh && Skeleton)
 			{
 				AddSkeletalMeshToAssembly(FoliageInstanceObjectPath, PartSkeletalMesh);
@@ -836,7 +838,7 @@ namespace PV::Export::Internal
  
 			if (!PartSkeletalMesh)
 			{
-				const TObjectPtr<UStaticMesh> PartStaticMesh = LoadObject<UStaticMesh>(nullptr, *FoliageInstanceObjectPath);
+				const TObjectPtr<UStaticMesh> PartStaticMesh = LoadObject<UStaticMesh>(nullptr, *FoliageInstanceObjectPath, {}, LOAD_NoWarn | LOAD_Quiet);
 				if (!PartStaticMesh)
 				{
 					UE_LOG(LogProceduralVegetationEditor, Warning, TEXT("Could not find foliage instance %s"), *FoliageInstanceObjectPath);
@@ -1674,12 +1676,14 @@ namespace PV::Export
 	EExportResult ExportCollectionAsMesh(
 		const TObjectPtr<UProceduralVegetation> InProceduralVegetation,
 		const FManagedArrayCollection& Collection,
-		const FPVExportParams& ExportParams,
+		const TObjectPtr<UPVExportEntry>& ExportEntry,
 		TArray<FString>& OutCreatedAssets,
 		const FStatusReportCallback& StatusReportCallback
 	)
 	{
 		TRACE_CPUPROFILER_EVENT_SCOPE(PV::Export::ExportCollectionAsMesh);
+		
+		const FPVExportParams& ExportParams = ExportEntry->Settings->ExportSettings;
  
 		TArray<UPackage*> ModifiedPackages;
  
@@ -1737,7 +1741,7 @@ namespace PV::Export
 		}
 		
 		const bool bCreateNaniteFoliage = ExportParams.bCreateNaniteFoliage && NaniteAssembliesSupported();
-		const bool bShouldExportFoliage = ExportParams.bShouldExportFoliage;
+		const bool bShouldExportFoliage = ExportEntry->bExportFoliage;
  
 		if (ExportParams.ExportMeshType == EPVExportMeshType::StaticMesh)
 		{
diff --git a/Engine/Plugins/Experimental/ProceduralVegetationEditor/Source/ProceduralVegetationEditor/Private/Helpers/PVExportHelper.h b/Engine/Plugins/Experimental/ProceduralVegetationEditor/Source/ProceduralVegetationEditor/Private/Helpers/PVExportHelper.h
--- a/Engine/Plugins/Experimental/ProceduralVegetationEditor/Source/ProceduralVegetationEditor/Private/Helpers/PVExportHelper.h
+++ b/Engine/Plugins/Experimental/ProceduralVegetationEditor/Source/ProceduralVegetationEditor/Private/Helpers/PVExportHelper.h
@@ -5,21 +5,26 @@
 #include "Facades/PVBoneFacade.h"
 #include "Engine/EngineTypes.h"
 
+#include "Exporter/PVExporter.h"
+
 enum class EPVCollisionGeneration : uint8;
+
+class UDynamicMesh;
 class UInstancedStaticMeshComponent;
-class UStaticMesh;
-class USkeletalMesh;
-class UPackage;
-class UProceduralVegetation;
 class UMaterialInterface;
+class UPVExportEntry;
 class UPVWindSettings;
+class UPackage;
+class UProceduralVegetation;
+class USkeletalMesh;
+class UStaticMesh;
+
 struct FGeometryScriptCopyMeshToAssetOptions;
 struct FManagedArrayCollection;
 struct FPVExportParams;
 
-class UDynamicMesh;
-namespace UE { namespace Geometry { class FDynamicMesh3; } }
-using UE::Geometry::FDynamicMesh3;
+using namespace UE::Geometry;
+namespace UE::Geometry { class FDynamicMesh3; }
 
 namespace PV::Export
 {
@@ -36,7 +41,7 @@ namespace PV::Export
 	EExportResult ExportCollectionAsMesh(
 		const TObjectPtr<UProceduralVegetation> InProceduralVegetation,
 		const FManagedArrayCollection& Collection,
-		const FPVExportParams& ExportParams,
+		const TObjectPtr<UPVExportEntry>& ExportEntry,
 		TArray<FString>& OutCreatedAssets,
 		const FStatusReportCallback& StatusReportCallback
 	);
diff --git a/Engine/Plugins/Experimental/ProceduralVegetationEditor/Source/ProceduralVegetationEditor/Private/Helpers/PVGraphHelpers.cpp b/Engine/Plugins/Experimental/ProceduralVegetationEditor/Source/ProceduralVegetationEditor/Private/Helpers/PVGraphHelpers.cpp
new file mode 100644
--- /dev/null
+++ b/Engine/Plugins/Experimental/ProceduralVegetationEditor/Source/ProceduralVegetationEditor/Private/Helpers/PVGraphHelpers.cpp
@@ -0,0 +1,77 @@
+// Copyright Epic Games, Inc. All Rights Reserved.
+
+#include "PVGraphHelpers.h"
+
+#include "PCGEdge.h"
+
+#include "Nodes/PVBaseSettings.h"
+#include "Dataflow/DataflowCollectionSpreadSheetWidget.h"
+
+#include "DataTypes/PVFoliageMeshData.h"
+#include "DataTypes/PVGrowthData.h"
+#include "DataTypes/PVMeshData.h"
+
+namespace PV::Graph
+{
+	UPVBaseSettings* GetNodeSettings(const UPCGNode* InNode)
+	{
+		return InNode
+			? Cast<UPVBaseSettings>(InNode->GetSettings())
+			: nullptr;
+	}
+
+	UPCGPin* GetInPinFromNode(UPCGNode* InNode)
+	{
+		return InNode
+			? InNode->GetInputPin(PCGPinConstants::DefaultInputLabel)
+			: nullptr;
+	}
+
+	UPCGPin* GetOutPinFromNode(UPCGNode* InNode)
+	{
+		return InNode
+			? InNode->GetOutputPin(PCGPinConstants::DefaultOutputLabel)
+			: nullptr;
+	}
+
+	UPCGPin* GetUpstreamPin(UPCGPin* InPin)
+	{
+		return (InPin && InPin->Edges.Num())
+			? InPin->Edges[0]->InputPin
+			: nullptr;
+	}
+
+	UPCGPin* GetDownstreamPin(UPCGPin* InPin)
+	{
+		return (InPin && InPin->Edges.Num())
+			? InPin->Edges[0]->OutputPin
+			: nullptr;
+	}
+
+	TArray<TObjectPtr<UPCGPin>> GetAllPVPinsFromNode(const UPCGNode* InNode)
+	{
+		static TArray PinTypes = {
+			FPCGDataTypeIdentifier{FPVDataTypeInfoGrowth::AsId()},
+			FPCGDataTypeIdentifier{FPVDataTypeInfoMesh::AsId()},
+			FPCGDataTypeIdentifier{FPVDataTypeInfoFoliageMesh::AsId()}
+		};
+
+		TArray<TObjectPtr<UPCGPin>> AllPins;
+		for (TObjectPtr<UPCGPin> OutputPin : InNode->GetOutputPins())
+		{
+			if (PinTypes.Contains(OutputPin->GetCurrentTypesID()))
+			{
+				AllPins.Add(OutputPin);
+			}
+		}
+		for (TObjectPtr<UPCGPin> InputPin : InNode->GetInputPins())
+		{
+			if (PinTypes.Contains(InputPin->GetCurrentTypesID()))
+			{
+				AllPins.Add(InputPin);
+			}
+		}
+
+		return AllPins;
+	}
+}
diff --git a/Engine/Plugins/Experimental/ProceduralVegetationEditor/Source/ProceduralVegetationEditor/Private/Helpers/PVGraphHelpers.h b/Engine/Plugins/Experimental/ProceduralVegetationEditor/Source/ProceduralVegetationEditor/Private/Helpers/PVGraphHelpers.h
new file mode 100644
--- /dev/null
+++ b/Engine/Plugins/Experimental/ProceduralVegetationEditor/Source/ProceduralVegetationEditor/Private/Helpers/PVGraphHelpers.h
@@ -0,0 +1,29 @@
+// Copyright Epic Games, Inc. All Rights Reserved.
+
+#pragma once
+
+#include "CoreMinimal.h"
+#include "PCGPin.h"
+
+class SCollectionSpreadSheetWidget;
+class UPCGNode;
+class UPVBaseSettings;
+
+namespace PV::Graph
+{
+	UPVBaseSettings* GetNodeSettings(const UPCGNode* InNode);
+
+	UPCGPin* GetInPinFromNode(UPCGNode* InNode);
+	UPCGPin* GetOutPinFromNode(UPCGNode* InNode);
+	UPCGPin* GetUpstreamPin(UPCGPin* InPin);
+	UPCGPin* GetDownstreamPin(UPCGPin* InPin);
+
+	TArray<TObjectPtr<UPCGPin>> GetAllPVPinsFromNode(const UPCGNode* InNode);
+
+	template <typename PVDataType>
+	requires std::is_base_of_v<FPCGDataTypeInfo, PVDataType>
+	bool IsPinOfType(UPCGPin* InPin)
+	{
+		return InPin && InPin->GetCurrentTypesID() == FPCGDataTypeIdentifier{PVDataType::AsId()};
+	}
+}
diff --git a/Engine/Plugins/Experimental/ProceduralVegetationEditor/Source/ProceduralVegetationEditor/Private/PVEditor.cpp b/Engine/Plugins/Experimental/ProceduralVegetationEditor/Source/ProceduralVegetationEditor/Private/PVEditor.cpp
--- a/Engine/Plugins/Experimental/ProceduralVegetationEditor/Source/ProceduralVegetationEditor/Private/PVEditor.cpp
+++ b/Engine/Plugins/Experimental/ProceduralVegetationEditor/Source/ProceduralVegetationEditor/Private/PVEditor.cpp
@@ -2,9 +2,12 @@
 
 #include "PVEditor.h"
 
+#include "AdvancedPreviewScene.h"
+#include "AssetEditorModeManager.h"
+#include "ContextObjectStore.h"
+#include "EditorModeManager.h"
 #include "EngineAnalytics.h"
-#include "IMessageLogListing.h"
-#include "MessageLogModule.h"
+#include "InteractiveToolManager.h"
 #include "PackagesDialog.h"
 #include "PCGDefaultExecutionSource.h"
 #include "PCGEditorCommands.h"
@@ -13,15 +16,16 @@
 #include "ProceduralVegetationEditorModule.h"
 #include "PVEditorCommands.h"
 #include "PVEditorSchema.h"
-#include "PVEditorSettings.h"
 #include "SPVEditorViewport.h"
 #include "ToolMenus.h"
-#include "ToolMenuSection.h"
 
 #include "Dataflow/DataflowCollectionSpreadSheetWidget.h"
 
 #include "DataTypes/PVData.h"
 #include "DataTypes/PVFoliageMeshData.h"
+#include "DataTypes/PVMeshData.h"
+
+#include "Exporter/PVExporter.h"
 
 #include "Facades/PVAttributesNames.h"
 
@@ -29,6 +33,7 @@
 
 #include "Helpers/PVAnalyticsHelper.h"
 #include "Helpers/PVExportHelper.h"
+#include "Helpers/PVGraphHelpers.h"
 #include "Helpers/PVUtilities.h"
 
 #include "Misc/MessageDialog.h"
@@ -46,25 +51,29 @@
 
 #define LOCTEXT_NAMESPACE "PVEditor"
 
-const FName FPVEditor::CollectionSpreadSheetTabId (TEXT("PVEditor_CollectionSpreadSheet"));
+const FName FPVEditor::CollectionSpreadSheetTabId(TEXT("PVEditor_CollectionSpreadSheet"));
+
+using namespace PV::Graph;
 
 FPVEditor::FPVEditor()
-{
-}
+{}
 
-void FPVEditor::Initialize(const EToolkitMode::Type InMode, const TSharedPtr<class IToolkitHost>& InToolkitHost, UProceduralVegetation* InProceduralVegetation)
+void FPVEditor::Initialize(
+	const EToolkitMode::Type InMode,
+	const TSharedPtr<IToolkitHost>& InToolkitHost,
+	UProceduralVegetation* InProceduralVegetation
+)
 {
 	PRAGMA_DISABLE_DEPRECATION_WARNINGS
 	bForceRefreshAttributeEvenIfClosed = true;
 	PRAGMA_ENABLE_DEPRECATION_WARNINGS
-	
+
 	check(InProceduralVegetation);
 
 	UE_LOG(LogProceduralVegetationEditor, Log, TEXT("PVEditor initialized for %s"), *InProceduralVegetation->GetFName().ToString());
 	ProceduralVegetationBeingEdited = InProceduralVegetation;
-	
-	UProceduralVegetationGraph* ProceduralVegetationGraph = Cast<UProceduralVegetationGraph>(InProceduralVegetation->GetGraph());
 
+	UProceduralVegetationGraph* ProceduralVegetationGraph = Cast<UProceduralVegetationGraph>(InProceduralVegetation->GetGraph());
 	check(ProceduralVegetationGraph);
 
 	// Initialize widgets before calling base
@@ -73,22 +82,19 @@ void FPVEditor::Initialize(const EToolkitMode::Type InMode, const TSharedPtr<cla
 	{
 		CollectionSpreadSheetWidget = CreateCollectionSpreadSheetWidget();
 	}
-	
+
 	FPCGEditor::Initialize(InMode, InToolkitHost, ProceduralVegetationGraph, InProceduralVegetation);
 
-	check(GEngine);
-	
-	UPCGEngineSubsystem* Subsystem = GEngine->GetEngineSubsystem<UPCGEngineSubsystem>();
+	GetInspectionDataManager().OnInspectedStackChangedDelegate.AddRaw(this, &FPVEditor::OnNodeInspected);
 
+	UPCGEngineSubsystem* Subsystem = GEngine->GetEngineSubsystem<UPCGEngineSubsystem>();
 	check(Subsystem);
 
-	Subsystem->GetOnPCGSourceGenerationDone().AddRaw(this, &FPVEditor::OnGraphExecuted);
-	
 	const FPCGDefaultExecutionSourceParams Params
 	{
 		.GraphInterface = ProceduralVegetationGraph
 	};
-	
+
 	ExecutionSource = Subsystem->CreateExecutionSource(Params);
 
 	// Set the stack
@@ -101,23 +107,19 @@ void FPVEditor::Initialize(const EToolkitMode::Type InMode, const TSharedPtr<cla
 	ExecutionSource->Generate();
 
 	// Select the first available output
-	for (auto EdNode : GetMainEditorGraph()->Nodes)
+	for (TObjectPtr<UEdGraphNode> EdNode : GetMainEditorGraph()->Nodes)
 	{
 		if (const UPCGEditorGraphNodeBase* PCGEdNode = Cast<UPCGEditorGraphNodeBase>(EdNode))
 		{
-			if (const UPCGNode* PCGNode = PCGEdNode->GetPCGNode())
+			if (PCGEdNode->GetPCGNode()->GetSettings()->IsA<UPVOutputSettings>())
 			{
-				if(PCGNode->GetSettings()->IsA<UPVOutputSettings>())
-				{
-					GetMainEditorGraph()->SelectNodeSet({PCGEdNode}, true);
-					break;
-				}
+				GetMainEditorGraph()->SelectNodeSet({PCGEdNode}, true);
+				break;
 			}
 		}
 	}
 
 	SessionStartTime = FDateTime::Now();
-
 	if (FEngineAnalytics::IsAvailable())
 	{
 		PV::Analytics::SendSessionStartedEvent();
@@ -136,6 +138,11 @@ void FPVEditor::AddReferencedObjects(FReferenceCollector& Collector)
 	Collector.AddReferencedObject(ExecutionSource);
 }
 
+void FPVEditor::CreateEditorModeManager()
+{
+	EditorModeManager = MakeShared<FAssetEditorModeManager>();
+}
+
 FText FPVEditor::GetToolkitName() const
 {
 	return GetLabelForObject(ProceduralVegetationBeingEdited);
@@ -156,38 +163,33 @@ FText FPVEditor::GetToolkitToolTipText() const
 	return GetToolTipTextForObject(ProceduralVegetationBeingEdited);
 }
 
-FLinearColor FPVEditor::GetWorldCentricTabColorScale() const
+FString FPVEditor::GetWorldCentricTabPrefix() const
 {
-	return FLinearColor::Blue;
+	return TEXT("Procedural Vegetation Editor ");
 }
 
-FString FPVEditor::GetWorldCentricTabPrefix() const
+FLinearColor FPVEditor::GetWorldCentricTabColorScale() const
 {
-	return TEXT("Procedural Vegetation Editor ");
+	return FLinearColor::Blue;
 }
 
 void FPVEditor::OnClose()
 {
 	UE_LOG(LogProceduralVegetationEditor, Log, TEXT("PVEditor Closed"));
-	
-	double TotalSeconds = (FDateTime::Now() - SessionStartTime).GetTotalSeconds();
-	PV::Analytics::SendSessionEndedEvent(TotalSeconds);
-	
-	FPCGEditor::OnClose();
 
-	if (ExecutionSource)
+	if (FEngineAnalytics::IsAvailable())
 	{
-		ExecutionSource->SetGraphInterface(nullptr);
+		const double TotalSeconds = (FDateTime::Now() - SessionStartTime).GetTotalSeconds();
+		PV::Analytics::SendSessionEndedEvent(TotalSeconds);
 	}
-}
 
-void FPVEditor::RegisterExtraTabFactories(FWorkflowAllowedTabSet& TabSet)
-{
-	if (CollectionSpreadSheetWidget.IsValid())
+	GetInspectionDataManager().OnInspectedStackChangedDelegate.RemoveAll(this);
+
+	FPCGEditor::OnClose();
+
+	if (ExecutionSource)
 	{
-		FPCGGenericTabFactoryParams Params(CollectionSpreadSheetTabId, SharedThis<FPCGEditor>(this), [this]() { return CollectionSpreadSheetWidget.ToSharedRef(); });
-		Params.Label = LOCTEXT("CollectionSpreadsheetTab", "Collection Spreadsheet");
-		TabSet.RegisterFactory(MakeShared<FPCGGenericTabFactory>(Params));
+		ExecutionSource->SetGraphInterface(nullptr);
 	}
 }
 
@@ -196,93 +198,6 @@ IPCGBaseSubsystem* FPVEditor::GetSubsystem() const
 	return UPCGEngineSubsystem::Get();
 }
 
-TAttribute<FGraphAppearanceInfo> FPVEditor::GetAppearanceInfo() const
-{
-	FGraphAppearanceInfo AppearanceInfo;
-	AppearanceInfo.CornerText = LOCTEXT("PVEditorCornerText", "Procedural Vegetation");
-
-	return AppearanceInfo;
-}
-
-TSharedRef<FTabManager::FLayout> FPVEditor::GetDefaultLayout() const
-{
-	return FTabManager::NewLayout("Standalone_PVGraphEditor_DefaultLayout_v1.0.6")
-	->AddArea // Main PCG Graph Editor Area
-	(
-		FTabManager::NewPrimaryArea()->SetOrientation(Orient_Vertical)
-		->Split // Top Section - Graph, Data Viewport, HLSL Source Editor, and Details View
-		(
-			FTabManager::NewSplitter()->SetOrientation(Orient_Horizontal)->SetSizeCoefficient(0.9f)
-			->SetSizeCoefficient(0.40f)
-			->Split
-			(
-				FTabManager::NewSplitter()->SetOrientation(Orient_Vertical)
-				->SetSizeCoefficient(0.1f)
-				->Split // Viewport
-				(
-					FTabManager::NewStack()
-					->SetSizeCoefficient(0.75f)
-					->AddTab(PCGEditorTabs::ViewportID[0], ETabState::OpenedTab)
-					->SetForegroundTab(FTabId(PCGEditorTabs::ViewportID[0]))
-					->SetHideTabWell(true)
-				)
-				->Split // Embedded graph
-				(
-					FTabManager::NewStack()
-					->SetSizeCoefficient(0.25f)
-					->AddTab(PCGEditorTabs::EmbeddedSubgraphsID, ETabState::OpenedTab)
-					->SetForegroundTab(FTabId(PCGEditorTabs::EmbeddedSubgraphsID))
-				)
-			)
-			->Split // Graph and details panel
-			(
-				FTabManager::NewSplitter()
-				->SetOrientation(Orient_Vertical)
-				->SetSizeCoefficient(0.25f)
-				->Split
-				(
-					FTabManager::NewStack()
-					->SetSizeCoefficient(0.6f)
-					->AddTab(PCGEditorTabs::GraphEditorID, ETabState::OpenedTab)
-					->SetHideTabWell(true)
-				)
-				->Split
-				(
-					FTabManager::NewStack()
-					->SetSizeCoefficient(0.4f)
-					->AddTab(PCGEditorTabs::PropertyDetailsID[0], ETabState::OpenedTab)
-					->AddTab(PCGEditorTabs::UserParamsID, ETabState::OpenedTab)
-					->SetForegroundTab(FTabId(PCGEditorTabs::PropertyDetailsID[0]))
-				)
-			)
-		)
-	);
-}
-
-bool FPVEditor::IsPanelAvailable(const FName PanelID) const
-{
-	static TSet<FName> SupportedIDs = {
-		PCGEditorTabs::FindID,
-		PCGEditorTabs::GraphEditorID,
-		PCGEditorTabs::LogID,
-		PCGEditorTabs::PropertyDetailsID[0],
-		PCGEditorTabs::UserParamsID,
-		PCGEditorTabs::ViewportID[0],
-		PCGEditorTabs::EmbeddedSubgraphsID
-	};
-	return SupportedIDs.Contains(PanelID);
-}
-
-void FPVEditor::RegisterToolbarInternal(FToolBarBuilder& ToolbarBuilder) const
-{
-	RegisterToolbarButton(ToolbarBuilder, EPCGToolbarButtons::Find);
-	RegisterToolbarButton(ToolbarBuilder, EPCGToolbarButtons::ForceRegen);
-	RegisterToolbarButton(ToolbarBuilder, EPCGToolbarButtons::Graph);
-	const FPVEditorCommands& Commands = FPVEditorCommands::Get();
-
-	ToolbarBuilder.AddToolBarButton(Commands.Export, NAME_None, TAttribute<FText>(), TAttribute<FText>(),
-		FSlateIcon(FAppStyle::GetAppStyleSetName(), "Icons.Toolbar.Export"));
-}
 
 void FPVEditor::BindCommands()
 {
@@ -292,51 +207,18 @@ void FPVEditor::BindCommands()
 
 	ToolkitCommands->MapAction(
 		EditorCommands.Export,
-		FExecuteAction::CreateSP(this, &FPVEditor::OnExport_Clicked));
+		FExecuteAction::CreateSP(this, &FPVEditor::OnExport));
 
 	ToolkitCommands->MapAction(
 		EditorCommands.LockNodeInspection,
 		FExecuteAction::CreateSP(this, &FPVEditor::OnLockNodeSelection));
 }
 
-UPVBaseSettings* GetNodeSettings(UPCGEditorGraphNodeBase* InNode)
-{
-	if (InNode && InNode->GetPCGNode())
-	{
-		return Cast<UPVBaseSettings>(InNode->GetPCGNode()->GetSettings());
-	}
-
-	return nullptr;
-}
-
-void FPVEditor::ChangeNodeInspection(UPCGEditorGraphNodeBase* InNode)
-{
-	if (UPVBaseSettings* Settings = GetNodeSettings(InNode))
-	{
-		TSharedPtr<SPVEditorViewport> Viewport = StaticCastSharedPtr<SPVEditorViewport>(GetViewportWidget(0));
-		
-		if (!Viewport.IsValid())
-		{
-			return;
-		}
-
-		Viewport->OnNodeInspectionChanged(Settings);
-		
-		NodeBeingInspected = InNode;
-		
-		UpdateCollectionSpreadSheet(InNode->GetPCGNode());
-		
-		SetNodeInspected(InNode, true);	
-	}
-}
-
 void FPVEditor::OnSelectedNodesChanged(const TSet<UObject*>& InNewSelection)
 {
 	FPCGEditor::OnSelectedNodesChanged(InNewSelection);
 
-	UPCGEditorGraphNodeBase* EdNode = nullptr;
-
-	if (UPCGNode* OldSelectedNode = GetSelectedNode() ? GetSelectedNode()->GetPCGNode() : nullptr)
+	if (const UPCGNode* OldSelectedNode = GetFirstSelectedNode())
 	{
 		if (UPVBaseSettings* Settings = Cast<UPVBaseSettings>(OldSelectedNode->GetSettings()))
 		{
@@ -357,26 +239,18 @@ void FPVEditor::OnSelectedNodesChanged(const TSet<UObject*>& InNewSelection)
 	}
 	SelectedNodes.Empty();
 
-	for(UObject* Selection: InNewSelection)
+	for (UObject* Selection : InNewSelection)
 	{
-		EdNode = Cast<UPCGEditorGraphNodeBase>(Selection);
-
-		if (EdNode && EdNode->GetPCGNode())
+		if (UPCGEditorGraphNodeBase* EdNode = Cast<UPCGEditorGraphNodeBase>(Selection))
 		{
-			SelectedNodes.Add(EdNode);
+			SelectedNodes.Add(EdNode->GetPCGNode());
 		}
 	}
 
-	if (UPCGNode* SelectedNode = GetSelectedNode() ? GetSelectedNode()->GetPCGNode() : nullptr)
+	const UPVBaseSettings* InspectionSettings = GetNodeSettings(NodeBeingInspected);
+	if (!InspectionSettings || !InspectionSettings->bLockInspection)
 	{
-		if (UPVBaseSettings* Settings = Cast<UPVBaseSettings>(SelectedNode->GetSettings()))
-		{
-			const UPVBaseSettings* InspectionSettings = GetNodeSettings(NodeBeingInspected);
-			if (!InspectionSettings || !InspectionSettings->bLockInspection)
-			{
-				ChangeNodeInspection(EdNode);
-			}
-		}	
+		ChangeNodeInspection(GetFirstSelectedNode());
 	}
 
 	// TODO: Hacky way to execute embedded subgraphs
@@ -392,582 +266,294 @@ void FPVEditor::OnSelectedNodesChanged(const TSet<UObject*>& InNewSelection)
 	}
 }
 
-void FPVEditor::OnGraphExecuted(IPCGBaseSubsystem* InSubsystem, IPCGGraphExecutionSource* InExecutionSource,
-	EPCGGenerationStatus InGenerationStatus)
+TAttribute<FGraphAppearanceInfo> FPVEditor::GetAppearanceInfo() const
 {
-	UpdateInspectedCollection();
-	UpdateStats();
-}
+	FGraphAppearanceInfo AppearanceInfo;
+	AppearanceInfo.CornerText = LOCTEXT("PVEditorCornerText", "Procedural Vegetation");
 
-void FPVEditor::UpdateInspectedCollection()
-{
-	if (const UPCGNode* InspectedNode = NodeBeingInspected ? NodeBeingInspected->GetPCGNode() : nullptr)
-	{
-		CollectionBeingInspected = MakeShared<FManagedArrayCollection>();
-		GetCollectionFromNode(InspectedNode, *CollectionBeingInspected);
-	}
+	return AppearanceInfo;
 }
 
-void FPVEditor::OnExport_Clicked()
+TSharedRef<FTabManager::FLayout> FPVEditor::GetDefaultLayout() const
 {
-	ExportInternal();
+	return FTabManager::NewLayout("Standalone_PVGraphEditor_DefaultLayout_v1.0.6")
+		->AddArea // Main PCG Graph Editor Area
+		(
+			FTabManager::NewPrimaryArea()
+			->SetOrientation(Orient_Vertical)
+			->Split // Top Section - Graph, Data Viewport, HLSL Source Editor, and Details View
+			(
+				FTabManager::NewSplitter()
+				->SetOrientation(Orient_Horizontal)->SetSizeCoefficient(0.9f)
+				->SetSizeCoefficient(0.40f)
+				->Split
+				(
+					FTabManager::NewSplitter()
+					->SetOrientation(Orient_Vertical)
+					->SetSizeCoefficient(0.1f)
+					->Split // Viewport
+					(
+						FTabManager::NewStack()
+						->SetSizeCoefficient(0.75f)
+						->AddTab(PCGEditorTabs::ViewportID[0], ETabState::OpenedTab)
+						->SetForegroundTab(FTabId(PCGEditorTabs::ViewportID[0]))
+						->SetHideTabWell(true)
+					)
+					->Split // Embedded graph
+					(
+						FTabManager::NewStack()
+						->SetSizeCoefficient(0.25f)
+						->AddTab(PCGEditorTabs::EmbeddedSubgraphsID, ETabState::OpenedTab)
+						->SetForegroundTab(FTabId(PCGEditorTabs::EmbeddedSubgraphsID))
+					)
+				)
+				->Split // Graph and details panel
+				(
+					FTabManager::NewSplitter()
+					->SetOrientation(Orient_Vertical)
+					->SetSizeCoefficient(0.25f)
+					->Split
+					(
+						FTabManager::NewStack()
+						->SetSizeCoefficient(0.6f)
+						->AddTab(PCGEditorTabs::GraphEditorID, ETabState::OpenedTab)
+						->SetHideTabWell(true)
+					)
+					->Split
+					(
+						FTabManager::NewStack()
+						->SetSizeCoefficient(0.4f)
+						->AddTab(PCGEditorTabs::PropertyDetailsID[0], ETabState::OpenedTab)
+						->AddTab(PCGEditorTabs::UserParamsID, ETabState::OpenedTab)
+						->SetForegroundTab(FTabId(PCGEditorTabs::PropertyDetailsID[0]))
+					)
+				)
+			)
+		);
 }
 
-FReply FPVEditor::ExportInternal()
+void FPVEditor::RegisterExtraTabFactories(FWorkflowAllowedTabSet& TabSet)
 {
-	const UProceduralVegetationGraph* ProceduralVegetationGraph = Cast<UProceduralVegetationGraph>(GetMainGraph());
-	
-	TSharedPtr<SPVExportSelectionDialog> ExportSelectionWidget =
-		SNew(SPVExportSelectionDialog)
-		.Title(LOCTEXT("Export Selection", "Choose outputs to export"))
-		.Graph(ProceduralVegetationGraph)
-		.SelectedNodes(SelectedNodes);
-
-	if (ExportSelectionWidget->ShowModal() != EAppReturnType::Ok)
+	if (CollectionSpreadSheetWidget.IsValid())
 	{
-		return FReply::Handled();
+		FPCGGenericTabFactoryParams Params(CollectionSpreadSheetTabId, SharedThis<FPCGEditor>(this),
+			[this]() { return CollectionSpreadSheetWidget.ToSharedRef(); });
+		Params.Label = LOCTEXT("CollectionSpreadsheetTab", "Collection Spreadsheet");
+		TabSet.RegisterFactory(MakeShared<FPCGGenericTabFactory>(Params));
 	}
+}
 
-	UE_LOG(LogProceduralVegetationEditor, Log, TEXT("Export Started"));
-	struct FScopedMessages
-	{
-		TArray<TSharedRef<FTokenizedMessage>> Messages;
-
-		TSharedRef<FTokenizedMessage> Error(const FText& ErrorMsg)
-		{
-			TSharedRef<FTokenizedMessage> Message = FTokenizedMessage::Create(EMessageSeverity::Error, ErrorMsg);
-			Messages.Add(Message);
-			return Message;
-		}
-
-		TSharedRef<FTokenizedMessage> Warning(const FText& ErrorMsg)
-		{
-			TSharedRef<FTokenizedMessage> Message = FTokenizedMessage::Create(EMessageSeverity::Warning, ErrorMsg);
-			Messages.Add(Message);
-			return Message;
-		}
-
-		TSharedRef<FTokenizedMessage> Info(const FText& ErrorMsg)
-		{
-			TSharedRef<FTokenizedMessage> Message = FTokenizedMessage::Create(EMessageSeverity::Info, ErrorMsg);
-			Messages.Add(Message);
-			return Message;
-		}
-		
-		bool ContainsErrors() const
-		{
-			for (TSharedRef<FTokenizedMessage> Message : Messages)
-			{
-				if (Message->GetSeverity() == EMessageSeverity::Error)
-				{
-					return true;
-				}
-			}
-
-			return false;
-		}
-
-		~FScopedMessages()
-		{
-			if (Messages.Num() > 0)
-			{
-				FMessageLogModule& MessageLogModule = FModuleManager::LoadModuleChecked<FMessageLogModule>("MessageLog");
-				TSharedRef<IMessageLogListing> PVEditorMessageLogListing = MessageLogModule.GetLogListing(PVEditor::MessageLogName);
-
-				PVEditorMessageLogListing->ClearMessages();
-				PVEditorMessageLogListing->AddMessages(Messages);
-				PVEditorMessageLogListing->NotifyIfAnyMessages(
-					ContainsErrors() ? LOCTEXT("PVEMeshExportFailed", "Mesh Export Failed") : LOCTEXT("PVEMeshExportComplete", "Mesh Export Completed"),
-					EMessageSeverity::Info,
-					true
-				);
-			}
-		}
-	};
-
-	FScopedMessages Messages;
-
-	const static auto ValidateExportSettings = [](const TArray<UPCGNode*>& NodesToExport, FScopedMessages& OutMessages)->bool
-	{
-		// Check that the paths are valid
-		bool bHasInvalidPath = false;
-		for (const UPCGNode* Node : NodesToExport)
-		{
-			const UPVOutputSettings* OutputSettings = CastChecked<const UPVOutputSettings>(Node->GetSettings());
-
-			FString PathError;
-			if (!OutputSettings->ExportSettings.Validate(PathError))
-			{
-				TSharedRef<FTokenizedMessage> Error = OutMessages.Error(
-					FText::Format(
-						LOCTEXT("ExportFailed_InvalidPath", "Invalid export path for node \"{0}\": \"{1}\""),
-						Node->GetDefaultTitle(),
-						FText::FromString(PathError)
-					)
-				);
-				Error->AddToken(FUObjectToken::Create(OutputSettings));
-				bHasInvalidPath = true;
-			}
-		}
-
-		return !bHasInvalidPath;
-	};
-
-	const static auto ShouldOverwriteAssets = [](const TArray<UPCGNode*>& NodesToExport)->bool 
-	{
-		TArray<FString> ExportPaths;
-		ExportPaths.Reserve(NodesToExport.Num() * 2);
-
-		for (const UPCGNode* Node : NodesToExport)
-		{
-			const UPVOutputSettings* OutputSettings = CastChecked<const UPVOutputSettings>(Node->GetSettings());
-			const FPVExportParams& ExportParams = OutputSettings->ExportSettings;
-			if (ExportParams.ReplacePolicy != EPVAssetReplacePolicy::Replace)
-			{
-				continue;
-			}
-			ExportPaths.Add(ExportParams.GetOutputMeshPackagePath());
-			if (ExportParams.ExportMeshType == EPVExportMeshType::SkeletalMesh)
-			{
-				ExportPaths.Add(ExportParams.GetOutputSkeletonPackagePath());
-
-				if (ExportParams.IsCollisionEnable())
-				{
-					ExportPaths.Add(ExportParams.GetOutputPhysicsAssetPackagePath());
-				}
-			}
-		}
-
-		TArray<UPackage*> ExistingPackages;
-		ExistingPackages.Reserve(ExportPaths.Num());
-		for (const FString& AssetPath : ExportPaths)
-		{
-			UPackage* ExistingPackage = LoadPackage(nullptr, *AssetPath, LOAD_None);
-			if (ExistingPackage)
-			{
-				ExistingPackages.Add(ExistingPackage);
-			}
-		}
-
-		if (ExistingPackages.Num() == 0)
-		{
-			return true;
-		}
-
-		FPackagesDialogModule& PackagesDialogModule = FModuleManager::LoadModuleChecked<FPackagesDialogModule>(TEXT("PackagesDialog"));
-		PackagesDialogModule.CreatePackagesDialog(
-			LOCTEXT("PackagesDialog_OverwriteFiles_Title", "WARNING: Assets will be overwritten"), 
-			LOCTEXT("PackagesDialog_OverwriteFiles_Message", "The following assets will be overwritten. Would you like to continue?."),
-			/*InReadOnly=*/true
-		);
-		PackagesDialogModule.AddButton(
-			DRT_Save, 
-			LOCTEXT("PackagesDialog_OverwriteFiles_ContinueButton", "Continue"), 
-			LOCTEXT("PackagesDialog_OverwriteFiles_ButtonTip", "Continues the export, overwrites any existing assets")
-		);
-		PackagesDialogModule.AddButton(
-			DRT_Cancel, 
-			LOCTEXT("PackagesDialog_OverwriteFiles_CancelButton", "Cancel"), 
-			LOCTEXT("CancelDeleteButtonTip", "Cancel export")
-		);
-
-		for (UPackage* Package : ExistingPackages)
-		{
-			PackagesDialogModule.AddPackageItem(Package, ECheckBoxState::Checked);
-		}
-
-		const EDialogReturnType UserResponse = PackagesDialogModule.ShowPackagesDialog();
-
-		return UserResponse == DRT_Save;
-	};
-
-	const auto GetOutputNodeCollections = [this](TArray<FManagedArrayCollection>& OutCollections, const TArray<UPCGNode*>& NodesToExport, FScopedMessages& OutMessages)->bool
-	{
-		OutCollections.Reset(NodesToExport.Num());
-		for (UPCGNode* Node : NodesToExport)
-		{
-			FManagedArrayCollection& OutCollection = OutCollections.AddDefaulted_GetRef();
-			if (!GetCollectionFromNode(Node, OutCollection))
-			{
-				TSharedRef<FTokenizedMessage> Error = OutMessages.Error(
-					FText::Format(
-						LOCTEXT("ExportFailed_NodeErrors", "Export failed for output node \"{0}\", Invalid output please fix the errors on the node."),
-						Node->GetDefaultTitle()
-					)
-				);
-				Error->AddToken(FUObjectToken::Create(Node));
-				return false;
-			}
-		}
-		
-		return true;
+bool FPVEditor::IsPanelAvailable(const FName PanelID) const
+{
+	static TSet<FName> SupportedIDs = {
+		PCGEditorTabs::FindID,
+		PCGEditorTabs::GraphEditorID,
+		PCGEditorTabs::LogID,
+		PCGEditorTabs::PropertyDetailsID[0],
+		PCGEditorTabs::UserParamsID,
+		PCGEditorTabs::ViewportID[0],
+		PCGEditorTabs::EmbeddedSubgraphsID
 	};
+	return SupportedIDs.Contains(PanelID);
+}
 
-	const auto FoliageDistributorValidation = [this](const TArray<UPCGNode*>& NodesToExport)->bool
-	{
-		TArray<const UPCGNode*> NodesWithoutFoliage;
-		for (UPCGNode* Node : NodesToExport)
-		{
-			GetPVDataFromNode(Node, [Node, &NodesWithoutFoliage](const UPVData& InPVData)
-			{
-				if(!InPVData.IsA<UPVFoliageMeshData>())
-				{
-					NodesWithoutFoliage.Add(Node);
-				}
-			});
-		}
-
-		if (NodesWithoutFoliage.Num() > 0)
-		{
-			FString NodeNames = "";
-			for (int32 i = 0; i < NodesWithoutFoliage.Num(); i++)
-			{
-				const UPCGNode* Node = NodesWithoutFoliage[i];
-				NodeNames += FString::Format(TEXT("\" {0} \""), {Node->GetAuthoredTitleName().ToString()});
-				
-				if (NodesWithoutFoliage.Num() > 1 && i == NodesWithoutFoliage.Num() - 2)
-				{
-					NodeNames += " and ";
-				}
-
-				if (i != NodesWithoutFoliage.Num() - 2)
-				{
-					NodeNames += ", ";
-				}
-			}
-			
-			FString Message = FString::Format(TEXT("{0} does not use input from foliage distributor node, Do you want to export with default foliage distribution."),{NodeNames});
-			const EAppReturnType::Type ret = FMessageDialog::Open(EAppMsgType::YesNoCancel, FText::Format(LOCTEXT("Export_Foliage_Warning", "{0}"), FText::FromString(Message)));
-			if (ret != EAppReturnType::Cancel)
-			{
-				const bool bShouldExportFoliage = ret == EAppReturnType::Yes;
-				for (const UPCGNode* Node : NodesWithoutFoliage)
-				{
-					UPVOutputSettings* OutputSettings = Cast<UPVOutputSettings>(Node->GetSettings());
-					OutputSettings->ExportSettings.bShouldExportFoliage = bShouldExportFoliage;
-				}
-			}
-			else
-			{
-				return false;
-			}
-		}
+void FPVEditor::RegisterToolbarInternal(FToolBarBuilder& ToolbarBuilder) const
+{
+	RegisterToolbarButton(ToolbarBuilder, EPCGToolbarButtons::Find);
+	RegisterToolbarButton(ToolbarBuilder, EPCGToolbarButtons::ForceRegen);
+	RegisterToolbarButton(ToolbarBuilder, EPCGToolbarButtons::Graph);
+	const FPVEditorCommands& Commands = FPVEditorCommands::Get();
 
-		return true;
-	};
+	ToolbarBuilder.AddToolBarButton(Commands.Export, NAME_None, TAttribute<FText>(), TAttribute<FText>(),
+		FSlateIcon(FAppStyle::GetAppStyleSetName(), "Icons.Toolbar.Export"));
+}
 
-	TArray<UPCGNode*> NodesToExport;
+bool FPVEditor::CanToggleInspected() const
+{
+	return false;
+}
 
-	const UPVEditorSettings* EditorSettings = GetDefault<UPVEditorSettings>();
-	
-	if (EditorSettings->ExportType == EPVExportType::Selection)
+TSharedRef<SPCGEditorViewport> FPVEditor::CreateViewportWidget()
+{
+	TSharedPtr<SPVEditorViewport> ViewportWidget = nullptr;
+	if (EditorViewport == nullptr)
 	{
-		UE_LOG(LogProceduralVegetationEditor, Log, TEXT("Exporting the selected nodes"));
-		
-		for (int i = 0; i < SelectedNodes.Num(); i++)
-		{
-			if (UPCGNode* Node = SelectedNodes[i] ?  SelectedNodes[i]->GetPCGNode() : nullptr)
-			{
-				if (const UPVOutputSettings* OutputSettings = Cast<UPVOutputSettings>(Node->GetSettings()))
-				{
-					NodesToExport.Add(Node);
-				}	
-			}
-		}
+		EditorViewport = ViewportWidget = SNew(SPVEditorViewport);
 	}
 	else
 	{
-		UE_LOG(LogProceduralVegetationEditor, Log, TEXT("Batch Exporting nodes"));
-		NodesToExport = ProceduralVegetationGraph->GetNodes().FilterByPredicate([](UPCGNode* InNode)->bool {
-			const UPVOutputSettings* OutputSettings = Cast<const UPVOutputSettings>(InNode->GetSettings());
-			return OutputSettings && OutputSettings->ExportSettings.bShouldExport;
-		});
+		ViewportWidget = SNew(SPVEditorViewport);
 	}
 
-	if (!ValidateExportSettings(NodesToExport, Messages))
-	{
-		return FReply::Handled();
-	}
+	return ViewportWidget.ToSharedRef();
+}
 
-	if (!ShouldOverwriteAssets(NodesToExport))
-	{
-		return FReply::Handled();
-	}
-	
-	TArray<FManagedArrayCollection> CollectionsToExport;
-	if (!GetOutputNodeCollections(CollectionsToExport, NodesToExport, Messages))
-	{
-		return FReply::Handled();
-	}
+TSharedRef<SCollectionSpreadSheetWidget> FPVEditor::CreateCollectionSpreadSheetWidget()
+{
+	return SNew(SCollectionSpreadSheetWidget);
+}
+
+FPCGStack FPVEditor::GetStackFromPin(const UPCGPin* InPin) const
+{
+	check(InPin);
 
-	if (!FoliageDistributorValidation(NodesToExport))
+	FPCGStack Stack;
+	if (const FPCGStack* PCGStack = GetStackBeingInspected())
 	{
-		return FReply::Handled();
+		Stack = CopyTemp(*PCGStack);
+		Stack.PushFrame(InPin->Node);
+		Stack.PushFrame(InPin);
 	}
 
-	FScopedSlowTask ProgressSlowTask(NodesToExport.Num(), LOCTEXT("ExportAsset_SlowTask", "Exporting Asset(s) ..."));
-	ProgressSlowTask.MakeDialog(false);
+	return Stack;
+}
 
-	for (int32 i = 0; i < NodesToExport.Num(); ++i)
+const UPVData* FPVEditor::GetDataFromPin(const UPCGPin* InPin)
+{
+	if (!InPin)
 	{
-		if (ProgressSlowTask.ShouldCancel())
-		{
-			break;
-		}
-
-		const UPCGNode* Node = NodesToExport[i];
-		const UPVOutputSettings* OutputSettings = CastChecked<const UPVOutputSettings>(Node->GetSettings());
-		check(OutputSettings != nullptr);
-
-		const FText AssetCountText = NodesToExport.Num() > 1
-			? FText::FromString(FString::Printf(TEXT(" [%d/%d]"), i + 1, NodesToExport.Num()))
-			: FText::FromString("");
-
-		ProgressSlowTask.EnterProgressFrame(1.f, FText::Format(LOCTEXT("ExportingAsset", "Exporting \"{0}\"{1}"), FText::FromName(OutputSettings->ExportSettings.MeshName), AssetCountText));
-
-		FScopedSlowTask ExportCollectionSlowTask(1.f, LOCTEXT("GeneratingMesh", "Generating Mesh ..."));
-		float PrevProgress = 0.f;
-		const auto OnStatusUpdated = [&](const FText& Stage, float Progress)->bool
-		{
-			const float WorkThisFrame = Progress - PrevProgress;
-			PrevProgress = Progress;
-
-			ExportCollectionSlowTask.FrameMessage = Stage;
-
-			ExportCollectionSlowTask.CompletedWork += ExportCollectionSlowTask.CurrentFrameScope;
-
-			const float WorkRemaining = ExportCollectionSlowTask.TotalAmountOfWork - ExportCollectionSlowTask.CompletedWork;
-			ExportCollectionSlowTask.CurrentFrameScope = FMath::Min(WorkRemaining, WorkThisFrame);
-			ExportCollectionSlowTask.ForceRefresh();
-
-			return !ProgressSlowTask.ShouldCancel();
-		};
+		return nullptr;
+	}
 
-		TArray<FString> CreatedAssets;
-		const PV::Export::EExportResult ExportResult = PV::Export::ExportCollectionAsMesh(
-			ProceduralVegetationBeingEdited, 
-			CollectionsToExport[i], 
-			OutputSettings->ExportSettings,
-			CreatedAssets,
-			OnStatusUpdated
-		);
+	const FPCGStack Stack = GetStackFromPin(InPin);
 
-		if (ExportResult == PV::Export::EExportResult::Fail)
-		{
-			Messages.Error(FText::Format(
-				LOCTEXT("ExportFailed_UnknownError", "Failed to export mesh \"{0}\" due to unknown error, export aborted."),
-				FText::FromName(OutputSettings->ExportSettings.MeshName)
-			));
-			
-			break;
-		}
-		else if (ExportResult == PV::Export::EExportResult::Cancelled)
-		{
-			FString MeshNames = "{ ";
-			for (int32 j = i; j < NodesToExport.Num(); j++)
+	const UPVData* PVData = nullptr;
+	ExecutionSource->GetExecutionState().GetInspection().InspectData(
+		Stack,
+		[&](const FPCGDataCollection& DataCollection)
 			{
-				MeshNames += CastChecked<const UPVOutputSettings>(NodesToExport[j]->GetSettings())->ExportSettings.MeshName.ToString();
-				if (j != (NodesToExport.Num() - 1))
-				{
-					MeshNames += ", ";
-				}
+				const FPCGTaggedData* TaggedData = DataCollection.TaggedData.GetData();
+				PVData = TaggedData
+					? Cast<UPVData>(TaggedData->Data)
+					: nullptr;
 			}
-			MeshNames += " }";
-
-			Messages.Warning(
-				FText::Format(
-					LOCTEXT("ExportFailed_ExportCanceled", "Failed to export meshes \"{0}\", operation canceled by user."),
-					FText::FromString(MeshNames)
-				));
-
-			break;
-		}
-		else if (ExportResult == PV::Export::EExportResult::Skipped)
-		{
-			Messages.Info(FText::Format(
-				LOCTEXT("ExportSkipped", "Export of mesh \"{0}\" skipped due to conflicting asset at output location"),
-				FText::FromName(OutputSettings->ExportSettings.MeshName)
-			));
-		}
-		else
-		{
-			auto Message = Messages.Info(FText::Format(
-				LOCTEXT("ExportComplete", "Mesh exported successfully \"{0}\""),
-				FText::FromName(OutputSettings->ExportSettings.MeshName)
-			));
+	);
 
-			for (const FString& AssetPath : CreatedAssets)
-			{
-				Message->AddToken(FAssetNameToken::Create(AssetPath));
-			}
-		}
-	}
-	
-	return FReply::Handled();
+	return PVData;
 }
 
-void FPVEditor::UpdateStats()
+void FPVEditor::OnNodeInspected(const FPCGStack& InPCGStack)
 {
-	TSharedPtr<SPVEditorViewport> Viewport = StaticCastSharedPtr<SPVEditorViewport>(GetViewportWidget(0));
-	if (!Viewport.IsValid() || NodeBeingInspected == nullptr)
-	{
-		return;
-	}
-
-	const UPVBaseSettings* NodeBeingInspectedSettings = GetNodeSettings(NodeBeingInspected);
-	if (!NodeBeingInspectedSettings)
-	{
-		return;
-	}
-
-	TArray<FText> OverlayStats;
-	GatherStats(OverlayStats);
-	Viewport->PopulateStatsOverlayText(OverlayStats);
-	
-	Viewport->SetOverlayText(
-		NodeBeingInspectedSettings->bLockInspection
-		? NodeBeingInspected->GetNodeTitle(ENodeTitleType::FullTitle)
-		: FText::GetEmpty()
-	);
+	UpdateInspectedCollection();
+	UpdateStats();
 }
 
-void FPVEditor::OnLockNodeSelection()
+void FPVEditor::ChangeNodeInspection(UPCGNode* InNode)
 {
-	UPVBaseSettings* NodeBeingInspectedSettings = GetNodeSettings(NodeBeingInspected);
-	UPVBaseSettings* SelectedNodeSettings = GetNodeSettings(GetSelectedNode());
-
-	if (NodeBeingInspectedSettings)
-	{
-		NodeBeingInspectedSettings->bLockInspection = !NodeBeingInspectedSettings->bLockInspection;
-	}
-
-	if (NodeBeingInspected != GetSelectedNode())
+	if (UPVBaseSettings* Settings = GetNodeSettings(InNode))
 	{
-		if (GetSelectedNode())
+		if (EditorViewport.IsValid())
 		{
-			SelectedNodeSettings->bLockInspection = true;
-			ChangeNodeInspection(GetSelectedNode());
+			EditorViewport->OnNodeInspectionChanged(Settings);
 		}
+
+		NodeBeingInspected = InNode;
+
+		SetNodeInspected(GetEditorNodeFromPCGNode(NodeBeingInspected), true);
 	}
-	
-	UpdateStats();
 }
 
-TSharedRef<SPCGEditorViewport> FPVEditor::CreateViewportWidget()
+void FPVEditor::OnLockNodeSelection()
 {
-	return SNew(SPVEditorViewport);
-}
+	const auto SelectedNode = GetFirstSelectedNode();
 
-FPCGStack FPVEditor::GetPVDataStack(const UPCGNode* InNode, const UPCGPin*& OutPin) const
-{
-	OutPin = nullptr;
-	const TArray<TObjectPtr<UPCGPin>>& Pins = InNode->GetOutputPins();
+	UPVBaseSettings* SelectedNodeSettings = GetNodeSettings(SelectedNode);
 
-	if (Pins.IsValidIndex(0))
+	if (UPVBaseSettings* NodeBeingInspectedSettings = GetNodeSettings(NodeBeingInspected))
 	{
-		OutPin = Pins[0];
+		NodeBeingInspectedSettings->bLockInspection = !NodeBeingInspectedSettings->bLockInspection;
 	}
 
-	const FPCGStack* PCGStack = GetStackBeingInspected();
-	check(PCGStack)
-
-	// Create a temporary stack with Node + valid Pin to query the exact DataCollection
-	FPCGStack Stack = *PCGStack;
-	TArray<FPCGStackFrame>& StackFrames = Stack.GetStackFramesMutable();
-	StackFrames.Reserve(StackFrames.Num() + 2);
-	StackFrames.Emplace(InNode);
-	if (OutPin)
+	if (NodeBeingInspected != SelectedNode)
 	{
-		StackFrames.Emplace(OutPin);
+		SelectedNodeSettings->bLockInspection = true;
+		ChangeNodeInspection(GetFirstSelectedNode());
 	}
 
-	return Stack;
+	UpdateStats();
 }
 
-const UPVData* FPVEditor::GetPVDataFromCollection(const FPCGDataCollection& InDataCollection, const UPCGPin* InPin) const
+UPCGNode* FPVEditor::GetFirstSelectedNode()
 {
-	const FPCGTaggedData* TaggedData = InDataCollection.TaggedData.GetData();
-
-	if (TaggedData != nullptr)
-	{
-		const UPVData* ProceduralVegetationData = Cast<UPVData>(TaggedData->Data);
-		if (ProceduralVegetationData != nullptr && InPin)
-		{
-			return ProceduralVegetationData;
-		}
-	}
-	
-	return nullptr;
+	return SelectedNodes.Num()
+		? SelectedNodes[0]
+		: nullptr;
 }
 
-bool FPVEditor::GetPVDataFromNode(const UPCGNode* InNode, TFunctionRef<void(const UPVData&)> PVDataFunc) const
+UPCGEditorGraphNodeBase* FPVEditor::GetEditorNodeFromPCGNode(const UPCGNode* InPCGNode) const
 {
-	const UPCGPin* OutPin = nullptr;
-	FPCGStack Stack = GetPVDataStack(InNode, OutPin);
+	const TObjectPtr<UEdGraphNode>* EdGraphNode = GetMainEditorGraph()->Nodes.FindByPredicate(
+		[&](UEdGraphNode* GraphNode)
+			{
+				const UPCGEditorGraphNodeBase* PCGEdGraphNode = Cast<UPCGEditorGraphNodeBase>(GraphNode);
+				return PCGEdGraphNode
+					? PCGEdGraphNode->GetPCGNode() == InPCGNode
+					: false;
+			});
+	return EdGraphNode
+		? CastChecked<UPCGEditorGraphNodeBase>(*EdGraphNode)
+		: nullptr;
+}
 
-	bool bFoundPVData = false;
 
-	ExecutionSource->GetExecutionState().GetInspection().InspectData(Stack, [this, Pin = OutPin, &bFoundPVData, &PVDataFunc](const FPCGDataCollection& InDataCollection)
+void FPVEditor::UpdateInspectedCollection()
+{
+	if (NodeBeingInspected)
 	{
-		if(const UPVData* PVData = GetPVDataFromCollection(InDataCollection, Pin))
+		if (const UPVData* const Data = GetDataFromPin(GetOutPinFromNode(NodeBeingInspected)))
 		{
-			bFoundPVData = true;
-			PVDataFunc(*PVData);
+			CollectionBeingInspected = Data->GetSharedCollection();
 		}
-	});
-		
-	return bFoundPVData;
-}
+	}
 
-bool FPVEditor::CanToggleInspected() const
-{
-	return false;
+	UpdateCollectionSpreadSheet();
 }
 
-UPCGEditorGraphNodeBase* FPVEditor::GetSelectedNode()
+void FPVEditor::UpdateCollectionSpreadSheet()
 {
-	if (SelectedNodes.Num())
+	if (!NodeBeingInspected)
 	{
-		return SelectedNodes[0];
+		CollectionSpreadSheetWidget->SetData(FString());
+		CollectionSpreadSheetWidget->RefreshWidget();
+
+		return;
 	}
 
-	return nullptr;
-}
+	const TSharedPtr<SDockTab> CollectionSpreadSheetTab =
+		GetTabManager()
+		? GetTabManager()->FindExistingLiveTab(CollectionSpreadSheetTabId)
+		: nullptr;
 
-bool FPVEditor::GetCollectionFromNode(const UPCGNode* InNode, FManagedArrayCollection& OutCollection) const
-{
-	return GetPVDataFromNode(InNode, [&OutCollection](const UPVData& InPVData)
-	{
-		OutCollection = InPVData.GetCollection();
-	});
-}
+	const bool bIsCollectionSpreadSheetVisible =
+		CollectionSpreadSheetWidget.IsValid() &&
+		CollectionSpreadSheetTab.IsValid() &&
+		CollectionSpreadSheetTab->GetVisibility().IsVisible();
 
-TSharedRef<SCollectionSpreadSheetWidget> FPVEditor::CreateCollectionSpreadSheetWidget()
-{
-	return SNew(SCollectionSpreadSheetWidget);
-}
+	if (bIsCollectionSpreadSheetVisible)
+	{
+		const TArray<TObjectPtr<UPCGPin>> AllPins = GetAllPVPinsFromNode(NodeBeingInspected);
 
-void FPVEditor::UpdateCollectionSpreadSheet(const UPCGNode* InSelectedNode)
-{
-	TSharedPtr<SDockTab> CollectionSpreadSheetTab = GetTabManager() ? GetTabManager()->FindExistingLiveTab(CollectionSpreadSheetTabId) : nullptr;
-	const bool bIsCollectionSpreadSheetVisible = CollectionSpreadSheetWidget.IsValid() && CollectionSpreadSheetTab.IsValid() && CollectionSpreadSheetTab->GetVisibility().IsVisible();
+		CollectionSpreadSheetWidget->SetData(NodeBeingInspected->GetNodeTitle(EPCGNodeTitleType::FullTitle).ToString());
 
-	if (InSelectedNode)
-	{
-		if (bIsCollectionSpreadSheetVisible)
+		TMap<FString, SCollectionInfo>& CollectionInfoMap = CollectionSpreadSheetWidget->GetCollectionTable()->GetCollectionInfoMap();
+		CollectionInfoMap.Empty();
+		for (TObjectPtr<UPCGPin> Pin : AllPins)
 		{
-			CollectionSpreadSheetWidget->GetCollectionTable()->GetCollectionInfoMap().Empty();
-			CollectionSpreadSheetWidget->GetCollectionTable()->GetCollectionInfoMap().Add("Out", {CollectionBeingInspected});
-			CollectionSpreadSheetWidget->SetData(InSelectedNode->GetNodeTitle(EPCGNodeTitleType::FullTitle).ToString());
+			if (const UPVData* const PinData = GetDataFromPin(Pin))
+			{
+				CollectionInfoMap.Emplace(
+					Pin->Properties.Label.ToString(),
+					{PinData->GetSharedCollection()}
+				);
+			}
 		}
-	}
-	else if (bIsCollectionSpreadSheetVisible)
-	{
-		CollectionSpreadSheetWidget->SetData(FString());
-	}
-	
-	if (bIsCollectionSpreadSheetVisible)
-	{
+		
+		// TODO: Check if there is a way to select default output
+		// if (GetOutPinFromNode(NodeBeingInspected))
+		// {
+		// 	CollectionSpreadSheetWidget->GetCollectionTable()->SetSelectedOutput(PCGPinConstants::DefaultOutputLabel);
+		// }
+
 		CollectionSpreadSheetWidget->RefreshWidget();
 	}
 }
@@ -978,7 +564,7 @@ void FPVEditor::GatherStats(TArray<FText>& OverlayStats)
 	{
 		return;
 	}
-	
+
 	if (CollectionBeingInspected->HasGroup(PV::GroupNames::PointGroup))
 	{
 		const int32 NumPoints = CollectionBeingInspected->NumElements(PV::GroupNames::PointGroup);
@@ -1016,4 +602,75 @@ void FPVEditor::GatherStats(TArray<FText>& OverlayStats)
 	}
 }
 
+void FPVEditor::UpdateStats()
+{
+	if (!EditorViewport)
+	{
+		return;
+	}
+
+	TArray<FText> OverlayStats;
+	GatherStats(OverlayStats);
+	EditorViewport->PopulateStatsOverlayText(OverlayStats);
+
+	if (const UPVBaseSettings* NodeBeingInspectedSettings = GetNodeSettings(NodeBeingInspected))
+	{
+		EditorViewport->SetOverlayText(
+			NodeBeingInspectedSettings->bLockInspection
+			? NodeBeingInspected->GetNodeTitle(EPCGNodeTitleType::FullTitle)
+			: FText::GetEmpty()
+		);
+	}
+}
+
+void FPVEditor::OnExport()
+{
+	UProceduralVegetationGraph* ProceduralVegetationGraph = Cast<UProceduralVegetationGraph>(GetMainGraph());
+
+	TArray<TObjectPtr<UPVExportEntry>> ExportEntries;
+	ProceduralVegetationGraph->ForEachNode(
+		[&](UPCGNode* Node)-> bool
+			{
+				if (!Node->GetSettings()->IsA<UPVOutputSettings>())
+				{
+					return true;
+				}
+
+				UPCGPin* const InPin = GetInPinFromNode(Node);
+				UPCGPin* const UpstreamOutPin = GetUpstreamPin(InPin);
+				
+				const bool bIsMeshDataPin = IsPinOfType<FPVDataTypeInfoMesh>(UpstreamOutPin);
+				const bool bIsFoliageMeshDataPin = IsPinOfType<FPVDataTypeInfoFoliageMesh>(UpstreamOutPin);
+				if (!bIsMeshDataPin && !bIsFoliageMeshDataPin)
+				{
+					return true;
+				}
+
+				if (const UPVData* const NodeData = GetDataFromPin(InPin))
+				{
+					const TObjectPtr<UPVExportEntry>& ExportEntry = ExportEntries.Emplace_GetRef(NewObject<UPVExportEntry>(ProceduralVegetationGraph));
+					ExportEntry->Initialize(NodeData->GetSharedCollection(), Node, SelectedNodes.Contains(Node), bIsFoliageMeshDataPin);
+				}
+				
+				return true;
+			}
+	);
+
+	const TSharedPtr<SPVExportSelectionDialog> ExportSelectionWidget =
+		SNew(SPVExportSelectionDialog)
+		.Title(LOCTEXT("ExportSelectionLabel", "Modify Export Settings"))
+		.ExportEntries(ExportEntries);
+
+	if (ExportSelectionWidget->ShowModal() != EAppReturnType::Ok)
+	{
+		return;
+	}
+
+	UE_LOG(LogProceduralVegetationEditor, Log, TEXT("Export Started"));
+	
+	FPVExporter Exporter(ProceduralVegetationBeingEdited, ExportEntries);
+	Exporter.Export();
+}
+
+
 #undef LOCTEXT_NAMESPACE
diff --git a/Engine/Plugins/Experimental/ProceduralVegetationEditor/Source/ProceduralVegetationEditor/Private/PVEditor.h b/Engine/Plugins/Experimental/ProceduralVegetationEditor/Source/ProceduralVegetationEditor/Private/PVEditor.h
--- a/Engine/Plugins/Experimental/ProceduralVegetationEditor/Source/ProceduralVegetationEditor/Private/PVEditor.h
+++ b/Engine/Plugins/Experimental/ProceduralVegetationEditor/Source/ProceduralVegetationEditor/Private/PVEditor.h
@@ -3,92 +3,105 @@
 #pragma once
 
 #include "PCGEditor.h"
-#include "DataTypes/PVData.h"
 #include "SPVEditorViewport.h"
 
+#include "DataTypes/PVData.h"
+
 class SCollectionSpreadSheetWidget;
-struct FManagedArrayCollection;
 class UPCGDefaultExecutionSource;
-class UProceduralVegetation;
 class UPCGEditorGraphNodeBase;
+class UPVEditorMode;
+class UProceduralVegetation;
+struct FManagedArrayCollection;
 
 class FPVEditor : public FPCGEditor
 {
 public:
 	FPVEditor();
 
-	void Initialize(const EToolkitMode::Type InMode, const TSharedPtr<class IToolkitHost>& InToolkitHost, UProceduralVegetation* InProceduralVegetation);
+	void Initialize(
+		const EToolkitMode::Type InMode,
+		const TSharedPtr<IToolkitHost>& InToolkitHost,
+		UProceduralVegetation* InProceduralVegetation
+	);
+
 	virtual TSubclassOf<UPCGEditorGraphSchema> GetSchemaClass() const override;
 
 	virtual void AddReferencedObjects(FReferenceCollector& Collector) override;
-	
+
 	// ~Begin FAssetEditorToolkit interface
+	virtual void CreateEditorModeManager() override;
+
 	virtual FText GetToolkitName() const override;
 	virtual FName GetToolkitFName() const override;
 	virtual FText GetBaseToolkitName() const override;
 	virtual FText GetToolkitToolTipText() const override;
-	virtual FLinearColor GetWorldCentricTabColorScale() const override;
+
 	virtual FString GetWorldCentricTabPrefix() const override;
+	virtual FLinearColor GetWorldCentricTabColorScale() const override;
+
 	virtual void OnClose() override;
 	// ~End FAssetEditorToolkit interface
 
+	// ~Begin FPCGEditor interface
 	virtual IPCGBaseSubsystem* GetSubsystem() const override;
+	// ~End FPCGEditor interface
 
 protected:
 	// ~Begin FPCGEditor interface
+	virtual void BindCommands() override;
+	virtual void OnSelectedNodesChanged(const TSet<UObject*>& InNewSelection) override;
+
 	virtual TAttribute<FGraphAppearanceInfo> GetAppearanceInfo() const override;
 	virtual TSharedRef<FTabManager::FLayout> GetDefaultLayout() const override;
+
 	virtual void RegisterExtraTabFactories(FWorkflowAllowedTabSet& TabSet) override;
 	virtual bool IsPanelAvailable(const FName PanelID) const override;
 	virtual void RegisterToolbarInternal(FToolBarBuilder& ToolbarBuilder) const override;
-	virtual void BindCommands() override;
-	virtual void OnSelectedNodesChanged(const TSet<UObject*>& InNewSelection) override;
+
+	virtual bool CanToggleInspected() const override;
+
+	virtual TSharedRef<SPCGEditorViewport> CreateViewportWidget() override;
 	// ~End FPCGEditor interface
 
-	void OnGraphExecuted(IPCGBaseSubsystem* InSubsystem, IPCGGraphExecutionSource* InExecutionSource,
-	EPCGGenerationStatus InGenerationStatus);
+	static TSharedRef<SCollectionSpreadSheetWidget> CreateCollectionSpreadSheetWidget();
+
+private:
+	FPCGStack GetStackFromPin(const UPCGPin* InPin) const;
+
+	const UPVData* GetDataFromPin(const UPCGPin* InPin);
+
+	void OnNodeInspected(const FPCGStack& InPCGStack);
+
+	void ChangeNodeInspection(UPCGNode* InNode);
+	void OnLockNodeSelection();
+
+	UPCGNode* GetFirstSelectedNode();
+	UPCGEditorGraphNodeBase* GetEditorNodeFromPCGNode(const UPCGNode* Node) const;
 
 	void UpdateInspectedCollection();
 
-	TObjectPtr<UPCGDefaultExecutionSource> ExecutionSource;
+	void UpdateCollectionSpreadSheet();
+
+	void GatherStats(TArray<FText>& OverlayStats);
+	void UpdateStats();
+
+	void OnExport();
 
 private:
 	static const FName CollectionSpreadSheetTabId;
-	
-	TObjectPtr<UProceduralVegetation> ProceduralVegetationBeingEdited = nullptr;
 
-	TArray<TObjectPtr<UPCGEditorGraphNodeBase>> SelectedNodes;
-	TSharedPtr<FManagedArrayCollection> CollectionBeingInspected = nullptr;
-	TObjectPtr<UPCGEditorGraphNodeBase> NodeBeingInspected = nullptr;
-	
-	TSharedPtr<SCollectionSpreadSheetWidget> CollectionSpreadSheetWidget;
-	
-	FDateTime SessionStartTime;
-	
-	TSharedRef<SCollectionSpreadSheetWidget> CreateCollectionSpreadSheetWidget();
-	
-	void UpdateCollectionSpreadSheet(const UPCGNode* InSelectedNode);
-
-	void GatherStats(TArray<FText>& OverlayStats);
-	
-	void OnExport_Clicked();
+	TObjectPtr<UPCGDefaultExecutionSource> ExecutionSource;
 
-	FReply ExportInternal();
+	TObjectPtr<UProceduralVegetation> ProceduralVegetationBeingEdited = nullptr;
 
-	void UpdateStats();
-	
-	void OnLockNodeSelection();
+	TSharedPtr<SPVEditorViewport> EditorViewport = nullptr;
 
-	void ChangeNodeInspection(UPCGEditorGraphNodeBase* InNode);
+	TObjectPtr<UPCGNode> NodeBeingInspected = nullptr;
+	TArray<TObjectPtr<UPCGNode>> SelectedNodes;
 
-protected:
-	virtual TSharedRef<SPCGEditorViewport> CreateViewportWidget() override;
-	
-	FPCGStack GetPVDataStack(const UPCGNode* InNode, UPCGPin const*& OutPin) const;
-	const UPVData* GetPVDataFromCollection(const FPCGDataCollection& InDataCollection, const UPCGPin* InPin) const;
-	bool GetPVDataFromNode(const UPCGNode* InNode, TFunctionRef<void(const UPVData&)> PVDataFunc) const;
-	virtual bool CanToggleInspected() const override;
-	bool GetCollectionFromNode(const UPCGNode* InNode, FManagedArrayCollection& OutCollection) const;
+	TSharedPtr<SCollectionSpreadSheetWidget> CollectionSpreadSheetWidget;
+	TSharedPtr<FManagedArrayCollection> CollectionBeingInspected = nullptr;
 
-	UPCGEditorGraphNodeBase* GetSelectedNode();
+	FDateTime SessionStartTime;
 };
diff --git a/Engine/Plugins/Experimental/ProceduralVegetationEditor/Source/ProceduralVegetationEditor/Private/SPVEditorViewport.cpp b/Engine/Plugins/Experimental/ProceduralVegetationEditor/Source/ProceduralVegetationEditor/Private/SPVEditorViewport.cpp
--- a/Engine/Plugins/Experimental/ProceduralVegetationEditor/Source/ProceduralVegetationEditor/Private/SPVEditorViewport.cpp
+++ b/Engine/Plugins/Experimental/ProceduralVegetationEditor/Source/ProceduralVegetationEditor/Private/SPVEditorViewport.cpp
@@ -4,6 +4,7 @@
 
 #include "AdvancedPreviewScene.h"
 #include "AdvancedPreviewSceneMenus.h"
+#include "AssetEditorViewportLayout.h"
 #include "PCGEditorSettings.h"
 #include "PVEditorCommands.h"
 #include "PVEditorSettings.h"
@@ -35,7 +36,6 @@
 
 #include "Widgets/SPVGradientBox.h"
 #include "Widgets/Input/SNumericEntryBox.h"
-#include "Widgets/Input/SSlider.h"
 #include "Widgets/Text/SRichTextBlock.h"
 #include "Widgets/Views/SListView.h"
 #include "Widgets/Views/STableRow.h"
@@ -45,14 +45,14 @@
 static FTableViewStyle& GetTransparentViewStyle()
 {
 	static FTableViewStyle TransparentViewStyle = []
-	{
-		const FSlateColorBrush Brush(FLinearColor::Transparent);
+		{
+			const FSlateColorBrush Brush(FLinearColor::Transparent);
 
-		FTableViewStyle ViewStyle = FAppStyle::Get().GetWidgetStyle<FTableViewStyle>("ListView");
-		ViewStyle.SetBackgroundBrush(Brush);
+			FTableViewStyle ViewStyle = FAppStyle::Get().GetWidgetStyle<FTableViewStyle>("ListView");
+			ViewStyle.SetBackgroundBrush(Brush);
 
-		return ViewStyle;
-	}();
+			return ViewStyle;
+		}();
 	return TransparentViewStyle;
 }
 
@@ -65,17 +65,17 @@ static FTableRowStyle& GetTransparentRowStyle()
 			const FSlateColorBrush Brush(FLinearColor::Transparent);
 
 			Style.SetEvenRowBackgroundBrush(Brush)
-				 .SetEvenRowBackgroundHoveredBrush(Brush)
-				 .SetOddRowBackgroundBrush(Brush)
-				 .SetOddRowBackgroundHoveredBrush(Brush)
-				 .SetSelectorFocusedBrush(Brush)
-				 .SetActiveBrush(Brush)
-				 .SetActiveHoveredBrush(Brush)
-				 .SetInactiveBrush(Brush)
-				 .SetInactiveHoveredBrush(Brush)
-				 .SetDropIndicator_Above(Brush)
-				 .SetDropIndicator_Onto(Brush)
-				 .SetDropIndicator_Below(Brush);
+			     .SetEvenRowBackgroundHoveredBrush(Brush)
+			     .SetOddRowBackgroundBrush(Brush)
+			     .SetOddRowBackgroundHoveredBrush(Brush)
+			     .SetSelectorFocusedBrush(Brush)
+			     .SetActiveBrush(Brush)
+			     .SetActiveHoveredBrush(Brush)
+			     .SetInactiveBrush(Brush)
+			     .SetInactiveHoveredBrush(Brush)
+			     .SetDropIndicator_Above(Brush)
+			     .SetDropIndicator_Onto(Brush)
+			     .SetDropIndicator_Below(Brush);
 
 			return Style;
 		}();
@@ -86,9 +86,9 @@ SPVEditorViewport::SPVEditorViewport()
 {
 	FoliageComponents = MakeShared<TArray<TObjectPtr<UMeshComponent>>>();
 	FoliageVisibility = MakeShared<TMap<FString, bool>>();
-	
+
 	PreviewNodeBackgroundBrush = MakeShared<FSlateColorBrush>(FColor{70, 100, 200});
-	
+
 	if (AdvancedPreviewScene)
 	{
 		if (!GetPreviewProfileController()->SetActiveProfile(UDefaultEditorProfiles::GreyAmbientProfileName.ToString()))
@@ -129,13 +129,13 @@ TSharedPtr<SWidget> SPVEditorViewport::BuildViewportToolbar()
 	check(AdvancedPreviewScene.IsValid());
 
 	const FName ToolbarName = "PVE.ViewportToolbar";
-	
+
 	if (!UToolMenus::Get()->IsMenuRegistered(ToolbarName))
 	{
 		UToolMenu* Menu = UToolMenus::Get()->RegisterMenu(ToolbarName, /*Parent=*/NAME_None, EMultiBoxType::SlimHorizontalToolBar);
 		Menu->StyleName = "ViewportToolbar";
 		Menu->AddSection("Left");
-		
+
 		FToolMenuSection& RightSection = Menu->AddSection("Right");
 		RightSection.Alignment = EToolMenuSectionAlign::Last;
 		RightSection.AddEntry(UE::UnrealEd::CreateCameraSubmenu(UE::UnrealEd::FViewportCameraMenuOptions().ShowAll()));
@@ -160,7 +160,7 @@ TSharedPtr<SWidget> SPVEditorViewport::BuildViewportToolbar()
 		Context.AddObject(CreateMannequinWidgetContext(CreateMannequinOffsetWidget()));
 		Context.AddObject(CreateFoliageComponentsContext());
 	}
-	
+
 	return UToolMenus::Get()->GenerateWidget(ToolbarName, Context);
 }
 
@@ -183,11 +183,11 @@ void SPVEditorViewport::BindCommands()
 		FIsActionChecked::CreateSP(this, &SPVEditorViewport::IsShowScaleVisChecked));
 
 	CommandList->MapAction(
-			Commands.ShowStats,
-			FExecuteAction::CreateSP(this, &SPVEditorViewport::ToggleShowStats),
-			FCanExecuteAction(),
-			FIsActionChecked::CreateSP(this, &SPVEditorViewport::IsShowStatsChecked));
-	
+		Commands.ShowStats,
+		FExecuteAction::CreateSP(this, &SPVEditorViewport::ToggleShowStats),
+		FCanExecuteAction(),
+		FIsActionChecked::CreateSP(this, &SPVEditorViewport::IsShowStatsChecked));
+
 	CommandList->MapAction(
 		Commands.AutoFocusViewport,
 		FExecuteAction::CreateSP(this, &SPVEditorViewport::ToggleAutofocusViewport),
@@ -197,10 +197,10 @@ void SPVEditorViewport::BindCommands()
 
 TSharedRef<FEditorViewportClient> SPVEditorViewport::MakeEditorViewportClient()
 {
-	TSharedRef<FEditorViewportClient> ViewportClient = SPCGEditorViewport::MakeEditorViewportClient();
-	ViewportClient->ExposureSettings.bFixed = true;
-	ViewportClient->ExposureSettings.FixedEV100 = -0.2f;
-	return ViewportClient;
+	Client = SPCGEditorViewport::MakeEditorViewportClient();
+	Client->ExposureSettings.bFixed = true;
+	Client->ExposureSettings.FixedEV100 = -0.2f;
+	return Client.ToSharedRef();
 }
 
 void SPVEditorViewport::PopulateViewportOverlays(TSharedRef<SOverlay> Overlay)
@@ -229,7 +229,7 @@ void SPVEditorViewport::PopulateViewportOverlays(TSharedRef<SOverlay> Overlay)
 				.Justification(ETextJustify::Center)
 			]
 		];
-	
+
 	Overlay
 		->AddSlot()
 		.VAlign(VAlign_Top)
@@ -274,7 +274,7 @@ void SPVEditorViewport::PopulateViewportOverlays(TSharedRef<SOverlay> Overlay)
 					SNew(SRichTextBlock)
 					.Text(FText::FromString(TEXT("<TextBlock.ShadowedText>Legend</>")))
 					.Justification(ETextJustify::Center)
-					.TextStyle(&FCoreStyle::Get().GetWidgetStyle<FTextBlockStyle>( "BoldFont" ))
+					.TextStyle(&FCoreStyle::Get().GetWidgetStyle<FTextBlockStyle>("BoldFont"))
 				]
 				+ SVerticalBox::Slot()
 				.FillHeight(1.f)
@@ -299,7 +299,7 @@ void SPVEditorViewport::PopulateViewportOverlays(TSharedRef<SOverlay> Overlay)
 				]
 			]
 		];
-	
+
 	Overlay
 		->AddSlot()
 		.VAlign(VAlign_Bottom)
@@ -391,28 +391,16 @@ void SPVEditorViewport::ResetVisualizationScene()
 {
 	if (MannequinComponent)
 	{
-		AdvancedPreviewScene->RemoveComponent(MannequinComponent);
-		ManagedResources.Remove(MannequinComponent);
 		MannequinComponent->MarkAsGarbage();
 		MannequinComponent = nullptr;
 	}
 
 	if (ScaleVisualizationComponent)
 	{
-		for (const TObjectPtr<UTextRenderComponent>& TextRenderComponent : ScaleVisualizationComponent->GetTextRenderComponents())
-		{
-			if (TextRenderComponent)
-			{
-				AdvancedPreviewScene->RemoveComponent(TextRenderComponent);
-				ManagedResources.Remove(TextRenderComponent);
-			}
-		}
-		AdvancedPreviewScene->RemoveComponent(ScaleVisualizationComponent);
-		ManagedResources.Remove(ScaleVisualizationComponent);
 		ScaleVisualizationComponent->MarkAsGarbage();
 		ScaleVisualizationComponent = nullptr;
 	}
-	
+
 	if (SkeletonVisualizerComponent)
 	{
 		SkeletonVisualizerComponent = nullptr;
@@ -429,7 +417,7 @@ TSharedRef<SWidget> SPVEditorViewport::CreateMannequinOffsetWidget() const
 		.MaxValue(2000.0f)
 		.MinSliderValue(-2000.0f)
 		.MaxSliderValue(2000.0f)
-		.Value_Lambda([this]{ return GetMannequinOffset(); })
+		.Value_Lambda([this] { return GetMannequinOffset(); })
 		.OnValueChanged_Lambda([this](const float NewValue)
 			{
 				SetMannequinOffset(NewValue, false);
@@ -468,7 +456,7 @@ void SPVEditorViewport::ToggleShowScaleVis()
 	UPVEditorSettings* EditorSettings = GetMutableDefault<UPVEditorSettings>();
 	EditorSettings->bShowScaleVisualization = !EditorSettings->bShowScaleVisualization;
 	EditorSettings->SaveConfig();
-	
+
 	SetScaleVisState(EditorSettings->bShowScaleVisualization);
 }
 
@@ -554,7 +542,7 @@ void SPVEditorViewport::OnSetupScene()
 	SPCGEditorViewport::OnSetupScene();
 
 	InitVisualizationScene();
-	
+
 	SetMannequinState(IsShowMannequinChecked());
 	SetScaleVisState(IsShowScaleVisChecked());
 
@@ -580,7 +568,7 @@ FToolMenuEntry SPVEditorViewport::CreateSettingsToolbarMenu()
 {
 	return FToolMenuEntry::InitDynamicEntry(TEXT("PVE Settings"), FNewToolMenuSectionDelegate::CreateLambda([](FToolMenuSection& Section)
 		{
-			FToolMenuEntry& Entry = Section.AddSubMenu(
+			Section.AddSubMenu(
 				/*Name=*/TEXT("PVESettingsSubmenu"),
 				FText::GetEmpty(),
 				FText::GetEmpty(),
@@ -588,12 +576,12 @@ FToolMenuEntry SPVEditorViewport::CreateSettingsToolbarMenu()
 					{
 						const UPVMannequinWidgetContext* MannequinContext = Submenu->FindContext<UPVMannequinWidgetContext>();
 						const UPVFoliageComponentsContext* FoliageContext = Submenu->FindContext<UPVFoliageComponentsContext>();
-						
+
 						// Add preference toggles to this section.
 						FToolMenuSection& OverlaySection = Submenu->FindOrAddSection(TEXT("PVE_OverlaySettings"),
 							LOCTEXT("PVE_OverlaySettingsSectionLabel", "Overlay Settings"));
 						OverlaySection.AddMenuEntry(FPVEditorCommands::Get().ShowStats);
-						
+
 						FToolMenuSection& MannequinSection = Submenu->FindOrAddSection(TEXT("PVE_MannequinSettings"),
 							LOCTEXT("PVE_MannequinSettingsSectionLabel", "Mannequin Settings"));
 						MannequinSection.AddMenuEntry(FPVEditorCommands::Get().ShowMannequin);
@@ -642,7 +630,8 @@ FToolMenuEntry SPVEditorViewport::CreateSettingsToolbarMenu()
 																FExecuteAction::CreateLambda([=]()
 																	{
 																		InstancedStaticMeshComponent->ToggleVisibility(true);
-																		Context->FoliageVisibility->FindOrAdd(ObjectName, true) = InstancedStaticMeshComponent->GetVisibleFlag();
+																		Context->FoliageVisibility->FindOrAdd(ObjectName, true) =
+																			InstancedStaticMeshComponent->GetVisibleFlag();
 																	}
 																),
 																FCanExecuteAction(),
@@ -705,19 +694,19 @@ TArray<EPVRenderType> VisualizationModeToRenderModes(EPVVisualizationMode InMode
 	{
 	case EPVVisualizationMode::Default:
 		return DefautlSettings;
-		
+
 	case EPVVisualizationMode::FoliageMesh:
 		return {EPVRenderType::Mesh, EPVRenderType::Foliage};
 
 	case EPVVisualizationMode::FoliageAttachmentMesh:
-		return {EPVRenderType::Mesh, EPVRenderType::FoliageAttachments};	
+		return {EPVRenderType::Mesh, EPVRenderType::FoliageAttachments};
 
 	case EPVVisualizationMode::FoliageGrid:
 		return {EPVRenderType::FoliageGrid};
 
 	case EPVVisualizationMode::Bones:
 		return {EPVRenderType::Bones};
-	
+
 	case EPVVisualizationMode::Mesh:
 		return {EPVRenderType::Mesh};
 
@@ -735,7 +724,6 @@ TArray<EPVRenderType> VisualizationModeToRenderModes(EPVVisualizationMode InMode
 
 	default:
 		return TArray<EPVRenderType>{};
-		
 	}
 }
 
@@ -831,9 +819,11 @@ FToolMenuEntry SPVEditorViewport::CreateVisualizationModeToolbarMenu()
 												{
 													const FText ModeName = EnumPtr->GetDisplayNameTextByIndex(i);
 													const FString Category = EnumPtr->GetMetaData(TEXT("Category"), i);
-													const ESkeletonVisualizationModes Mode = static_cast<ESkeletonVisualizationModes>(EnumPtr->GetValueByIndex(i));
+													const ESkeletonVisualizationModes Mode = static_cast<ESkeletonVisualizationModes>(EnumPtr->
+														GetValueByIndex(i));
 
-													FToolMenuSection& DebugModeSelectionSection = SubMenu->FindOrAddSection(*Category, FText::FromString(Category));
+													FToolMenuSection& DebugModeSelectionSection = SubMenu->FindOrAddSection(*Category,
+														FText::FromString(Category));
 
 													DebugModeSelectionSection.AddMenuEntry(
 														NAME_None,
@@ -888,8 +878,9 @@ void SPVEditorViewport::OnVisualizationModeChanged(EPVVisualizationMode InMode)
 
 	if (InspectedNodeSettings)
 	{
-		const TArray<EPVRenderType> RenderTypes = VisualizationModeToRenderModes(CurrentVisualizationMode, InspectedNodeSettings->GetDefaultRenderType());
-												
+		const TArray<EPVRenderType> RenderTypes = VisualizationModeToRenderModes(CurrentVisualizationMode,
+			InspectedNodeSettings->GetDefaultRenderType());
+
 		InspectedNodeSettings->SetCurrentRenderType(RenderTypes);
 
 		Invalidate();
@@ -906,11 +897,8 @@ void SPVEditorViewport::OnNodeInspectionChanged(UPVBaseSettings* InSettings)
 	SupportedRenderTypes = InSettings->GetSupportedRenderTypes();
 	CurrentVisualizationMode = EPVVisualizationMode::Default;
 	InspectedNodeSettings = InSettings;
-	
-	InspectedNodeSettings->SetCurrentRenderType(InSettings->GetDefaultRenderType());
 
-	// TODO: @Tayyab add lock visualization based on the inspection lock state
-	// InSettings->bLockInspection
+	InspectedNodeSettings->SetCurrentRenderType(InSettings->GetDefaultRenderType());
 
 	if (IsAutoFocusViewportChecked())
 	{
@@ -930,13 +918,13 @@ void SPVEditorViewport::SetOverlayText(const FText& CurrentlyLockedNodeName)
 void SPVEditorViewport::PopulateStatsOverlayText(const TArrayView<FText> TextItems)
 {
 	static FText NormalTextStyle = FText::FromString(TEXT("<TextBlock.ShadowedText>{0}</>"));
-	
+
 	FTextBuilder FinalText;
 	for (const FText& TextItem : TextItems)
 	{
 		FinalText.AppendLineFormat(NormalTextStyle, TextItem);
 	}
-	
+
 	StatsOverlayText->SetText(FinalText.ToText());
 }
 
diff --git a/Engine/Plugins/Experimental/ProceduralVegetationEditor/Source/ProceduralVegetationEditor/Private/Visualizations/PVDataVisualization.cpp b/Engine/Plugins/Experimental/ProceduralVegetationEditor/Source/ProceduralVegetationEditor/Private/Visualizations/PVDataVisualization.cpp
--- a/Engine/Plugins/Experimental/ProceduralVegetationEditor/Source/ProceduralVegetationEditor/Private/Visualizations/PVDataVisualization.cpp
+++ b/Engine/Plugins/Experimental/ProceduralVegetationEditor/Source/ProceduralVegetationEditor/Private/Visualizations/PVDataVisualization.cpp
@@ -120,16 +120,16 @@ FPCGSetupSceneFunc FPVDataVisualization::GetViewportSetupFunc(const UPCGSettings
 			}
 		}
 
-		for (const auto& ViewportObject : WeakSettings->GetViewportObjects())
+		for (const auto& [MeshObject, Transform] : WeakSettings->GetViewportObjects())
 		{
-			if (UStaticMesh* Mesh = Cast<UStaticMesh>(ViewportObject.Key))
+			if (UStaticMesh* Mesh = Cast<UStaticMesh>(MeshObject))
 			{
 				TObjectPtr<UStaticMeshComponent> StaticMeshComponent = NewObject<UStaticMeshComponent>(GetTransientPackage(), NAME_None, RF_Transient);
 				StaticMeshComponent->SetStaticMesh(Mesh);
 				StaticMeshComponent->UpdateBounds();
-
+		
 				InOutParams.ManagedResources.Add(StaticMeshComponent);
-				InOutParams.Scene->AddComponent(StaticMeshComponent, ViewportObject.Value);
+				InOutParams.Scene->AddComponent(StaticMeshComponent, Transform);
 			}
 		}
 
@@ -160,7 +160,7 @@ void FPVDataVisualization::MeshRenderer(const FManagedArrayCollection& InCollect
 		const TManagedArray<FString>& MaterialPaths = InCollection.GetAttribute<FString>("MaterialPath", FGeometryCollection::MaterialGroup);
 		for (const FString& MaterialPath : MaterialPaths.GetConstArray())
 		{
-			Materials.Add(LoadObject<UMaterialInterface>(nullptr, MaterialPath));
+			Materials.Add(LoadObject<UMaterialInterface>(nullptr, MaterialPath, {}, LOAD_NoWarn | LOAD_Quiet));
 		}
 	}
  
@@ -216,7 +216,6 @@ void FPVDataVisualization::FoliageRenderer(const FManagedArrayCollection& InColl
 	auto OnComponentCreated = [&InOutParams](UMeshComponent* InInstancedComponent)
 	{
 		InOutParams.ManagedResources.Add(InInstancedComponent);
-		InInstancedComponent->RegisterComponentWithWorld(InOutParams.Scene->GetWorld());
 		InOutParams.Scene->AddComponent(InInstancedComponent, FTransform::Identity);
 	};
  
@@ -323,7 +322,7 @@ void FPVDataVisualization::FoliageGridRenderer(const FManagedArrayCollection& In
 	{
 		FString FoliageName = Info.Mesh.ToString();
 		
-		if (UObject* FoliageMesh = LoadObject<UObject>(nullptr, *FoliageName))
+		if (UObject* FoliageMesh = LoadObject<UObject>(nullptr, *FoliageName, {}, LOAD_NoWarn | LOAD_Quiet))
 		{
 			FBox Bounds = GetMeshBounds(Cast<UStaticMesh>(FoliageMesh));
 			if (!Bounds.IsValid)
@@ -348,7 +347,7 @@ void FPVDataVisualization::FoliageGridRenderer(const FManagedArrayCollection& In
 	{
 		const FString& FoliagePath = Info.Mesh.ToString();
 		
-		UObject* FoliageObject = LoadObject<UObject>(nullptr, FoliagePath);
+		UObject* FoliageObject = LoadObject<UObject>(nullptr, FoliagePath, {}, LOAD_NoWarn | LOAD_Quiet);
 		if (UStaticMesh* StaticFoliageMesh = Cast<UStaticMesh>(FoliageObject))
 		{
 			UStaticMeshComponent* MeshComponent = NewObject<UStaticMeshComponent>(GetTransientPackage(), NAME_None, RF_Transient);
diff --git a/Engine/Plugins/Experimental/ProceduralVegetationEditor/Source/ProceduralVegetationEditor/Private/Widgets/SPVExportIem.cpp b/Engine/Plugins/Experimental/ProceduralVegetationEditor/Source/ProceduralVegetationEditor/Private/Widgets/SPVExportIem.cpp
deleted file mode 100644
diff --git a/Engine/Plugins/Experimental/ProceduralVegetationEditor/Source/ProceduralVegetationEditor/Private/Widgets/SPVExportItem.h b/Engine/Plugins/Experimental/ProceduralVegetationEditor/Source/ProceduralVegetationEditor/Private/Widgets/SPVExportItem.h
deleted file mode 100644
diff --git a/Engine/Plugins/Experimental/ProceduralVegetationEditor/Source/ProceduralVegetationEditor/Private/Widgets/SPVExportSelectionDialog.cpp b/Engine/Plugins/Experimental/ProceduralVegetationEditor/Source/ProceduralVegetationEditor/Private/Widgets/SPVExportSelectionDialog.cpp
--- a/Engine/Plugins/Experimental/ProceduralVegetationEditor/Source/ProceduralVegetationEditor/Private/Widgets/SPVExportSelectionDialog.cpp
+++ b/Engine/Plugins/Experimental/ProceduralVegetationEditor/Source/ProceduralVegetationEditor/Private/Widgets/SPVExportSelectionDialog.cpp
@@ -1,73 +1,186 @@
 // Copyright Epic Games, Inc. All Rights Reserved.
 #include "SPVExportSelectionDialog.h"
-#include "Widgets/Layout/SSeparator.h"
-#include "SPrimaryButton.h"
-#include "Widgets/Layout/SUniformGridPanel.h"
+
 #include "Editor.h"
-#include "SPVExportItem.h"
+#include "IStructureDetailsView.h"
 #include "ProceduralVegetation.h"
-#include "PVEditorSettings.h"
+#include "PropertyEditorModule.h"
+#include "SEnumCombo.h"
+#include "SPrimaryButton.h"
+
+#include "Algo/AnyOf.h"
+
 #include "Nodes/PCGEditorGraphNodeBase.h"
 #include "Nodes/PVOutputSettings.h"
+
+#include "Widgets/Input/SSegmentedControl.h"
 #include "Widgets/Input/STextComboBox.h"
 #include "Widgets/Layout/SScrollBox.h"
+#include "Widgets/Layout/SUniformGridPanel.h"
 
 #define LOCTEXT_NAMESPACE "SPVExportSelectionDialog"
 
+class SExportEntryStatRow : public SMultiColumnTableRow<TSharedPtr<FStatType>>
+{
+public:
+	SLATE_BEGIN_ARGS(SExportEntryStatRow)
+		{}
+
+		SLATE_ARGUMENT(TSharedPtr<FStatType>, Item)
+	SLATE_END_ARGS()
+
+	void Construct(const FArguments& InArgs, const TSharedRef<STableViewBase>& OwnerTableView)
+	{
+		Item = InArgs._Item;
+		SMultiColumnTableRow<TSharedPtr<FStatType>>::Construct(FSuperRowType::FArguments(), OwnerTableView);
+	}
+
+	virtual TSharedRef<SWidget> GenerateWidgetForColumn(const FName& ColumnName) override
+	{
+		if (!Item.IsValid())
+		{
+			return SNew(STextBlock).Text(FText::GetEmpty());
+		}
+
+		if (ColumnName == "Name")
+		{
+			return SNew(STextBlock)
+				.Text(Item->Key);
+		}
+		else if (ColumnName == "Count")
+		{
+			return SNew(STextBlock)
+				.Justification(ETextJustify::Right)
+				.Text(FText::AsNumber(Item->Value));
+		}
+
+		return SNew(STextBlock)
+			.Text(FText::FromString(TEXT("Invalid Column")));
+	}
+
+private:
+	TSharedPtr<FStatType> Item;
+};
+
 void SPVExportSelectionDialog::Construct(const FArguments& InArgs)
 {
-	Graph = InArgs._Graph;
-	SelectedNodes = InArgs._SelectedNodes;
+	ExportEntries = InArgs._ExportEntries;
+
+	const bool bIsAnyNodeSelected = Algo::AnyOf(
+		ExportEntries,
+		[](const TObjectPtr<UPVExportEntry>& ExportEntry)
+			{
+				return ExportEntry->bIsNodeSelected;
+			}
+	);
+	OnExportTypeChanged(
+		bIsAnyNodeSelected
+		? EPVExportType::Selected
+		: EPVExportType::Batch
+	);
+
+	FPropertyEditorModule& PropertyEditorModule = FModuleManager::GetModuleChecked<FPropertyEditorModule>("PropertyEditor");
+	FDetailsViewArgs DetailsViewArgs;
+	DetailsViewArgs.bAllowSearch = true;
+	DetailsViewArgs.NameAreaSettings = FDetailsViewArgs::HideNameArea;
+	DetailsView = PropertyEditorModule.CreateDetailView(DetailsViewArgs);
 
-	BuildExportTypeOptions();
-	
-	SWindow::Construct(SWindow::FArguments()
+	SWindow::Construct(
+		SWindow::FArguments()
 		.Title(InArgs._Title)
 		.SupportsMinimize(false)
 		.SupportsMaximize(false)
-		.ClientSize(FVector2D(350, 450))
+		.ClientSize(FVector2D(800, 600)
+		)
 		[
 			SNew(SVerticalBox)
 			+ SVerticalBox::Slot()
-			.AutoHeight()
-			.Padding(2, 2, 2, 0)
+			.Padding(2, 2, 2, 2)
 			[
-				SNew(SBorder)
-				.BorderImage(FAppStyle::GetBrush("ToolPanel.GroupBorder"))
+				SNew(SSplitter)
+				+ SSplitter::Slot()
+				.Resizable(true)
+				.SizeRule(SSplitter::FractionOfParent)
+				.Value(0.3f)
 				[
-					SNew(SHorizontalBox)
-					+ SHorizontalBox::Slot()
-					.HAlign(HAlign_Left)
-					.VAlign(VAlign_Center)
-					.AutoWidth()
-					.Padding(2, 10, 20, 10)
+					SNew(SVerticalBox)
+					+ SVerticalBox::Slot()
+					.AutoHeight()
+					.VAlign(VAlign_Top)
+					.HAlign(HAlign_Center)
 					[
-						SNew(STextBlock)
-						.Text(LOCTEXT("PVExportSelectionDialog_ExportTypeLabel", "Export type"))
+						SNew(SSegmentedControl<EPVExportType>)
+						.Value(this, &SPVExportSelectionDialog::GetExportType)
+						.OnValueChanged(this, &SPVExportSelectionDialog::OnExportTypeChanged)
+						+ SSegmentedControl<EPVExportType>::Slot(EPVExportType::Batch)
+						.HAlign(HAlign_Center)
+						.VAlign(VAlign_Center)
+						.ToolTip(LOCTEXT("PVExportTypeAllLabel", "Batch export all output nodes"))
+						[
+							SNew(STextBlock)
+							.Text(LOCTEXT("PVBatchExportLabel", "Batch"))
+						]
+						+ SSegmentedControl<EPVExportType>::Slot(EPVExportType::Selected)
+						.HAlign(HAlign_Center)
+						.VAlign(VAlign_Center)
+						.ToolTip(LOCTEXT("PVExportTypeAllLabel", "Export only selected output nodes in the graph"))
+						[
+							SNew(STextBlock)
+							.Text(LOCTEXT("PVSelectionExportLabel", "Selected"))
+						]
 					]
-					
-					+ SHorizontalBox::Slot()
-					.FillWidth(1.0f)
-					.HAlign(HAlign_Fill)
+					+ SVerticalBox::Slot()
+					.FillHeight(1.0f)
+					.VAlign(VAlign_Fill)
 					[
-						SNew(STextComboBox)
-						.OptionsSource(&ExportTypes)
-						.InitiallySelectedItem(CurrentExportType)
-						.OnSelectionChanged(this, &SPVExportSelectionDialog::OnExportTypeChanged)
+						SNew(SSplitter)
+						.Orientation(EOrientation::Orient_Vertical)
+						+ SSplitter::Slot()
+						.Resizable(true)
+						.SizeRule(SSplitter::FractionOfParent)
+						.Value(0.73f)
+						[
+							SNew(SBorder)
+							.BorderImage(FAppStyle::GetBrush("ToolPanel.GroupBorder"))
+							[
+								SAssignNew(ExportEntriesTableView, SListView<TObjectPtr<UPVExportEntry>>)
+								.ListItemsSource(&ExportEntries)
+								.OnGenerateRow(this, &SPVExportSelectionDialog::OnGenerateExportItem)
+								.OnSelectionChanged(this, &SPVExportSelectionDialog::OnExportEntrySelected)
+								.SelectionMode(ESelectionMode::Single)
+								.ClearSelectionOnClick(false)
+							]
+						]
+						+ SSplitter::Slot()
+						.Resizable(true)
+						.SizeRule(SSplitter::FractionOfParent)
+						.Value(0.27f)
+						[
+							SAssignNew(StatTableView, SListView<TSharedPtr<FStatType>>)
+							.ListItemsSource(&CurrentExportEntryStats)
+							.OnGenerateRow(this, &SPVExportSelectionDialog::OnGenerateStatItem)
+							.HeaderRow(
+								SNew(SHeaderRow)
+								+ SHeaderRow::Column("Name")
+								.DefaultLabel(FText::FromString("Export Stats"))
+								.FillWidth(0.7f)
+								+ SHeaderRow::Column("Count")
+								.DefaultLabel(FText::FromString(""))
+								.FillWidth(0.3f)
+							)
+							.SelectionMode(ESelectionMode::None)
+						]
 					]
 				]
-			]
-			+ SVerticalBox::Slot() // Add user input block
-			.Padding(2, 2, 2, 4)
-			[
-				SNew(SBorder)
-				.BorderImage(FAppStyle::GetBrush("ToolPanel.GroupBorder"))
+				+ SSplitter::Slot()
+				.Resizable(true)
+				.SizeRule(SSplitter::FractionOfParent)
+				.Value(0.7f)
 				[
-					SNew(SVerticalBox)
-					+ SVerticalBox::Slot()
-					.FillHeight(1)
+					SNew(SBorder)
+					.BorderImage(FAppStyle::GetBrush("ToolPanel.GroupBorder"))
 					[
-						SAssignNew(ScrollBox,SScrollBox)
+						DetailsView.ToSharedRef()
 					]
 				]
 			]
@@ -75,7 +188,7 @@ void SPVExportSelectionDialog::Construct(const FArguments& InArgs)
 			.AutoHeight()
 			.HAlign(HAlign_Right)
 			.VAlign(VAlign_Bottom)
-			.Padding(8.f, 16.f)
+			.Padding(8.0f, 6.0f, 8.0f, 8.0f)
 			[
 				SNew(SUniformGridPanel)
 				.MinDesiredSlotWidth(FAppStyle::GetFloat("StandardDialog.MinDesiredSlotWidth"))
@@ -95,106 +208,126 @@ void SPVExportSelectionDialog::Construct(const FArguments& InArgs)
 					.OnClicked(this, &SPVExportSelectionDialog::OnButtonClick, EAppReturnType::Cancel)
 				]
 			]
-		]);
+		]
+	);
 
-	AddExportItems();
-}
-
-EPVExportType SPVExportSelectionDialog::GetExportType()
-{
-	EPVExportType ExportType = EPVExportType::Selection;
-	
-	if (ExportTypes.Contains(CurrentExportType))
+	if (ExportEntriesTableView)
 	{
-		int SelectionIndex = ExportTypes.IndexOfByKey(CurrentExportType);
-		ExportType = (EPVExportType)SelectionIndex;
+		ExportEntriesTableView->SetSelection(ExportEntries[0]);
 	}
-
-	return ExportType;
 }
 
-void SPVExportSelectionDialog::AddExportItems()
+TSharedRef<ITableRow> SPVExportSelectionDialog::OnGenerateExportItem(
+	TObjectPtr<UPVExportEntry> InExportEntry,
+	const TSharedRef<STableViewBase>& OwnerTable
+)
 {
-	ScrollBox->ClearChildren();
-	
-	if (GetExportType() == EPVExportType::BatchExport)
-	{
-		Graph->ForEachNode([this](UPCGNode* Node)
-		{
-			const UPVOutputSettings* OutputSettings = Cast<UPVOutputSettings>(Node->GetSettings());
-
-			if (OutputSettings)
-			{
-				FText OutputName = Node->GetNodeTitle(EPCGNodeTitleType::FullTitle);
-				
-				ScrollBox->AddSlot()
-				.Padding(5)
-				[
-					SNew(SPVExportItem)
-					.Name(OutputName)
-					.OnOutputSelectionChanged(this, &SPVExportSelectionDialog::OnOutputSelectionChanged)
-					.bIsSelected(OutputSettings->ExportSettings.bShouldExport)
-				];
-				ScrollBox->AddSlot()
-				[
-					SNew(SSeparator)
-					.Thickness(1)
-				];
-			}
-
-			return true;
-		});
-	}
-	else
-	{
-		FString OutputName;
-
-		if (SelectedNodes.IsEmpty())
-		{
-			OutputName = "None";
-		}
-
-		for (int i = 0; i < SelectedNodes.Num(); i++)
-		{
-			if (const UPCGNode* SelectedNode = SelectedNodes[i] ? SelectedNodes[i]->GetPCGNode() : nullptr)
-			{
-				if (const UPVOutputSettings* OutputSettings = Cast<UPVOutputSettings>(SelectedNode->GetSettings()))
-				{
-					if (OutputName.IsEmpty())
-					{
-						OutputName = FString::Printf(TEXT("%s"), *SelectedNode->GetNodeTitle(EPCGNodeTitleType::FullTitle).ToString());
-					}
-					else
-					{
-						OutputName = FString::Printf(TEXT("%s\n%s"), *OutputName,*SelectedNode->GetNodeTitle(EPCGNodeTitleType::FullTitle).ToString());
-					}
-				}	
-			}
-		}
-		
-		FText Outputs = FText::FromString(FString::Printf(TEXT("%s"), *OutputName));
-			
-		ScrollBox->AddSlot()
-		.Padding(5)
+	const FText OutputName = InExportEntry->Node->GetNodeTitle(EPCGNodeTitleType::FullTitle);
+	return SNew(STableRow<TObjectPtr<UPVExportEntry>>, OwnerTable)
+		.Padding(2.0f)
 		[
 			SNew(SHorizontalBox)
 			+ SHorizontalBox::Slot()
+			.HAlign(HAlign_Left)
+			.VAlign(VAlign_Center)
 			.AutoWidth()
 			[
-				SNew(STextBlock)
-				.Text(FText::FromString(FString::Printf(TEXT("Current outputs(s): "))))	
+				SNew(SCheckBox)
+				.IsChecked_Lambda([=]()
+					{
+						return InExportEntry->bExport
+							? ECheckBoxState::Checked
+							: ECheckBoxState::Unchecked;
+					})
+				.OnCheckStateChanged_Lambda([=, this](const ECheckBoxState NewState)
+					{
+						InExportEntry->bExport = NewState == ECheckBoxState::Checked;
+						SetExportType();
+					})
 			]
 			+ SHorizontalBox::Slot()
-			.AutoWidth()
+			.HAlign(HAlign_Left)
+			.VAlign(VAlign_Center)
+			.FillWidth(1.0f)
+			.Padding(2.0f, 0.0f)
 			[
 				SNew(STextBlock)
-				.Text(Outputs)	
+				.Justification(ETextJustify::Left)
+				.Text(OutputName)
 			]
 		];
+}
+
+TSharedRef<ITableRow> SPVExportSelectionDialog::OnGenerateStatItem(TSharedPtr<FStatType> InStatEntry, const TSharedRef<STableViewBase>& OwnerTable)
+{
+	return SNew(SExportEntryStatRow, OwnerTable).Item(InStatEntry);
+}
+
+void SPVExportSelectionDialog::OnExportEntrySelected(TObjectPtr<UPVExportEntry> InExportEntry, ESelectInfo::Type SelectInfo)
+{
+	check(DetailsView);
+	CurrentlySelectedEntry = InExportEntry;
+	if (CurrentlySelectedEntry)
+	{
+		SetupCurrentExportEntryStats();
 	}
+	else
+	{
+		CurrentExportEntryStats.Empty();
+	}
+	if (StatTableView)
+	{
+		StatTableView->RebuildList();
+	}
+	DetailsView->SetObject(
+		CurrentlySelectedEntry
+		? CurrentlySelectedEntry->Settings
+		: nullptr,
+		true
+	);
+}
 
-	UpdateExportButtonState();
+void SPVExportSelectionDialog::SetupCurrentExportEntryStats()
+{
+	check(CurrentlySelectedEntry);
+	CurrentExportEntryStats.Empty();
+	if (CurrentlySelectedEntry->Stats.NumPoints > -1)
+	{
+		CurrentExportEntryStats.Emplace(MakeShared<TPair<FText, int32>>(FText::FromString("Points"), CurrentlySelectedEntry->Stats.NumPoints));
+	}
+	if (CurrentlySelectedEntry->Stats.NumBranches > -1)
+	{
+		CurrentExportEntryStats.Emplace(MakeShared<TPair<FText, int32>>(FText::FromString("Branches"), CurrentlySelectedEntry->Stats.NumBranches));
+	}
+	if (CurrentlySelectedEntry->Stats.NumVertices > -1)
+	{
+		CurrentExportEntryStats.Emplace(MakeShared<TPair<FText, int32>>(FText::FromString("Vertices (Trunk)"),
+			CurrentlySelectedEntry->Stats.NumVertices));
+	}
+	if (CurrentlySelectedEntry->Stats.NumTris > -1)
+	{
+		CurrentExportEntryStats.Emplace(
+			MakeShared<TPair<FText, int32>>(FText::FromString("Triangles (Trunk)"), CurrentlySelectedEntry->Stats.NumTris));
+	}
+	if (CurrentlySelectedEntry->Stats.NumBones > -1)
+	{
+		CurrentExportEntryStats.Emplace(MakeShared<TPair<FText, int32>>(FText::FromString("Bones"), CurrentlySelectedEntry->Stats.NumBones));
+	}
+	if (CurrentlySelectedEntry->bExportFoliage)
+	{
+		if (CurrentlySelectedEntry->Stats.NumFoliage > -1)
+		{
+			CurrentExportEntryStats.Emplace(MakeShared<TPair<FText, int32>>(FText::FromString("Unique Foliage"),
+				CurrentlySelectedEntry->Stats.NumFoliage));
+		}
+		if (CurrentlySelectedEntry->Stats.NumFoliageInstances > -1)
+		{
+			CurrentExportEntryStats.Emplace(MakeShared<TPair<FText, int32>>(FText::FromString("Num Foliage Instances"),
+				CurrentlySelectedEntry->Stats.NumFoliageInstances));
+		}
+	}
 }
+
 FReply SPVExportSelectionDialog::OnButtonClick(EAppReturnType::Type ButtonID)
 {
 	UserResponse = ButtonID;
@@ -211,108 +344,56 @@ FReply SPVExportSelectionDialog::OnButtonClick(EAppReturnType::Type ButtonID)
 	return FReply::Handled();
 }
 
-void SPVExportSelectionDialog::OnExportTypeChanged(TSharedPtr<FString> Selection, ESelectInfo::Type SelectInfo)
+EPVExportType SPVExportSelectionDialog::GetExportType() const
 {
-	if (!Selection->IsEmpty())
-	{
-		const int SelectionIndex = ExportTypes.IndexOfByKey(Selection);
-
-		if (ExportTypes.IsValidIndex(SelectionIndex))
-		{
-			CurrentExportType = ExportTypes[SelectionIndex];
-		
-			UPVEditorSettings* EditorSettings = GetMutableDefault<UPVEditorSettings>();
-			EditorSettings->ExportType = (EPVExportType)SelectionIndex;
-
-			AddExportItems();	
-		}
-	}
+	return ExportType;
 }
 
-void SPVExportSelectionDialog::UpdateExportButtonState()
+void SPVExportSelectionDialog::OnExportTypeChanged(const EPVExportType InNewValue)
 {
-	bool HasValidExport = false;
-
-	if (GetExportType() == EPVExportType::BatchExport)
+	ExportType = InNewValue;
+	if (ExportType == EPVExportType::Batch)
 	{
-		if (Graph)
+		for (const TObjectPtr<UPVExportEntry>& ExportEntry : ExportEntries)
 		{
-			Graph->ForEachNode([this, &HasValidExport](UPCGNode* Node)
-			{
-				if (UPVOutputSettings* OutputSettings = Cast<UPVOutputSettings>(Node->GetSettings()))
-				{
-					HasValidExport |= OutputSettings->ExportSettings.bShouldExport;
-				}
-		
-				return true;
-			});	
-		}	
+			ExportEntry->bExport = true;
+		}
 	}
-	else
+	else if (ExportType == EPVExportType::Selected)
 	{
-		for (int i = 0; i < SelectedNodes.Num(); i++)
+		for (const TObjectPtr<UPVExportEntry>& ExportEntry : ExportEntries)
 		{
-			if (const UPCGNode* SelectedNode = SelectedNodes[i] ? SelectedNodes[i]->GetPCGNode() : nullptr)
-			{
-				const UPVOutputSettings* OutputSettings = Cast<UPVOutputSettings>(SelectedNode->GetSettings());
-
-				if (OutputSettings)
-				{
-					HasValidExport = true;
-				}
-			}
+			ExportEntry->bExport = ExportEntry->bIsNodeSelected;
 		}
 	}
-
-	ExportButton->SetEnabled(HasValidExport);
-	
 }
 
-EAppReturnType::Type SPVExportSelectionDialog::ShowModal()
-{
-	GEditor->EditorAddModalWindow(SharedThis(this));
-	return UserResponse;
-}
-void SPVExportSelectionDialog::OnOutputSelectionChanged(const FString ItemName, ECheckBoxState NewState)
+void SPVExportSelectionDialog::SetExportType()
 {
-	Graph->ForEachNode([this, &ItemName, &NewState](UPCGNode* Node)
+	bool bSelectedMode = true;
+	for (const TObjectPtr<UPVExportEntry>& ExportEntry : ExportEntries)
 	{
-		if(ItemName == Node->GetNodeTitle(EPCGNodeTitleType::FullTitle).ToString())
-		{
-			if (UPVOutputSettings* OutputSettings = Cast<UPVOutputSettings>(Node->GetSettings()))
-			{
-				OutputSettings->ExportSettings.bShouldExport = NewState ==ECheckBoxState::Checked;
-			}
-		}
-		
-		return true;
-	});
-
-	UpdateExportButtonState();
+		bSelectedMode &= ExportEntry->bExport == ExportEntry->bIsNodeSelected;
+	}
+	ExportType = bSelectedMode
+		? EPVExportType::Selected
+		: EPVExportType::Batch;
 }
 
-void SPVExportSelectionDialog::BuildExportTypeOptions()
+void SPVExportSelectionDialog::UpdateExportButtonState()
 {
-	const UEnum* EnumPtr = StaticEnum<EPVExportType>();
-	                      
-	if (EnumPtr)
+	bool HasValidExport = false;
+	for (const TObjectPtr<UPVExportEntry>& ExportEntry : ExportEntries)
 	{
-		for (int32 i = 0; i < EnumPtr->NumEnums() - 1; i++)
-		{
-			FText ModeName = EnumPtr->GetDisplayNameTextByIndex(i);
-			ExportTypes.Add(MakeShared<FString>(ModeName.ToString()));
-		}
+		HasValidExport |= ExportEntry->bExport;
 	}
+	ExportButton->SetEnabled(HasValidExport);
+}
 
-	if (ExportTypes.Num() > 0)
-	{
-		const UPVEditorSettings* EditorSettings = GetDefault<UPVEditorSettings>();
-		const int ExportTypeIndex = (int)EditorSettings->ExportType; 
-		
-		if (ExportTypeIndex >= 0 && ExportTypeIndex < ExportTypes.Num())
-		{
-			CurrentExportType = ExportTypes[ExportTypeIndex];
-		}
-	}
+EAppReturnType::Type SPVExportSelectionDialog::ShowModal()
+{
+	GEditor->EditorAddModalWindow(SharedThis(this));
+	return UserResponse;
 }
+
 #undef LOCTEXT_NAMESPACE
diff --git a/Engine/Plugins/Experimental/ProceduralVegetationEditor/Source/ProceduralVegetationEditor/Private/Widgets/SPVExportSelectionDialog.h b/Engine/Plugins/Experimental/ProceduralVegetationEditor/Source/ProceduralVegetationEditor/Private/Widgets/SPVExportSelectionDialog.h
--- a/Engine/Plugins/Experimental/ProceduralVegetationEditor/Source/ProceduralVegetationEditor/Private/Widgets/SPVExportSelectionDialog.h
+++ b/Engine/Plugins/Experimental/ProceduralVegetationEditor/Source/ProceduralVegetationEditor/Private/Widgets/SPVExportSelectionDialog.h
@@ -1,25 +1,44 @@
 // Copyright Epic Games, Inc. All Rights Reserved.
 #pragma once
+
 #include "CoreMinimal.h"
+
+#include "Exporter/PVExporter.h"
+
 #include "Input/Reply.h"
+
 #include "Widgets/SWindow.h"
+#include "Widgets/Views/SListView.h"
+
+#include "SPVExportSelectionDialog.generated.h"
 
+class IDetailsView;
+class IStructureDetailsView;
+class ITableRow;
+class STableViewBase;
 class UPCGEditorGraphNodeBase;
 class UPVOutputSettings;
 class SScrollBox;
 class UProceduralVegetationGraph;
 class UPCGNode;
 class SPrimaryButton;
-enum class EPVExportType;
+
+using FStatType = TPair<FText, int32>;
+
+UENUM()
+enum class EPVExportType : int32
+{
+	Batch,
+	Selected
+};
 
 class SPVExportSelectionDialog : public SWindow
 {
 	SLATE_BEGIN_ARGS(SPVExportSelectionDialog)
-	{}
-	SLATE_ARGUMENT(FText, Title)
-	SLATE_ARGUMENT(const UProceduralVegetationGraph*, Graph)
-	SLATE_ARGUMENT(TArray<FText>, ExportItems)
-	SLATE_ARGUMENT(TArray<TObjectPtr<UPCGEditorGraphNodeBase>>, SelectedNodes)
+		{}
+
+		SLATE_ARGUMENT(FText, Title)
+		SLATE_ARGUMENT(TArray<TObjectPtr<UPVExportEntry>>, ExportEntries)
 	SLATE_END_ARGS()
 
 public:
@@ -28,26 +47,37 @@ public:
 	/** Displays the dialog in a blocking fashion */
 	EAppReturnType::Type ShowModal();
 
-protected:
+private:
 	FReply OnButtonClick(EAppReturnType::Type ButtonID);
-	EAppReturnType::Type UserResponse = EAppReturnType::Type::Cancel;
+
+	TSharedRef<ITableRow> OnGenerateExportItem(TObjectPtr<UPVExportEntry> InExportEntry, const TSharedRef<STableViewBase>& OwnerTable);
+	TSharedRef<ITableRow> OnGenerateStatItem(TSharedPtr<FStatType> InStatEntry, const TSharedRef<STableViewBase>& OwnerTable);
+
+	void OnExportEntrySelected(TObjectPtr<UPVExportEntry> InExportEntry, ESelectInfo::Type SelectInfo = ESelectInfo::Direct);
+	void SetupCurrentExportEntryStats();
 	
+	EPVExportType GetExportType() const;
+	void OnExportTypeChanged(const EPVExportType InNewValue);
+	void SetExportType();
+
+	void UpdateExportButtonState();
+
 private:
 	TSharedPtr<SScrollBox> ScrollBox;
+
 	TSharedPtr<SPrimaryButton> ExportButton;
-	const UProceduralVegetationGraph* Graph = nullptr;
-	TArray<TObjectPtr<UPCGEditorGraphNodeBase>> SelectedNodes;
+	TSharedPtr<SListView<TSharedPtr<FStatType>>> StatTableView;
+	TSharedPtr<SListView<TObjectPtr<UPVExportEntry>>> ExportEntriesTableView;
 
-	TArray<TSharedPtr<FString>> ExportTypes;
-	TSharedPtr<FString> CurrentExportType;
+	TArray<TObjectPtr<UPVExportEntry>> ExportEntries;
+	
+	TSharedPtr<IDetailsView> DetailsView;
 
-	EPVExportType GetExportType();
+	TObjectPtr<UPVExportEntry> CurrentlySelectedEntry = nullptr;
 	
-	void AddExportItems();
-	void OnOutputSelectionChanged(const FString ItemName, ECheckBoxState NewState);
+	TArray<TSharedPtr<FStatType>> CurrentExportEntryStats;
 
-	void BuildExportTypeOptions();
-	void OnExportTypeChanged(TSharedPtr<FString> Selection, ESelectInfo::Type SelectInfo);
+	EPVExportType ExportType = EPVExportType::Selected;
 
-	void UpdateExportButtonState();
+	EAppReturnType::Type UserResponse = EAppReturnType::Type::Cancel;
 };
diff --git a/Engine/Plugins/Experimental/ProceduralVegetationEditor/Source/ProceduralVegetationEditor/Public/PVEditorSettings.h b/Engine/Plugins/Experimental/ProceduralVegetationEditor/Source/ProceduralVegetationEditor/Public/PVEditorSettings.h
--- a/Engine/Plugins/Experimental/ProceduralVegetationEditor/Source/ProceduralVegetationEditor/Public/PVEditorSettings.h
+++ b/Engine/Plugins/Experimental/ProceduralVegetationEditor/Source/ProceduralVegetationEditor/Public/PVEditorSettings.h
@@ -6,17 +6,11 @@
 
 #include "PVEditorSettings.generated.h"
 
-UENUM()
-enum class EPVExportType
-{
-	Selection	UMETA(DisplayName = "Selected output(s) node"),
-	BatchExport UMETA(DisplayName = "Batch export")
-};
-
 UCLASS(MinimalAPI, config=EditorPerProjectUserSettings, meta = (DisplayName = "Procedural Vegetation Editor Settings"))
 class UPVEditorSettings : public UObject
 {
 	GENERATED_BODY()
+
 public:
 	/** Color used for data pins of type GrowthData */
 	UPROPERTY(EditAnywhere, config, Category = Node, meta = (HideAlphaChannel))
@@ -44,7 +38,4 @@ public:
 
 	UPROPERTY(EditAnywhere, config, Category = Export)
 	bool bAutoFocusViewport = false;
-
-	UPROPERTY(EditAnywhere, config, Category = Settings)
-	EPVExportType ExportType = EPVExportType::Selection;
 };
