[PVE] Fixed foliage palette deprecation issue
Issue: For old graphs, the Foliage palette could be connected to the output without adding a foliage distributor. This case was not handled correctly during deprecation fix.
Fix/How: Now we only add an edge when we find a valid "Foliage"- labeled pin in the target node.

#Synced-CL 50934585

diff --git a/Engine/Plugins/Experimental/ProceduralVegetationEditor/Source/ProceduralVegetation/Private/Nodes/PVFoliagePaletteSettings.cpp b/Engine/Plugins/Experimental/ProceduralVegetationEditor/Source/ProceduralVegetation/Private/Nodes/PVFoliagePaletteSettings.cpp
--- a/Engine/Plugins/Experimental/ProceduralVegetationEditor/Source/ProceduralVegetation/Private/Nodes/PVFoliagePaletteSettings.cpp
+++ b/Engine/Plugins/Experimental/ProceduralVegetationEditor/Source/ProceduralVegetation/Private/Nodes/PVFoliagePaletteSettings.cpp
@@ -84,7 +84,7 @@ void UPVFoliagePaletteSettings::ApplyDeprecation(UPCGNode* InOutNode)
 	
 	for (const auto& Connection : PinConnectionsToUpdate)
 	{
-		if (Connection.Node.IsValid() && Connection.TargetNode.IsValid())
+		if (Connection.Node.IsValid() && Connection.TargetNode.IsValid() && Connection.TargetNode->GetInputPin(Connection.ToPinLabel))
 		{
 			Connection.Node->AddEdgeTo(Connection.FromPinLabel, Connection.TargetNode.Get(), Connection.ToPinLabel);
 		}
