Refactor growth data attribute generation

The purpose of this CL is to make the plant traversal and attribute generation logic more robust such that it can handle both fused and split points, as well as improving the ergonomics of the AttributesHelper and PlantTraversalHelper.

Refactor Plant traversal functions:
- Move PVE plant traversal functions from PVAttributesHelper to PVPlantTaversalHelper
- Refactor PVPlantTaversalHelper to allow for traversal of both fused and split point plants
- Refactor PVTransformHelper to allow the caller to supply custom point state

Refactor PVAttributesHelper
- Add typed manged array accessor (PVAttributes.h). This is used to allow for strongly typed use of the TManagedAttributeAccessor type.
- Refactor PVAttributesHelper to allow for implicit conversion between FManagedArrayCollection and the needed attributes for a given attribute computation function
- Add new function for re-computing all attributes for a given collection

Attribute generation improvements
- Fix incorrectly generated Apical and Axillary direction vectors

#Synced-CL 51008641

diff --git a/Engine/Plugins/Experimental/ProceduralVegetationEditor/Source/ProceduralVegetation/Private/Helpers/PVAttributesHelper.cpp b/Engine/Plugins/Experimental/ProceduralVegetationEditor/Source/ProceduralVegetation/Private/Helpers/PVAttributesHelper.cpp
--- a/Engine/Plugins/Experimental/ProceduralVegetationEditor/Source/ProceduralVegetation/Private/Helpers/PVAttributesHelper.cpp
+++ b/Engine/Plugins/Experimental/ProceduralVegetationEditor/Source/ProceduralVegetation/Private/Helpers/PVAttributesHelper.cpp
@@ -1,50 +1,41 @@
 // Copyright Epic Games, Inc. All Rights Reserved.
 
 #include "Helpers/PVAttributesHelper.h"
+#include "Helpers/PVPlantTraversalHelper.h"
 
-void PV::AttributesHelper::ComputeBudNumbers(
-	FPointBudNumberAttributeView PointBudNumberAttribute,
-	FBranchSourceBudNumberAttributeView BranchSourceBudNumberAttribute,
-	FBranchPointsAttributeConstView BranchPointsAttribute
-)
+void PV::AttributesHelper::ComputeBudNumbers(const FComputeBudNumbersAttributes& InAttributes)
 {
-	for (int32 PointIndex = 0; PointIndex < PointBudNumberAttribute.Num(); ++PointIndex)
+	for (int32 PointIndex = 0; PointIndex < InAttributes.PointBudNumber.Num(); ++PointIndex)
 	{
-		PointBudNumberAttribute[PointIndex] = PointIndex + 1;
+		InAttributes.PointBudNumber[PointIndex] = PointIndex + 1;
 	}
 
-	for (int32 BranchIndex = 0; BranchIndex < BranchSourceBudNumberAttribute.Num(); ++BranchIndex)
+	for (int32 BranchIndex = 0; BranchIndex < InAttributes.BranchSourceBudNumber.Num(); ++BranchIndex)
 	{
-		check(BranchPointsAttribute[BranchIndex].Num() > 1);
-		const int32 BranchRootPointIndex = BranchPointsAttribute[BranchIndex][0];
-		BranchSourceBudNumberAttribute[BranchIndex] = PointBudNumberAttribute[BranchRootPointIndex];
+		check(InAttributes.BranchPoints[BranchIndex].Num() > 1);
+		const int32 BranchRootPointIndex = InAttributes.BranchPoints[BranchIndex][0];
+		InAttributes.BranchSourceBudNumber[BranchIndex] = InAttributes.PointBudNumber[BranchRootPointIndex];
 	}
 }
 
-void PV::AttributesHelper::ComputeNjordPixelIndex(
-	PV::FPointNjordPixelIndexAttributeView PointNjordPixelIndexAttribute,
-	PV::FBranchPointsAttributeConstView BranchPointsAttribute,
-	PV::FBranchParentNumberAttributeConstView BranchParentNumberAttribute,
-	PV::FBranchNumberAttributeConstView BranchNumberAttribute,
-	PV::FBranchChildrenAttributeConstView BranchChildrenAttribute
-)
+void PV::AttributesHelper::ComputeNjordPixelIndex(const FComputeNjordPixelIndexAttributes& InAttributes)
 {
-	using namespace PV::AttributesHelper;
+	using namespace PV::PlantTraversalHelper;
 
 	RecursiveWalkBranches<int32>(
-		BranchParentNumberAttribute,
-		BranchChildrenAttribute,
-		BranchNumberAttribute,
+		InAttributes.BranchParentNumber,
+		InAttributes.BranchChildren,
+		InAttributes.BranchNumber,
 		1, 
 		[&](int32 BranchIndex, int32& NjordPixelIndex)->EForEachResult
 		{
-			const bool bIsTrunk = GetBranchParentIndex(BranchParentNumberAttribute, BranchNumberAttribute, BranchIndex) == INDEX_NONE;
-			const TArray<int32>& BranchPoints = BranchPointsAttribute[BranchIndex];
+			const bool bIsTrunk = GetBranchParentIndex(InAttributes.BranchParentNumber, InAttributes.BranchNumber, BranchIndex) == INDEX_NONE;
+			const TArray<int32>& BranchPoints = InAttributes.BranchPoints[BranchIndex];
 			for (int32 i = 0; i < BranchPoints.Num(); ++i)
 			{
 				if (bIsTrunk || i != 0)
 				{
-					PointNjordPixelIndexAttribute[BranchPoints[i]] = NjordPixelIndex;
+					InAttributes.PointNjordPixelIndex[BranchPoints[i]] = NjordPixelIndex;
 					NjordPixelIndex += 1;
 				}
 			}
@@ -54,45 +45,42 @@ void PV::AttributesHelper::ComputeNjordPixelIndex(
 	);
 }
 
-void PV::AttributesHelper::ComputeLengthFromRoot(
-	PV::FPointLengthFromRootAttributeView PointLengthFromRootAttribute,
-	PV::FPointLengthFromSeedAttributeView PointLengthFromSeedAttribute,
-	PV::FPointPositionAttributeConstView PointPositionAttribute,
-	PV::FBranchPointsAttributeConstView BranchPointsAttribute,
-	PV::FBranchParentNumberAttributeConstView BranchParentNumberAttribute,
-	PV::FBranchNumberAttributeConstView BranchNumberAttribute,
-	PV::FBranchChildrenAttributeConstView BranchChildrenAttribute
-)
+void PV::AttributesHelper::ComputeLengthFromRoot(const FComputeLengthFromRootAttributes& InAttributes)
 {
-	using namespace PV::AttributesHelper;
+	using namespace PV::PlantTraversalHelper;
+
+	if (!InAttributes.IsValid())
+	{
+		return;
+	}
 
-	RecursiveWalkBranches(BranchParentNumberAttribute, BranchChildrenAttribute, BranchNumberAttribute, [&](int32 BranchIndex)->EForEachResult
+	RecursiveWalkBranches(InAttributes.BranchParentNumber, InAttributes.BranchChildren, InAttributes.BranchNumber, [&](int32 BranchIndex)->EForEachResult
 	{
-		const TArray<int32>& CurrentBranchPoints = BranchPointsAttribute[BranchIndex];
+		const TArray<int32>& CurrentBranchPoints = InAttributes.BranchPoints[BranchIndex];
 		check(CurrentBranchPoints.Num() > 0);
 
 		const int32 FirstPointIndex = CurrentBranchPoints[0];
-		const int32 ParentBranchIndex = GetBranchParentIndex(BranchParentNumberAttribute, BranchNumberAttribute, BranchIndex);
+		const int32 ParentBranchIndex = GetBranchParentIndex(InAttributes.BranchParentNumber, InAttributes.BranchNumber, BranchIndex);
 
 		if (ParentBranchIndex == INDEX_NONE)
 		{
-			PointLengthFromRootAttribute[FirstPointIndex] = 0.f; // No parent, length to root is 0
-			PointLengthFromSeedAttribute[FirstPointIndex] = 0.f;
+			InAttributes.PointLengthFromRoot[FirstPointIndex] = 0.f; // No parent, length to root is 0
+			InAttributes.PointLengthFromSeed[FirstPointIndex] = 0.f;
 		}
 		else
 		{
-			const TArray<int32>& ParentBranchPoints = BranchPointsAttribute[ParentBranchIndex];
+			const TArray<int32>& ParentBranchPoints = InAttributes.BranchPoints[ParentBranchIndex];
 			check(ParentBranchPoints.Num() > 0);
 
-			const FVector3f& FirstPoint = PointPositionAttribute[FirstPointIndex];
+			const FVector3f& FirstPoint = InAttributes.PointPosition[FirstPointIndex];
 
 			int32 ClosestParentPointIndex = ParentBranchPoints[0];
-			float ClosestParentPointDistSquared = (PointPositionAttribute[ClosestParentPointIndex] - FirstPoint).SizeSquared();
+			float ClosestParentPointDistSquared = (InAttributes.PointPosition[ClosestParentPointIndex] - FirstPoint).SizeSquared();
 
 			for (int32 i = 1; i < ParentBranchPoints.Num(); ++i)
 			{
 				const int32 PointIndex = ParentBranchPoints[i];
-				const FVector3f ParentPoint = PointPositionAttribute[PointIndex];
+				const FVector3f ParentPoint = InAttributes.PointPosition[PointIndex];
 
 				const float DistSquared = (ParentPoint - FirstPoint).SizeSquared();
 				if (DistSquared > ClosestParentPointDistSquared)
@@ -105,13 +93,13 @@ void PV::AttributesHelper::ComputeLengthFromRoot(
 					if (i < ParentBranchPoints.Num() - 1)
 					{
 						const int32 NextPointIndex = ParentBranchPoints[i + 1];
-						return (PointPositionAttribute[NextPointIndex] - PointPositionAttribute[PointIndex]).GetSafeNormal();
+						return (InAttributes.PointPosition[NextPointIndex] - InAttributes.PointPosition[PointIndex]).GetSafeNormal();
 					}
 
 					if (i > 0)
 					{
 						const int32 PrevPointIndex = ParentBranchPoints[i - 1];
-						return (PointPositionAttribute[PointIndex] - PointPositionAttribute[PrevPointIndex]).GetSafeNormal();
+						return (InAttributes.PointPosition[PointIndex] - InAttributes.PointPosition[PrevPointIndex]).GetSafeNormal();
 					}
 
 					return FVector3f::UpVector;
@@ -125,8 +113,8 @@ void PV::AttributesHelper::ComputeLengthFromRoot(
 				}
 			}
 
-			PointLengthFromRootAttribute[FirstPointIndex] = PointLengthFromRootAttribute[ClosestParentPointIndex] + (ClosestParentPointDistSquared > 0 ? FMath::Sqrt(ClosestParentPointDistSquared) : 0) * 0.01 /*cm->m*/;
-			PointLengthFromSeedAttribute[FirstPointIndex] = PointLengthFromRootAttribute[FirstPointIndex];
+			InAttributes.PointLengthFromRoot[FirstPointIndex] = InAttributes.PointLengthFromRoot[ClosestParentPointIndex] + (ClosestParentPointDistSquared > 0 ? FMath::Sqrt(ClosestParentPointDistSquared) : 0) * 0.01 /*cm->m*/;
+			InAttributes.PointLengthFromSeed[FirstPointIndex] = InAttributes.PointLengthFromRoot[FirstPointIndex];
 		}
 
 		for (int32 i = 1; i < CurrentBranchPoints.Num(); ++i)
@@ -134,198 +122,218 @@ void PV::AttributesHelper::ComputeLengthFromRoot(
 			const int32 PointIndex = CurrentBranchPoints[i];
 			const int32 PrevPointIndex = CurrentBranchPoints[i - 1];
 
-			const FVector3f& Position = PointPositionAttribute[PointIndex];
-			const FVector3f& PrevPosition = PointPositionAttribute[PrevPointIndex];
+			const FVector3f& Position = InAttributes.PointPosition[PointIndex];
+			const FVector3f& PrevPosition = InAttributes.PointPosition[PrevPointIndex];
 
-			PointLengthFromRootAttribute[PointIndex] = PointLengthFromRootAttribute[PrevPointIndex] + (Position - PrevPosition).Size() * 0.01 /*cm->m*/;
-			PointLengthFromSeedAttribute[PointIndex] = PointLengthFromRootAttribute[PointIndex];
+			InAttributes.PointLengthFromRoot[PointIndex] = InAttributes.PointLengthFromRoot[PrevPointIndex] + (Position - PrevPosition).Size() * 0.01 /*cm->m*/;
+			InAttributes.PointLengthFromSeed[PointIndex] = InAttributes.PointLengthFromRoot[PointIndex];
 		}
 
 		return EForEachResult::Continue;
 	});
 }
 
-void PV::AttributesHelper::ComputeBudDevelopment(
-	PV::FBudDevelopmentAttributeView BudDevelopmentAttribute,
-	PV::FPointPlantGradientAttributeConstView PointPlantGradientAttribute,
-	PV::FBranchPointsAttributeConstView BranchPointsAttribute,
-	PV::FBranchParentNumberAttributeConstView BranchParentNumberAttribute,
-	PV::FBranchParentsAttributeConstView BranchParentsAttribute,
-	PV::FBranchNumberAttributeConstView BranchNumberAttribute
-)
+void PV::AttributesHelper::ComputeBudDevelopment(const FComputeBudDevelopmentAttributes& InAttributes)
 {
-	using namespace PV::AttributesHelper;
+	using namespace PV::PlantTraversalHelper;
+
+	if (!InAttributes.IsValid())
+	{
+		return;
+	}
 
 	int32 MaxBudAge = 0;
-	for (int32 BranchIndex = 0; BranchIndex < BranchPointsAttribute.Num(); ++BranchIndex)
+	for (int32 BranchIndex = 0; BranchIndex < InAttributes.BranchPoints.Num(); ++BranchIndex)
 	{
-		MaxBudAge = FMath::Max(MaxBudAge, BranchPointsAttribute[BranchIndex].Num());
+		MaxBudAge = FMath::Max(MaxBudAge, InAttributes.BranchPoints[BranchIndex].Num());
 	}
 
-	ForEachUniquePointOnBranches_FusedPoints(
-		BranchPointsAttribute,
-		BranchParentNumberAttribute,
-		BranchNumberAttribute,
-		[&](int32 BranchIndex, int32 PointIndex)
+	ForEachUniquePointOnBranches(
+		InAttributes.BranchPoints,
+		InAttributes.BranchParentNumber,
+		InAttributes.BranchNumber,
+		[&](int32 BranchIndex, int32 BranchPointIndex, int32 PointIndex)
 		{
-			const float Gradient = PointPlantGradientAttribute[PointIndex];
+			const float Gradient = InAttributes.PointPlantGradient[PointIndex];
 			const int32 BudAge = FMath::RoundToInt(FMath::GetMappedRangeValueClamped(FVector2f(0, 1), FVector2f(1, float(MaxBudAge)), Gradient));
 
-			BudDevelopmentAttribute[PointIndex].BudAge = BudAge;
-			BudDevelopmentAttribute[PointIndex].RelativeBudAge = BudAge;
-			BudDevelopmentAttribute[PointIndex].Generation = BranchParentsAttribute[BranchIndex].Num() + 1;
+			InAttributes.BudDevelopment[PointIndex].BudAge = BudAge;
+			InAttributes.BudDevelopment[PointIndex].RelativeBudAge = BudAge;
+			InAttributes.BudDevelopment[PointIndex].Generation = InAttributes.BranchParents[BranchIndex].Num() + 1;
 
 			return EForEachResult::Continue;
 		}
 	);
 
-	ForEachUniquePointOnBranches_FusedPoints(
-		BranchPointsAttribute,
-		BranchParentNumberAttribute,
-		BranchNumberAttribute,
-		[&](int32 BranchIndex, int32 PointIndex)
+	ForEachUniquePointOnBranches(
+		InAttributes.BranchPoints,
+		InAttributes.BranchParentNumber,
+		InAttributes.BranchNumber,
+		[&](int32 BranchIndex, int32 BranchPointIndex, int32 PointIndex)
 		{
-			const bool bIsRootBranch = GetBranchParentIndex(BranchParentNumberAttribute, BranchNumberAttribute, BranchIndex) == INDEX_NONE;
-			const int32 RootPointIndex = bIsRootBranch ? BranchPointsAttribute[BranchIndex][0] : BranchPointsAttribute[BranchIndex][1];
-			const int32 BranchAge = BudDevelopmentAttribute[RootPointIndex].BudAge;
-			BudDevelopmentAttribute[PointIndex].BranchAge = BranchAge;
+			const bool bIsRootBranch = GetBranchParentIndex(InAttributes.BranchParentNumber, InAttributes.BranchNumber, BranchIndex) == INDEX_NONE;
+			const int32 RootPointIndex = bIsRootBranch ? InAttributes.BranchPoints[BranchIndex][0] : InAttributes.BranchPoints[BranchIndex][1];
+			const int32 BranchAge = InAttributes.BudDevelopment[RootPointIndex].BudAge;
+			InAttributes.BudDevelopment[PointIndex].BranchAge = BranchAge;
 
 			return EForEachResult::Continue;
 		}
 	);
 }
 
-void PV::AttributesHelper::ComputeBudDirections(
-	PV::FBudDirectionAttributeView BudDirectionAttribute,
-	PV::FPointPositionAttributeConstView PointPositionAttribute,
-	PV::FBranchChildrenAttributeConstView BranchChildrenAttribute,
-	PV::FBranchParentNumberAttributeConstView BranchParentNumberAttribute,
-	PV::FBranchNumberAttributeConstView BranchNumberAttribute,
-	PV::FBranchPointsAttributeConstView BranchPointsAttribute
-)
+void PV::AttributesHelper::ComputeBudDirections(const FComputeBudDirectionsAttributes& InAttributes)
 {
-	using namespace PV::AttributesHelper;
+	using namespace PV::PlantTraversalHelper;
 
-	RecursiveWalkPlantPoints_FusedPoints<int32>(
-		BranchChildrenAttribute,
-		BranchParentNumberAttribute,
-		BranchNumberAttribute,
-		BranchPointsAttribute,
-		INDEX_NONE,
-		[&](int32 BranchIndex, int32 BranchPointIndex, int32 PointIndex, const TArray<int32>& ChildPointIndices, int32& PrevPointIndex)->EForEachResult
+	ForEachUniquePointOnBranches(
+		InAttributes.BranchPoints,
+		InAttributes.BranchParentNumber,
+		InAttributes.BranchNumber,
+		[&](int32 BranchIndex, int32 BranchPointIndex, int32 PointIndex)
 		{
-			const TArray<int32>& BranchPoints = BranchPointsAttribute[BranchIndex];
+			const TArray<int32>& BranchPoints = InAttributes.BranchPoints[BranchIndex];
 
-			const FVector3f Apical = [&]()->FVector3f
+			const FVector3f Apical = Invoke([&]()
 			{
-				if (BranchPointIndex < BranchPoints.Num() - 1)
+				if (BranchPointIndex == 0)
 				{
 					const int32 NextPointIndex = BranchPoints[BranchPointIndex + 1];
-					return (PointPositionAttribute[NextPointIndex] - PointPositionAttribute[PointIndex]).GetSafeNormal();
+					return (InAttributes.PointPosition[NextPointIndex] - InAttributes.PointPosition[PointIndex]).GetSafeNormal();
 				}
-				else if (BranchPointIndex > 0)
+				else if (ensure(BranchPointIndex < BranchPoints.Num()))
 				{
 					const int32 PrevPointIndex = BranchPoints[BranchPointIndex - 1];
-					return (PointPositionAttribute[PointIndex] - PointPositionAttribute[PrevPointIndex]).GetSafeNormal();
+					return (InAttributes.PointPosition[PointIndex] - InAttributes.PointPosition[PrevPointIndex]).GetSafeNormal();
 				}
 
-				ensure(false); // We should never be able to get here unless we have a branch with only 1 point
-				return FVector3f::UpVector;
-			}();
+				return FVector3f::ZeroVector;
+			});
+
+			InAttributes.BudDirection[PointIndex].Apical = Apical;
+			InAttributes.BudDirection[PointIndex].LightOptimal = Apical;
+			InAttributes.BudDirection[PointIndex].LightSubOptimal = Apical;
+			InAttributes.BudDirection[PointIndex].GuideCurve = Apical;
+
+			return EForEachResult::Continue;
+		}
+	);
+
+	struct FWalkPlantPointsState
+	{
+		int32 AxillaryCount = 0;
+		int32 PrevPointIndex = INDEX_NONE;
+	};
+	RecursiveWalkPlantPoints<FWalkPlantPointsState>(
+		InAttributes.BranchChildren,
+		InAttributes.BranchParentNumber,
+		InAttributes.BranchNumber,
+		InAttributes.BranchPoints,
+		InAttributes.PointLengthFromRoot,
+		FWalkPlantPointsState(),
+		[&](const FRecursiveWalkBranchPointsParams& Params, FWalkPlantPointsState& State)->EForEachResult
+		{
+			const TArray<int32>& BranchPoints = InAttributes.BranchPoints[Params.BranchIndex];
+
+			const FVector3f Apical = InAttributes.BudDirection[Params.PointIndex].Apical;
+			const FVector3f UpVector = FRotationMatrix44f::MakeFromZX(Apical, FVector3f::ForwardVector).ToQuat().GetRightVector();
+			const FVector3f Axillary = Invoke([&]()
+			{
+				const bool bIsRootPoint = State.PrevPointIndex == INDEX_NONE;
+				if (bIsRootPoint)
+				{
+					State.AxillaryCount = 0;
+					return FVector3f::ForwardVector;
+				}
 
-			const FVector3f UpVector = FRotationMatrix44f::MakeFromZX(Apical, FVector3f::ForwardVector).ToQuat().GetRightVector() * -1;
-			const FVector3f PrevAxillary = PrevPointIndex == INDEX_NONE ? FVector3f::ForwardVector : BudDirectionAttribute[PrevPointIndex].Axillary;
-			const FVector3f Axillary = FQuat4f(Apical, FMath::DegreesToRadians(135.f)).RotateVector(PrevAxillary);
+				if (Params.ChildPointIndices.Num() > 0)
+				{
+					State.AxillaryCount = 0;
 
-			BudDirectionAttribute[PointIndex].Apical = Apical;
-			BudDirectionAttribute[PointIndex].Axillary = Axillary;
-			BudDirectionAttribute[PointIndex].LightOptimal = Apical;
-			BudDirectionAttribute[PointIndex].LightSubOptimal = Apical;
-			BudDirectionAttribute[PointIndex].GuideCurve = Apical;
-			BudDirectionAttribute[PointIndex].UpVector = UpVector;
+					const FVector3f& ChildApical = InAttributes.BudDirection[Params.ChildPointIndices[0]].Apical;
+					return FRotationMatrix44f::MakeFromZX(Apical, ChildApical).ToQuat().GetForwardVector();
+				}
+
+				State.AxillaryCount += 1;
+
+				return FQuat4f(Apical, FMath::DegreesToRadians(137.5f * State.AxillaryCount)).RotateVector(UpVector);
+			});
+
+			InAttributes.BudDirection[Params.PointIndex].Axillary = Axillary;
+			InAttributes.BudDirection[Params.PointIndex].UpVector = UpVector;
 			
-			PrevPointIndex = PointIndex;
+			State.PrevPointIndex = Params.PointIndex;
 
 			return EForEachResult::Continue;
 		}
 	);
 }
 
-void PV::AttributesHelper::ComputeBudHormoneLevels(
-	PV::FBudHormoneLevelsAttributeView BudHormoneLevelsAttribute,
-	PV::FPointPlantGradientAttributeConstView PointPlantGradientAttribute
-)
+void PV::AttributesHelper::ComputeBudHormoneLevels(const FComputeBudHormoneLevelsAttributes& InAttributes)
 {
-	using namespace PV::AttributesHelper;
+	using namespace PV::PlantTraversalHelper;
 		
-	for (int32 PointIndex = 0; PointIndex < BudHormoneLevelsAttribute.Num(); PointIndex++)
+	for (int32 PointIndex = 0; PointIndex < InAttributes.BudHormoneLevels.Num(); PointIndex++)
 	{
-		const float Gradient = PointPlantGradientAttribute[PointIndex];
+		const float Gradient = InAttributes.PointPlantGradient[PointIndex];
 		const float InverseGradient = 1.f - Gradient;
 
-		BudHormoneLevelsAttribute[PointIndex].Apical = FMath::GetMappedRangeValueClamped(FVector2f(0.75, 1.0), FVector2f(0, 1), InverseGradient);
-		BudHormoneLevelsAttribute[PointIndex].Axillary = 0;
-		BudHormoneLevelsAttribute[PointIndex].AxillaryInhibition = 0;
-		BudHormoneLevelsAttribute[PointIndex].Radical = 0;
-		BudHormoneLevelsAttribute[PointIndex].Ethylene = Gradient;
-		BudHormoneLevelsAttribute[PointIndex].Cykininin = 0;
+		InAttributes.BudHormoneLevels[PointIndex].Apical = FMath::GetMappedRangeValueClamped(FVector2f(0.75, 1.0), FVector2f(0, 1), InverseGradient);
+		InAttributes.BudHormoneLevels[PointIndex].Axillary = 0;
+		InAttributes.BudHormoneLevels[PointIndex].AxillaryInhibition = 0;
+		InAttributes.BudHormoneLevels[PointIndex].Radical = 0;
+		InAttributes.BudHormoneLevels[PointIndex].Ethylene = Gradient;
+		InAttributes.BudHormoneLevels[PointIndex].Cykininin = 0;
 	}
 }
 
-void PV::AttributesHelper::ComputeGradients(
-	PV::FPointGroundGradientAttributeView PointGroundGradientAttribute,
-	PV::FPointScaleGradientAttributeView PointScaleGradientAttribute,
-	PV::FPointPlantGradientAttributeView PointPlantGradientAttribute,
-	PV::FPointBranchGradientAttributeView PointBranchGradientAttribute,
-	PV::FPointLengthFromRootAttributeConstView PointLengthFromRootAttribute,
-	PV::FPointScaleAttributeConstView PointScaleAttribute,
-	PV::FBranchPointsAttributeConstView BranchPointsAttribute,
-	PV::FBranchParentNumberAttributeConstView BranchParentNumberAttribute,
-	PV::FBranchNumberAttributeConstView BranchNumberAttribute,
-	PV::FBranchChildrenAttributeConstView BranchChildrenAttribute
-)
+void PV::AttributesHelper::ComputeGradients(const FComputeGradientsAttributes& InAttributes)
 {
-	using namespace PV::AttributesHelper;
+	using namespace PV::PlantTraversalHelper;
 
-	const float* MaxLengthFromRoot = Algo::MaxElement(PointLengthFromRootAttribute);
-	const float* MaxPointScale = Algo::MaxElement(PointScaleAttribute);
+	if (!InAttributes.IsValid())
+	{
+		return;
+	}
+
+	const float* MaxLengthFromRoot = Algo::MaxElement(InAttributes.PointLengthFromRoot);
+	const float* MaxPointScale = Algo::MaxElement(InAttributes.PointScale);
 	if (MaxLengthFromRoot == nullptr || MaxPointScale == nullptr || *MaxLengthFromRoot == 0 || *MaxPointScale == 0)
 	{
 		return;
 	}
 
-	for (int32 PointIndex = 0; PointIndex < PointGroundGradientAttribute.Num(); PointIndex++)
+	for (int32 PointIndex = 0; PointIndex < InAttributes.PointGroundGradient.Num(); PointIndex++)
 	{
-		const float LengthFromRoot = PointLengthFromRootAttribute[PointIndex];
+		const float LengthFromRoot = InAttributes.PointLengthFromRoot[PointIndex];
 		const float LengthFromRootAlpha = LengthFromRoot / *MaxLengthFromRoot;
-		PointGroundGradientAttribute[PointIndex] = 1.0 - LengthFromRootAlpha;
+		InAttributes.PointGroundGradient[PointIndex] = 1.0 - LengthFromRootAlpha;
 
-		const float PointScale = PointScaleAttribute[PointIndex];
+		const float PointScale = InAttributes.PointScale[PointIndex];
 		const float PointScaleAlpha = PointScale / *MaxPointScale;
-		PointScaleGradientAttribute[PointIndex] = PointScaleAlpha;
+		InAttributes.PointScaleGradient[PointIndex] = PointScaleAlpha;
 	}
 
 	// PlantGradient - A gradient from 1-0 where the tip of each branch is 0, and the base of each branch is whatever the gradient of the parent branch is at that point.
 	// BranchGradient - A gradient from 1-0 where 1 = base of the branch and 0 = tip of the branch
 	RecursiveWalkBranches(
-		BranchParentNumberAttribute,
-		BranchChildrenAttribute,
-		BranchNumberAttribute,
+		InAttributes.BranchParentNumber,
+		InAttributes.BranchChildren,
+		InAttributes.BranchNumber,
 		[&](int32 BranchIndex)->EForEachResult
 		{
-			const TArray<int32>& BranchPoints = BranchPointsAttribute[BranchIndex];
+			const TArray<int32>& BranchPoints = InAttributes.BranchPoints[BranchIndex];
 			check(BranchPoints.Num() > 1);
 
-			const float BasePointPlantGradient = GetBranchParentIndex(BranchParentNumberAttribute, BranchNumberAttribute, BranchIndex) == INDEX_NONE
+			const float BasePointPlantGradient = GetBranchParentIndex(InAttributes.BranchParentNumber, InAttributes.BranchNumber, BranchIndex) == INDEX_NONE
 				? 1.0
-				: PointPlantGradientAttribute[BranchPoints[0]];
+				: InAttributes.PointPlantGradient[BranchPoints[0]];
 
 			for (int32 i = 0; i < BranchPoints.Num(); ++i)
 			{
 				const float BranchAlpha = i / float(BranchPoints.Num() - 1);
 				const float PointPlantGradient = FMath::Lerp(BasePointPlantGradient, 0, BranchAlpha);
-				PointPlantGradientAttribute[BranchPoints[i]] = PointPlantGradient;
-				PointBranchGradientAttribute[BranchPoints[i]] = 1.0 - BranchAlpha;
+				InAttributes.PointPlantGradient[BranchPoints[i]] = PointPlantGradient;
+				InAttributes.PointBranchGradient[BranchPoints[i]] = 1.0 - BranchAlpha;
 			}
 
 			return EForEachResult::Continue;
@@ -333,23 +341,22 @@ void PV::AttributesHelper::ComputeGradients(
 	);
 }
 
-void PV::AttributesHelper::ComputeBudStatus(
-	PV::FBudStatusAttributeView BudStatusAttribute,
-	PV::FBranchPointsAttributeConstView BranchPointsAttribute,
-	PV::FBranchParentNumberAttributeConstView BranchParentNumberAttribute,
-	PV::FBranchNumberAttributeConstView BranchNumberAttribute,
-	PV::FBranchChildrenAttributeConstView BranchChildrenAttribute
-)
+void PV::AttributesHelper::ComputeBudStatus(const FComputeBudStatusAttributes& InAttributes)
 {
-	using namespace PV::AttributesHelper;
+	using namespace PV::PlantTraversalHelper;
 
-	ForEachUniquePointOnBranches_FusedPoints(
-		BranchPointsAttribute,
-		BranchParentNumberAttribute,
-		BranchNumberAttribute,
-		[&](int32 BranchIndex, int32 PointIndex)->EForEachResult
+	if (!InAttributes.IsValid())
+	{
+		return;
+	}
+
+	ForEachUniquePointOnBranches(
+		InAttributes.BranchPoints,
+		InAttributes.BranchParentNumber,
+		InAttributes.BranchNumber,
+		[&](int32 BranchIndex, int32 BranchPointIndex, int32 PointIndex)->EForEachResult
 		{
-			PV::FBudStatusView BudStatusView = BudStatusAttribute[PointIndex];
+			PV::FBudStatusView BudStatusView = InAttributes.BudStatus[PointIndex];
 			BudStatusView.ApicalMeristem = 0;
 			BudStatusView.Codominant = 0;
 			BudStatusView.Axillary = 0;
@@ -361,7 +368,7 @@ void PV::AttributesHelper::ComputeBudStatus(
 			BudStatusView.BrokenTip = 0;
 			BudStatusView.Broken = 0;
 
-			if (BranchPointsAttribute[BranchIndex].Last() == PointIndex) // IsTip
+			if (InAttributes.BranchPoints[BranchIndex].Last() == PointIndex) // IsTip
 			{
 				BudStatusView.ApicalMeristem = 1;
 				BudStatusView.Dormant = 1;
@@ -369,19 +376,19 @@ void PV::AttributesHelper::ComputeBudStatus(
 				BudStatusView.NumTriggered = 0;
 				BudStatusView.Inactive = 0;
 			}
-			else if (BranchPointsAttribute[BranchIndex][0] == PointIndex) // Is branch root point.
-				// Due to ForEachUniquePointOnBranches_FusedPoints only iterating shared points on the parent branch, we don't have to check whether this is the root branch
-				// (as the root branch is the only branch which will have the root point iterated on by ForEachUniquePointOnBranches_FusedPoints).
+			else if (InAttributes.BranchPoints[BranchIndex][0] == PointIndex
+				&& GetBranchParentIndex(InAttributes.BranchParentNumber, InAttributes.BranchNumber, BranchIndex) == INDEX_NONE) // Is plant root point
 			{
 				BudStatusView.Seed = 1;
 			}
 			else
 			{
-				const int32 NumChildren = GetBranchPointChildBranchIndices_FusedPoints(
-					BranchChildrenAttribute,
-					BranchParentNumberAttribute,
-					BranchNumberAttribute,
-					BranchPointsAttribute,
+				const int32 NumChildren = GetBranchPointChildBranchIndices(
+					InAttributes.BranchChildren,
+					InAttributes.BranchParentNumber,
+					InAttributes.BranchNumber,
+					InAttributes.BranchPoints,
+					InAttributes.PointLengthFromRoot,
 					BranchIndex, 
 					PointIndex
 				).Num();
@@ -398,34 +405,24 @@ void PV::AttributesHelper::ComputeBudStatus(
 	);
 }
 
-void PV::AttributesHelper::ComputeBudLateralMeristem(
-	PV::FBudLateralMeristemAttributeView BudLateralMeristemAttribute,
-	PV::FPointPositionAttributeConstView PointPositionAttribute,
-	PV::FPointScaleAttributeConstView PointScaleAttribute,
-	PV::FPointLengthFromRootAttributeConstView PointLengthFromRootAttribute,
-	PV::FBudDirectionAttributeConstView BudDirectionAttribute,
-	PV::FBranchPointsAttributeConstView BranchPointsAttribute,
-	PV::FBranchParentNumberAttributeConstView BranchParentNumberAttribute,
-	PV::FBranchChildrenAttributeConstView BranchChildrenAttribute,
-	PV::FBranchNumberAttributeConstView BranchNumberAttribute
-)
+void PV::AttributesHelper::ComputeBudLateralMeristem(const FComputeBudLateralMeristemAttributes& InAttributes)
 {
-	using namespace PV::AttributesHelper;
+	using namespace PV::PlantTraversalHelper;
 
 	RecursiveWalkBranches_Reversed(
-		BranchParentNumberAttribute,
-		BranchChildrenAttribute,
-		BranchNumberAttribute,
+		InAttributes.BranchParentNumber,
+		InAttributes.BranchChildren,
+		InAttributes.BranchNumber,
 		[&](int32 BranchIndex)->EForEachResult
 		{
-			const TArray<int32>& BranchPointIndices = BranchPointsAttribute[BranchIndex];
+			const TArray<int32>& BranchPointIndices = InAttributes.BranchPoints[BranchIndex];
 			check(BranchPointIndices.Num() > 1);
 
-			const FVector3f& FirstBranchPointPos = PointPositionAttribute[BranchPointIndices[0]];
-			const FVector3f& SecondBranchPointPos = PointPositionAttribute[BranchPointIndices[1]];
+			const FVector3f& FirstBranchPointPos = InAttributes.PointPosition[BranchPointIndices[0]];
+			const FVector3f& SecondBranchPointPos = InAttributes.PointPosition[BranchPointIndices[1]];
 
 			const FVector3f BranchDir = (SecondBranchPointPos - FirstBranchPointPos).GetSafeNormal();
-			const FVector3f FirstPointNormal = BudDirectionAttribute[BranchPointIndices[0]].Apical;;
+			const FVector3f FirstPointNormal = InAttributes.BudDirection[BranchPointIndices[0]].Apical;;
 			
 			const float ParentDot = FVector3f::DotProduct(BranchDir, FirstPointNormal);
 
@@ -433,18 +430,18 @@ void PV::AttributesHelper::ComputeBudLateralMeristem(
 			for (int32 i = BranchPointIndices.Num() - 1; i>= 0; --i)
 			{
 				const int32 PointIndex = BranchPointIndices[i];
-				const float ExisitingDavinciValue = BudLateralMeristemAttribute[PointIndex].Davinci;
-				const float PointScale = PointScaleAttribute[PointIndex];
+				const float ExisitingDavinciValue = InAttributes.BudLateralMeristem[PointIndex].Davinci;
+				const float PointScale = InAttributes.PointScale[PointIndex];
 				const float AccumulatedDavinciValue = PreviousPointDavinciValue + ExisitingDavinciValue + PointScale;
 				PreviousPointDavinciValue = AccumulatedDavinciValue;
 
-				BudLateralMeristemAttribute[PointIndex].LateralMeristem = 0.0;
-				BudLateralMeristemAttribute[PointIndex].Multiplier = 1;
-				BudLateralMeristemAttribute[PointIndex].Inactive = 0.0;
-				BudLateralMeristemAttribute[PointIndex].Davinci = AccumulatedDavinciValue;
-				BudLateralMeristemAttribute[PointIndex].ParentDot = ParentDot;
-				BudLateralMeristemAttribute[PointIndex].RootDistance = PointLengthFromRootAttribute[PointIndex];
-				BudLateralMeristemAttribute[PointIndex].Degredation = 0.0;
+				InAttributes.BudLateralMeristem[PointIndex].LateralMeristem = 0.0;
+				InAttributes.BudLateralMeristem[PointIndex].Multiplier = 1;
+				InAttributes.BudLateralMeristem[PointIndex].Inactive = 0.0;
+				InAttributes.BudLateralMeristem[PointIndex].Davinci = AccumulatedDavinciValue;
+				InAttributes.BudLateralMeristem[PointIndex].ParentDot = ParentDot;
+				InAttributes.BudLateralMeristem[PointIndex].RootDistance = InAttributes.PointLengthFromRoot[PointIndex];
+				InAttributes.BudLateralMeristem[PointIndex].Degredation = 0.0;
 			}
 
 			return EForEachResult::Continue;
@@ -463,6 +460,24 @@ void PV::AttributesHelper::ComputeBudLightDetected(PV::FBudLightDetectedAttribut
 	}
 }
 
+void PV::AttributesHelper::ComputeBranchHierarchyNumbers(const FComputeBranchHierarchyNumberAttributes& InAttributes)
+{
+	using namespace PV::PlantTraversalHelper;
+
+	RecursiveWalkBranches<int32>(
+		InAttributes.BranchParentNumber,
+		InAttributes.BranchChildren,
+		InAttributes.BranchNumber,
+		1,
+		[&](int32 BranchIndex, int32& BranchHierarchyNumber)->EForEachResult
+		{
+			InAttributes.BranchHierarchyNumber[BranchIndex] = BranchHierarchyNumber;
+			BranchHierarchyNumber += 1;
+			return EForEachResult::Continue;
+		}
+	);
+}
+
 void PV::AttributesHelper::InterpolatePointsScales(
 	PV::FPointScaleAttributeView PointScaleAttribute,
 	PV::FPointLengthFromRootAttributeConstView PointLengthFromRootAttribute,
@@ -472,7 +487,7 @@ void PV::AttributesHelper::InterpolatePointsScales(
 	PV::FBranchNumberAttributeConstView BranchNumberAttribute
 )
 {
-	using namespace PV::AttributesHelper;
+	using namespace PV::PlantTraversalHelper;
 
 	RecursiveWalkBranches(
 		BranchParentNumberAttribute,
@@ -538,78 +553,164 @@ void PV::AttributesHelper::InterpolatePointsScales(
 	);
 }
 
-bool PV::AttributesHelper::SplitFusedPoints(FManagedArrayCollection& Collection)
+bool PV::AttributesHelper::SplitFusedPoints(FSplitFusedPointsAttributes InAttributes)
 {
-	if (!Collection.HasGroup(PV::GroupNames::PointGroup) || !Collection.HasGroup(PV::GroupNames::BranchGroup))
+	if (!InAttributes.IsValid())
+	{
+		return false;
+	}
+
+	if (!PV::PlantTraversalHelper::HasFusedPoints(InAttributes.BranchPoints, InAttributes.BranchParentNumber, InAttributes.BranchNumber))
 	{
 		return false;
 	}
 
-	auto BranchPointsAttribute = FBranchPointsAttribute::GetAttribute(Collection);
-	auto BranchParentNumberAttribute = FBranchParentNumberAttribute::GetAttribute(Collection);
-
-	auto PointPositionAttribute = FPointPositionAttribute::GetAttribute(Collection);
-	auto LengthFromRootAttribute = FPointLengthFromRootAttribute::GetAttribute(Collection);
-	auto LengthFromSeedAttribute = FPointLengthFromSeedAttribute::GetAttribute(Collection);
-	auto PointScaleGradientAttribute = FPointScaleGradientAttribute::GetAttribute(Collection);
-	auto PointHullGradientAttribute = FPointHullGradientAttribute::GetAttribute(Collection);
-	auto PointMainTrunkGradientAttribute = FPointMainTrunkGradientAttribute::GetAttribute(Collection);
-	auto GroundGradientAttribute = FPointGroundGradientAttribute::GetAttribute(Collection);
-	auto PointScaleAttribute = FPointScaleAttribute::GetAttribute(Collection);
-	auto BudDirectionAttribute = FBudDirectionAttribute::GetAttribute(Collection);
-	auto BudStatusAttribute = FBudStatusAttribute::GetAttribute(Collection);
-	auto BudLateralMeristemAttribute = FBudLateralMeristemAttribute::GetAttribute(Collection);
-	auto PointBranchGradientAttribute = FPointBranchGradientAttribute::GetAttribute(Collection);
-	auto BudHormoneLevelsAttribute = FBudHormoneLevelsAttribute::GetAttribute(Collection);
-	auto PlantGradientAttribute = FPointPlantGradientAttribute::GetAttribute(Collection);
-	auto BudNumberAttribute = FPointBudNumberAttribute::GetAttribute(Collection);
-	auto NjordPixelIndexAttribute = FPointNjordPixelIndexAttribute::GetAttribute(Collection);
-	auto BudLightDetectedAttribute = FBudLightDetectedAttribute::GetAttribute(Collection);
-	auto BudDevelopmentAttribute = FBudDevelopmentAttribute::GetAttribute(Collection);
-
-	const int32* MaxBudNumber = Algo::MaxElement(BudNumberAttribute);
+	const int32* MaxBudNumber = Algo::MaxElement(InAttributes.PointBudNumber.Get());
 	check(MaxBudNumber != nullptr);
 
 	int32 NewBudNumber = *MaxBudNumber + 1;
 
-	for (int32 i = 0; i < BranchPointsAttribute.Num(); ++i)
+	for (int32 i = 0; i < InAttributes.BranchPoints.Num(); ++i)
 	{
 		const int32 BranchIndex = i;
-		if (BranchParentNumberAttribute[BranchIndex] == 0)
+		if (InAttributes.BranchParentNumber[BranchIndex] == 0)
 		{
 			continue;
 		}
 
-		const int32 NumPoints = PointPositionAttribute.Num();
-		Collection.Resize(NumPoints + 1, PV::GroupNames::PointGroup);
+		const int32 NumPoints = InAttributes.PointPosition.Num();
+		InAttributes.PointPosition.AddElements(1);
 
-		TArray<int32>& BranchPoints = BranchPointsAttribute[BranchIndex];
+		TArray<int32>& BranchPoints = InAttributes.BranchPoints[BranchIndex];
 		const int32 PrevPointIndex = BranchPoints[0];
 		const int32 NewPointIndex = NumPoints;
 
 		BranchPoints[0] = NewPointIndex;
 
-		PointPositionAttribute[NewPointIndex] = PointPositionAttribute[PrevPointIndex];
-		LengthFromRootAttribute[NewPointIndex] = LengthFromRootAttribute[PrevPointIndex];
-		LengthFromSeedAttribute[NewPointIndex] = LengthFromSeedAttribute[PrevPointIndex];
-		PointScaleGradientAttribute[NewPointIndex] = PointScaleGradientAttribute[PrevPointIndex];
-		PointHullGradientAttribute[NewPointIndex] = PointHullGradientAttribute[PrevPointIndex];
-		PointMainTrunkGradientAttribute[NewPointIndex] = PointMainTrunkGradientAttribute[PrevPointIndex];
-		GroundGradientAttribute[NewPointIndex] = GroundGradientAttribute[PrevPointIndex];
-		PointScaleAttribute[NewPointIndex] = PointScaleAttribute[PrevPointIndex];
-		BudDirectionAttribute[NewPointIndex].Array = BudDirectionAttribute[PrevPointIndex].Array;
-		BudStatusAttribute[NewPointIndex].Array = BudStatusAttribute[PrevPointIndex].Array;
-		BudLateralMeristemAttribute[NewPointIndex].Array = BudLateralMeristemAttribute[PrevPointIndex].Array;
-		PointBranchGradientAttribute[NewPointIndex] = PointBranchGradientAttribute[PrevPointIndex];
-		BudHormoneLevelsAttribute[NewPointIndex].Array = BudHormoneLevelsAttribute[PrevPointIndex].Array;
-		PlantGradientAttribute[NewPointIndex] = PlantGradientAttribute[PrevPointIndex];
-		BudNumberAttribute[NewPointIndex] = NewBudNumber;
-		NjordPixelIndexAttribute[NewPointIndex] = NjordPixelIndexAttribute[PrevPointIndex];
-		BudLightDetectedAttribute[NewPointIndex].Array = BudLightDetectedAttribute[PrevPointIndex].Array;
-		BudDevelopmentAttribute[NewPointIndex].Array = BudDevelopmentAttribute[PrevPointIndex].Array;
+		InAttributes.PointPosition.ModifyAt(NewPointIndex) = InAttributes.PointPosition[PrevPointIndex];
+		InAttributes.PointLengthFromRoot.ModifyAt(NewPointIndex) = InAttributes.PointLengthFromRoot[PrevPointIndex];
+		InAttributes.PointLengthFromSeed.ModifyAt(NewPointIndex) = InAttributes.PointLengthFromSeed[PrevPointIndex];
+		InAttributes.PointScaleGradient.ModifyAt(NewPointIndex) = InAttributes.PointScaleGradient[PrevPointIndex];
+		InAttributes.PointHullGradient.ModifyAt(NewPointIndex) = InAttributes.PointHullGradient[PrevPointIndex];
+		InAttributes.PointMainTrunkGradient.ModifyAt(NewPointIndex) = InAttributes.PointMainTrunkGradient[PrevPointIndex];
+		InAttributes.PointGroundGradient.ModifyAt(NewPointIndex) = InAttributes.PointGroundGradient[PrevPointIndex];
+		InAttributes.PointScale.ModifyAt(NewPointIndex) = InAttributes.PointScale[PrevPointIndex];
+		InAttributes.BudDirection.ModifyAt(NewPointIndex).Array = InAttributes.BudDirection[PrevPointIndex].Array;
+		InAttributes.BudStatus.ModifyAt(NewPointIndex).Array = InAttributes.BudStatus[PrevPointIndex].Array;
+		InAttributes.BudLateralMeristem.ModifyAt(NewPointIndex).Array = InAttributes.BudLateralMeristem[PrevPointIndex].Array;
+		InAttributes.PointBranchGradient.ModifyAt(NewPointIndex) = InAttributes.PointBranchGradient[PrevPointIndex];
+		InAttributes.BudHormoneLevels.ModifyAt(NewPointIndex).Array = InAttributes.BudHormoneLevels[PrevPointIndex].Array;
+		InAttributes.PointPlantGradient.ModifyAt(NewPointIndex) = InAttributes.PointPlantGradient[PrevPointIndex];
+		InAttributes.PointBudNumber.ModifyAt(NewPointIndex) = NewBudNumber;
+		InAttributes.PointNjordPixelIndex.ModifyAt(NewPointIndex) = InAttributes.PointNjordPixelIndex[PrevPointIndex];
+		InAttributes.BudLightDetected.ModifyAt(NewPointIndex).Array = InAttributes.BudLightDetected[PrevPointIndex].Array;
+		InAttributes.BudDevelopment.ModifyAt(NewPointIndex).Array = InAttributes.BudDevelopment[PrevPointIndex].Array;
 
 		NewBudNumber++;
 	}
 
+	return true;
+}
+
+bool PV::AttributesHelper::RecomputeAllGrowthDataAttributes(const FRecomputeGrowthDataAttributes& InAttributes)
+{
+	if (!InAttributes.IsValid())
+	{
+		return false;
+	}
+
+	ComputeBudNumbers({
+		InAttributes.PointBudNumber,
+		InAttributes.BranchSourceBudNumber,
+		InAttributes.BranchPoints
+	});
+
+	ComputeNjordPixelIndex({
+		InAttributes.PointNjordPixelIndex,
+		InAttributes.BranchPoints,
+		InAttributes.BranchParentNumber,
+		InAttributes.BranchNumber,
+		InAttributes.BranchChildren
+	});
+
+	ComputeLengthFromRoot({
+		InAttributes.PointLengthFromRoot,
+		InAttributes.PointLengthFromSeed,
+		InAttributes.PointPosition,
+		InAttributes.BranchPoints,
+		InAttributes.BranchParentNumber,
+		InAttributes.BranchNumber,
+		InAttributes.BranchChildren
+	});
+
+	ComputeBudDirections({
+		InAttributes.BudDirection,
+		InAttributes.PointPosition,
+		InAttributes.PointLengthFromRoot,
+		InAttributes.BranchChildren,
+		InAttributes.BranchParentNumber,
+		InAttributes.BranchNumber,
+		InAttributes.BranchPoints
+	});
+
+	ComputeBudLightDetected(InAttributes.BudLightDetected);
+
+	InterpolatePointsScales(
+		InAttributes.PointScale,
+		InAttributes.PointLengthFromRoot,
+		InAttributes.BranchPoints,
+		InAttributes.BranchParentNumber,
+		InAttributes.BranchChildren,
+		InAttributes.BranchNumber
+	);
+
+	ComputeBudLateralMeristem({
+		InAttributes.BudLateralMeristem,
+		InAttributes.PointPosition,
+		InAttributes.PointScale,
+		InAttributes.PointLengthFromRoot,
+		InAttributes.BudDirection,
+		InAttributes.BranchPoints,
+		InAttributes.BranchParentNumber,
+		InAttributes.BranchChildren,
+		InAttributes.BranchNumber
+	});
+
+	ComputeGradients({
+		InAttributes.PointGroundGradient,
+		InAttributes.PointScaleGradient,
+		InAttributes.PointPlantGradient,
+		InAttributes.PointBranchGradient,
+		InAttributes.PointLengthFromRoot,
+		InAttributes.PointScale,
+		InAttributes.BranchPoints,
+		InAttributes.BranchParentNumber,
+		InAttributes.BranchNumber,
+		InAttributes.BranchChildren
+	});
+
+	ComputeBudDevelopment({
+		InAttributes.BudDevelopment,
+		InAttributes.PointPlantGradient,
+		InAttributes.BranchPoints,
+		InAttributes.BranchParentNumber,
+		InAttributes.BranchParents,
+		InAttributes.BranchNumber
+	});
+
+	ComputeBudHormoneLevels({
+		InAttributes.BudHormoneLevels,
+		InAttributes.PointPlantGradient
+	});
+
+	ComputeBudStatus({
+		InAttributes.BudStatus,
+		InAttributes.BranchPoints,
+		InAttributes.BranchParentNumber,
+		InAttributes.BranchNumber,
+		InAttributes.BranchChildren,
+		InAttributes.PointLengthFromRoot
+	});
+
 	return true;
 }
\ No newline at end of file
diff --git a/Engine/Plugins/Experimental/ProceduralVegetationEditor/Source/ProceduralVegetation/Private/Helpers/PVAttributesHelper.h b/Engine/Plugins/Experimental/ProceduralVegetationEditor/Source/ProceduralVegetation/Private/Helpers/PVAttributesHelper.h
new file mode 100644
--- /dev/null
+++ b/Engine/Plugins/Experimental/ProceduralVegetationEditor/Source/ProceduralVegetation/Private/Helpers/PVAttributesHelper.h
@@ -0,0 +1,651 @@
+// Copyright Epic Games, Inc. All Rights Reserved.
+#pragma once
+
+#include "Utils/PVAttributes.h"
+
+namespace PV::AttributesHelper
+{
+	struct FComputeBudNumbersAttributes
+	{
+		FPointBudNumberAttributeView PointBudNumber;
+		FBranchSourceBudNumberAttributeView BranchSourceBudNumber;
+		FBranchPointsAttributeConstView BranchPoints;
+
+		FComputeBudNumbersAttributes(
+			FPointBudNumberAttributeView InPointBudNumber,
+			FBranchSourceBudNumberAttributeView InBranchSourceBudNumber,
+			FBranchPointsAttributeConstView InBranchPoints)
+			: PointBudNumber(InPointBudNumber)
+			, BranchSourceBudNumber(InBranchSourceBudNumber)
+			, BranchPoints(InBranchPoints)
+		{}
+
+		FComputeBudNumbersAttributes(FManagedArrayCollection& InCollection)
+			: PointBudNumber(FPointBudNumberAttribute::FindAttribute(InCollection))
+			, BranchSourceBudNumber(FBranchSourceBudNumberAttribute::FindAttribute(InCollection))
+			, BranchPoints(FBranchPointsAttribute::FindAttribute(InCollection))
+		{}
+
+		bool IsValid() const
+		{
+			return PointBudNumber.IsValid()
+				&& BranchSourceBudNumber.IsValid()
+				&& BranchPoints.IsValid();
+		}
+	};
+
+	// Re-computes the point bud numbers and writes them to the BudNumber and BranchSourceBudNumber attributes.
+	PROCEDURALVEGETATION_API extern void ComputeBudNumbers(const FComputeBudNumbersAttributes& InAttributes);
+
+	struct FComputeNjordPixelIndexAttributes
+	{
+		PV::FPointNjordPixelIndexAttributeView PointNjordPixelIndex;
+		PV::FBranchPointsAttributeConstView BranchPoints;
+		PV::FBranchParentNumberAttributeConstView BranchParentNumber;
+		PV::FBranchNumberAttributeConstView BranchNumber;
+		PV::FBranchChildrenAttributeConstView BranchChildren;
+
+		FComputeNjordPixelIndexAttributes(
+			PV::FPointNjordPixelIndexAttributeView InPointNjordPixelIndex,
+			PV::FBranchPointsAttributeConstView InBranchPoints,
+			PV::FBranchParentNumberAttributeConstView InBranchParentNumber,
+			PV::FBranchNumberAttributeConstView InBranchNumber,
+			PV::FBranchChildrenAttributeConstView InBranchChildren)
+			: PointNjordPixelIndex(InPointNjordPixelIndex)
+			, BranchPoints(InBranchPoints)
+			, BranchParentNumber(InBranchParentNumber)
+			, BranchNumber(InBranchNumber)
+			, BranchChildren(InBranchChildren)
+		{
+		}
+
+		FComputeNjordPixelIndexAttributes(FManagedArrayCollection& InCollection)
+			: PointNjordPixelIndex(PV::FPointNjordPixelIndexAttribute::FindAttribute(InCollection))
+			, BranchPoints(PV::FBranchPointsAttribute::FindAttribute(InCollection))
+			, BranchParentNumber(PV::FBranchParentNumberAttribute::FindAttribute(InCollection))
+			, BranchNumber(PV::FBranchNumberAttribute::FindAttribute(InCollection))
+			, BranchChildren(PV::FBranchChildrenAttribute::FindAttribute(InCollection))
+		{
+		}
+
+		bool IsValid() const
+		{
+			return PointNjordPixelIndex.IsValid()
+				&& BranchPoints.IsValid()
+				&& BranchParentNumber.IsValid()
+				&& BranchNumber.IsValid()
+				&& BranchChildren.IsValid();
+		}
+	};
+
+	// Re-computes all njord pixel indices and writes them to the NjordPixelIndexAttribute.
+	PROCEDURALVEGETATION_API extern void ComputeNjordPixelIndex(const FComputeNjordPixelIndexAttributes& InAttributes);
+
+	struct FComputeLengthFromRootAttributes
+	{
+		PV::FPointLengthFromRootAttributeView PointLengthFromRoot;
+		PV::FPointLengthFromSeedAttributeView PointLengthFromSeed;
+		PV::FPointPositionAttributeConstView PointPosition;
+		PV::FBranchPointsAttributeConstView BranchPoints;
+		PV::FBranchParentNumberAttributeConstView BranchParentNumber;
+		PV::FBranchNumberAttributeConstView BranchNumber;
+		PV::FBranchChildrenAttributeConstView BranchChildren;
+
+		FComputeLengthFromRootAttributes(
+			PV::FPointLengthFromRootAttributeView InPointLengthFromRoot,
+			PV::FPointLengthFromSeedAttributeView InPointLengthFromSeed,
+			PV::FPointPositionAttributeConstView InPointPosition,
+			PV::FBranchPointsAttributeConstView InBranchPoints,
+			PV::FBranchParentNumberAttributeConstView InBranchParentNumber,
+			PV::FBranchNumberAttributeConstView InBranchNumber,
+			PV::FBranchChildrenAttributeConstView InBranchChildren)
+			: PointLengthFromRoot(InPointLengthFromRoot)
+			, PointLengthFromSeed(InPointLengthFromSeed)
+			, PointPosition(InPointPosition)
+			, BranchPoints(InBranchPoints)
+			, BranchParentNumber(InBranchParentNumber)
+			, BranchNumber(InBranchNumber)
+			, BranchChildren(InBranchChildren)
+		{
+		}
+
+		FComputeLengthFromRootAttributes(FManagedArrayCollection& Collection)
+			: PointLengthFromRoot(PV::FPointLengthFromRootAttribute::FindAttribute(Collection))
+			, PointLengthFromSeed(PV::FPointLengthFromSeedAttribute::FindAttribute(Collection))
+			, PointPosition(PV::FPointPositionAttribute::FindAttribute(Collection))
+			, BranchPoints(PV::FBranchPointsAttribute::FindAttribute(Collection))
+			, BranchParentNumber(PV::FBranchParentNumberAttribute::FindAttribute(Collection))
+			, BranchNumber(PV::FBranchNumberAttribute::FindAttribute(Collection))
+			, BranchChildren(PV::FBranchChildrenAttribute::FindAttribute(Collection))
+		{
+		}
+
+		bool IsValid() const
+		{
+			return PointLengthFromRoot.IsValid()
+				&& PointLengthFromSeed.IsValid()
+				&& PointPosition.IsValid()
+				&& BranchPoints.IsValid()
+				&& BranchParentNumber.IsValid()
+				&& BranchNumber.IsValid()
+				&& BranchChildren.IsValid();
+		}
+	};
+
+	// Re-computes the length from root and length from seed for all points.
+	PROCEDURALVEGETATION_API extern void ComputeLengthFromRoot(const FComputeLengthFromRootAttributes& InAttributes);
+
+	struct FComputeBudDevelopmentAttributes
+	{
+		PV::FBudDevelopmentAttributeView BudDevelopment;
+		PV::FPointPlantGradientAttributeConstView PointPlantGradient;
+		PV::FBranchPointsAttributeConstView BranchPoints;
+		PV::FBranchParentNumberAttributeConstView BranchParentNumber;
+		PV::FBranchParentsAttributeConstView BranchParents;
+		PV::FBranchNumberAttributeConstView BranchNumber;
+
+		FComputeBudDevelopmentAttributes(
+			PV::FBudDevelopmentAttributeView InBudDevelopment,
+			PV::FPointPlantGradientAttributeConstView InPointPlantGradient,
+			PV::FBranchPointsAttributeConstView InBranchPoints,
+			PV::FBranchParentNumberAttributeConstView InBranchParentNumber,
+			PV::FBranchParentsAttributeConstView InBranchParents,
+			PV::FBranchNumberAttributeConstView InBranchNumber)
+			: BudDevelopment(InBudDevelopment)
+			, PointPlantGradient(InPointPlantGradient)
+			, BranchPoints(InBranchPoints)
+			, BranchParentNumber(InBranchParentNumber)
+			, BranchParents(InBranchParents)
+			, BranchNumber(InBranchNumber)
+		{}
+
+		FComputeBudDevelopmentAttributes(FManagedArrayCollection& InCollection)
+			: BudDevelopment(PV::FBudDevelopmentAttribute::FindAttribute(InCollection))
+			, PointPlantGradient(PV::FPointPlantGradientAttribute::FindAttribute(InCollection))
+			, BranchPoints(PV::FBranchPointsAttribute::FindAttribute(InCollection))
+			, BranchParentNumber(PV::FBranchParentNumberAttribute::FindAttribute(InCollection))
+			, BranchParents(PV::FBranchParentsAttribute::FindAttribute(InCollection))
+			, BranchNumber(PV::FBranchNumberAttribute::FindAttribute(InCollection))
+		{}
+
+		bool IsValid() const
+		{
+			return BudDevelopment.IsValid()
+				&& PointPlantGradient.IsValid()
+				&& BranchPoints.IsValid()
+				&& BranchParentNumber.IsValid()
+				&& BranchParents.IsValid()
+				&& BranchNumber.IsValid();
+		}
+	};
+
+	// Re-computes bud development for all points using the existing branch hierarchy and plant gradient. 
+	PROCEDURALVEGETATION_API extern void ComputeBudDevelopment(const FComputeBudDevelopmentAttributes& InAttributes);
+
+	struct FComputeBudDirectionsAttributes
+	{
+		PV::FBudDirectionAttributeView BudDirection;
+		PV::FPointPositionAttributeConstView PointPosition;
+		PV::FPointLengthFromRootAttributeConstView PointLengthFromRoot;
+		PV::FBranchChildrenAttributeConstView BranchChildren;
+		PV::FBranchParentNumberAttributeConstView BranchParentNumber;
+		PV::FBranchNumberAttributeConstView BranchNumber;
+		PV::FBranchPointsAttributeConstView BranchPoints;
+
+		FComputeBudDirectionsAttributes(
+			PV::FBudDirectionAttributeView InBudDirection,
+			PV::FPointPositionAttributeConstView InPointPosition,
+			PV::FPointLengthFromRootAttributeConstView InPointLengthFromRoot,
+			PV::FBranchChildrenAttributeConstView InBranchChildren,
+			PV::FBranchParentNumberAttributeConstView InBranchParentNumber,
+			PV::FBranchNumberAttributeConstView InBranchNumber,
+			PV::FBranchPointsAttributeConstView InBranchPoints)
+			: BudDirection(InBudDirection)
+			, PointPosition(InPointPosition)
+			, PointLengthFromRoot(InPointLengthFromRoot)
+			, BranchChildren(InBranchChildren)
+			, BranchParentNumber(InBranchParentNumber)
+			, BranchNumber(InBranchNumber)
+			, BranchPoints(InBranchPoints)
+		{}
+
+		FComputeBudDirectionsAttributes(FManagedArrayCollection& InCollection)
+			: BudDirection(PV::FBudDirectionAttribute::FindAttribute(InCollection))
+			, PointPosition(PV::FPointPositionAttribute::FindAttribute(InCollection))
+			, PointLengthFromRoot(PV::FPointLengthFromRootAttribute::FindAttribute(InCollection))
+			, BranchChildren(PV::FBranchChildrenAttribute::FindAttribute(InCollection))
+			, BranchParentNumber(PV::FBranchParentNumberAttribute::FindAttribute(InCollection))
+			, BranchNumber(PV::FBranchNumberAttribute::FindAttribute(InCollection))
+			, BranchPoints(PV::FBranchPointsAttribute::FindAttribute(InCollection))
+		{}
+
+		bool IsValid() const
+		{
+			return BudDirection.IsValid()
+				&& PointPosition.IsValid()
+				&& PointLengthFromRoot.IsValid()
+				&& BranchChildren.IsValid()
+				&& BranchParentNumber.IsValid()
+				&& BranchNumber.IsValid()
+				&& BranchPoints.IsValid();
+		}
+	};
+
+	// Re-computes the bud directions for all points.
+	PROCEDURALVEGETATION_API extern void ComputeBudDirections(const FComputeBudDirectionsAttributes& InAttributes);
+
+	struct FComputeBudHormoneLevelsAttributes
+	{
+		PV::FBudHormoneLevelsAttributeView BudHormoneLevels;
+		PV::FPointPlantGradientAttributeConstView PointPlantGradient;
+
+		bool IsValid() const
+		{
+			return BudHormoneLevels && PointPlantGradient;
+		}
+	};
+
+	// Re-computes the bud hormone levels using the plant gradient for all points.
+	PROCEDURALVEGETATION_API extern void ComputeBudHormoneLevels(const FComputeBudHormoneLevelsAttributes& InAttributes);
+
+	struct FComputeGradientsAttributes
+	{
+		PV::FPointGroundGradientAttributeView PointGroundGradient;
+		PV::FPointScaleGradientAttributeView PointScaleGradient;
+		PV::FPointPlantGradientAttributeView PointPlantGradient;
+		PV::FPointBranchGradientAttributeView PointBranchGradient;
+		PV::FPointLengthFromRootAttributeConstView PointLengthFromRoot;
+		PV::FPointScaleAttributeConstView PointScale;
+		PV::FBranchPointsAttributeConstView BranchPoints;
+		PV::FBranchParentNumberAttributeConstView BranchParentNumber;
+		PV::FBranchNumberAttributeConstView BranchNumber;
+		PV::FBranchChildrenAttributeConstView BranchChildren;
+
+		FComputeGradientsAttributes(
+			PV::FPointGroundGradientAttributeView InPointGroundGradient,
+			PV::FPointScaleGradientAttributeView InPointScaleGradient,
+			PV::FPointPlantGradientAttributeView InPointPlantGradient,
+			PV::FPointBranchGradientAttributeView InPointBranchGradient,
+			PV::FPointLengthFromRootAttributeConstView InPointLengthFromRoot,
+			PV::FPointScaleAttributeConstView InPointScale,
+			PV::FBranchPointsAttributeConstView InBranchPoints,
+			PV::FBranchParentNumberAttributeConstView InBranchParentNumber,
+			PV::FBranchNumberAttributeConstView InBranchNumber,
+			PV::FBranchChildrenAttributeConstView InBranchChildren)
+			: PointGroundGradient(InPointGroundGradient)
+			, PointScaleGradient(InPointScaleGradient)
+			, PointPlantGradient(InPointPlantGradient)
+			, PointBranchGradient(InPointBranchGradient)
+			, PointLengthFromRoot(InPointLengthFromRoot)
+			, PointScale(InPointScale)
+			, BranchPoints(InBranchPoints)
+			, BranchParentNumber(InBranchParentNumber)
+			, BranchNumber(InBranchNumber)
+			, BranchChildren(InBranchChildren)
+		{
+		}
+
+		FComputeGradientsAttributes(FManagedArrayCollection& InCollection)
+			: PointGroundGradient(PV::FPointGroundGradientAttribute::FindAttribute(InCollection))
+			, PointScaleGradient(PV::FPointScaleGradientAttribute::FindAttribute(InCollection))
+			, PointPlantGradient(PV::FPointPlantGradientAttribute::FindAttribute(InCollection))
+			, PointBranchGradient(PV::FPointBranchGradientAttribute::FindAttribute(InCollection))
+			, PointLengthFromRoot(PV::FPointLengthFromRootAttribute::FindAttribute(InCollection))
+			, PointScale(PV::FPointScaleAttribute::FindAttribute(InCollection))
+			, BranchPoints(PV::FBranchPointsAttribute::FindAttribute(InCollection))
+			, BranchParentNumber(PV::FBranchParentNumberAttribute::FindAttribute(InCollection))
+			, BranchNumber(PV::FBranchNumberAttribute::FindAttribute(InCollection))
+			, BranchChildren(PV::FBranchChildrenAttribute::FindAttribute(InCollection))
+		{}
+
+		bool IsValid() const
+		{
+			return PointGroundGradient
+				&& PointScaleGradient.IsValid()
+				&& PointPlantGradient.IsValid()
+				&& PointBranchGradient.IsValid()
+				&& PointLengthFromRoot.IsValid()
+				&& PointScale.IsValid()
+				&& BranchPoints.IsValid()
+				&& BranchParentNumber.IsValid()
+				&& BranchNumber.IsValid()
+				&& BranchChildren.IsValid();
+		}
+	};
+
+	// Re-computes the point- ground gradient, scale gradient, plant gradient, and branch gradient.
+	PROCEDURALVEGETATION_API extern void ComputeGradients(const FComputeGradientsAttributes& InAttributes);
+
+	struct FComputeBudStatusAttributes
+	{ 
+		PV::FBudStatusAttributeView BudStatus;
+		PV::FBranchPointsAttributeConstView BranchPoints;
+		PV::FBranchParentNumberAttributeConstView BranchParentNumber;
+		PV::FBranchNumberAttributeConstView BranchNumber;
+		PV::FBranchChildrenAttributeConstView BranchChildren;
+		PV::FPointLengthFromRootAttributeConstView PointLengthFromRoot;
+
+		FComputeBudStatusAttributes(PV::FBudStatusAttributeView InBudStatus,
+			PV::FBranchPointsAttributeConstView InBranchPoints,
+			PV::FBranchParentNumberAttributeConstView InBranchParentNumber,
+			PV::FBranchNumberAttributeConstView InBranchNumber,
+			PV::FBranchChildrenAttributeConstView InBranchChildren,
+			PV::FPointLengthFromRootAttributeConstView InPointLengthFromRoot)
+			: BudStatus(InBudStatus)
+			, BranchPoints(InBranchPoints)
+			, BranchParentNumber(InBranchParentNumber)
+			, BranchNumber(InBranchNumber)
+			, BranchChildren(InBranchChildren)
+			, PointLengthFromRoot(InPointLengthFromRoot)
+		{
+		}
+
+		FComputeBudStatusAttributes(FManagedArrayCollection& InCollection)
+			: BudStatus(PV::FBudStatusAttribute::FindAttribute(InCollection))
+			, BranchPoints(PV::FBranchPointsAttribute::FindAttribute(InCollection))
+			, BranchParentNumber(PV::FBranchParentNumberAttribute::FindAttribute(InCollection))
+			, BranchNumber(PV::FBranchNumberAttribute::FindAttribute(InCollection))
+			, BranchChildren(PV::FBranchChildrenAttribute::FindAttribute(InCollection))
+			, PointLengthFromRoot(PV::FPointLengthFromRootAttribute::FindAttribute(InCollection))		
+		{
+		}
+
+		bool IsValid() const
+		{
+			return BudStatus
+				&& BranchPoints
+				&& BranchParentNumber
+				&& BranchNumber
+				&& BranchChildren
+				&& PointLengthFromRoot;
+		}
+	};
+
+	// Re-computes the bud status for all points.
+	PROCEDURALVEGETATION_API extern void ComputeBudStatus(const FComputeBudStatusAttributes& InAttributes);
+
+	struct FComputeBudLateralMeristemAttributes
+	{
+		PV::FBudLateralMeristemAttributeView BudLateralMeristem;
+		PV::FPointPositionAttributeConstView PointPosition;
+		PV::FPointScaleAttributeConstView PointScale;
+		PV::FPointLengthFromRootAttributeConstView PointLengthFromRoot;
+		PV::FBudDirectionAttributeConstView BudDirection;
+		PV::FBranchPointsAttributeConstView BranchPoints;
+		PV::FBranchParentNumberAttributeConstView BranchParentNumber;
+		PV::FBranchChildrenAttributeConstView BranchChildren;
+		PV::FBranchNumberAttributeConstView BranchNumber;
+
+		FComputeBudLateralMeristemAttributes(FManagedArrayCollection& InCollection)
+			: BudLateralMeristem(PV::FBudLateralMeristemAttribute::FindAttribute(InCollection))
+			, PointPosition(PV::FPointPositionAttribute::FindAttribute(InCollection))
+			, PointScale(PV::FPointScaleAttribute::FindAttribute(InCollection))
+			, PointLengthFromRoot(PV::FPointLengthFromRootAttribute::FindAttribute(InCollection))
+			, BudDirection(PV::FBudDirectionAttribute::FindAttribute(InCollection))
+			, BranchPoints(PV::FBranchPointsAttribute::FindAttribute(InCollection))
+			, BranchParentNumber(PV::FBranchParentNumberAttribute::FindAttribute(InCollection))
+			, BranchChildren(PV::FBranchChildrenAttribute::FindAttribute(InCollection))
+			, BranchNumber(PV::FBranchNumberAttribute::FindAttribute(InCollection))
+		{
+		}
+
+		FComputeBudLateralMeristemAttributes(
+			PV::FBudLateralMeristemAttributeView InBudLateralMeristem,
+			PV::FPointPositionAttributeConstView InPointPosition,
+			PV::FPointScaleAttributeConstView InPointScale,
+			PV::FPointLengthFromRootAttributeConstView InPointLengthFromRoot,
+			PV::FBudDirectionAttributeConstView InBudDirection,
+			PV::FBranchPointsAttributeConstView InBranchPoints,
+			PV::FBranchParentNumberAttributeConstView InBranchParentNumber,
+			PV::FBranchChildrenAttributeConstView InBranchChildren,
+			PV::FBranchNumberAttributeConstView InBranchNumber)
+			: BudLateralMeristem(InBudLateralMeristem)
+			, PointPosition(InPointPosition)
+			, PointScale(InPointScale)
+			, PointLengthFromRoot(InPointLengthFromRoot)
+			, BudDirection(InBudDirection)
+			, BranchPoints(InBranchPoints)
+			, BranchParentNumber(InBranchParentNumber)
+			, BranchChildren(InBranchChildren)
+			, BranchNumber(InBranchNumber)
+		{
+		}
+
+		bool IsValid() const
+		{
+			return BudLateralMeristem
+				&& PointPosition
+				&& PointScale
+				&& PointLengthFromRoot
+				&& BudDirection
+				&& BranchPoints
+				&& BranchParentNumber
+				&& BranchChildren
+				&& BranchNumber;
+		}
+	};
+
+	// Re-computes the bud lateral meristem for all points.
+	PROCEDURALVEGETATION_API extern void ComputeBudLateralMeristem(const FComputeBudLateralMeristemAttributes& InAttributes);
+
+	// Re-computes the bud light detected for all points.
+	PROCEDURALVEGETATION_API extern void ComputeBudLightDetected(PV::FBudLightDetectedAttributeView BudLightDetectedAttribute);
+
+	struct FComputeBranchHierarchyNumberAttributes
+	{
+		PV::FBranchParentNumberAttributeConstView BranchParentNumber;
+		PV::FBranchNumberAttributeConstView BranchNumber;
+		PV::FBranchChildrenAttributeConstView BranchChildren;
+		PV::FBranchHierarchyNumberAttributeView BranchHierarchyNumber;
+
+		FComputeBranchHierarchyNumberAttributes(
+			PV::FBranchParentNumberAttributeConstView InBranchParentNumber,
+			PV::FBranchNumberAttributeConstView InBranchNumber,
+			PV::FBranchChildrenAttributeConstView InBranchChildren,
+			PV::FBranchHierarchyNumberAttributeView InBranchHierarchyNumber)
+			: BranchParentNumber(InBranchParentNumber)
+			, BranchNumber(InBranchNumber)
+			, BranchChildren(InBranchChildren)
+			, BranchHierarchyNumber(InBranchHierarchyNumber)
+		{
+		}
+
+		FComputeBranchHierarchyNumberAttributes(FManagedArrayCollection& InCollection)
+			: BranchParentNumber(PV::FBranchParentNumberAttribute::FindAttribute(InCollection))
+			, BranchNumber(PV::FBranchNumberAttribute::FindAttribute(InCollection))
+			, BranchChildren(PV::FBranchChildrenAttribute::FindAttribute(InCollection))
+			, BranchHierarchyNumber(PV::FBranchHierarchyNumberAttribute::FindAttribute(InCollection))
+		{
+		}
+
+		bool IsValid() const
+		{
+			return BranchParentNumber.IsValid()
+				&& BranchNumber.IsValid()
+				&& BranchChildren.IsValid()
+				&& BranchHierarchyNumber.IsValid();
+		}
+	};
+
+	// Re-computes the BranchHierarchyNumber attribute given an already existing branch hierarchy.
+	PROCEDURALVEGETATION_API extern void ComputeBranchHierarchyNumbers(const FComputeBranchHierarchyNumberAttributes& InAttributes);
+
+	// This function ensures all points have a valid point scale by walking the point hierarchy
+	// and finding points with missing point scales. It then interpolates the missing point scales
+	// using the point scale of the parent/child points.
+	PROCEDURALVEGETATION_API extern void InterpolatePointsScales(
+		PV::FPointScaleAttributeView PointScaleAttribute,
+		PV::FPointLengthFromRootAttributeConstView PointLengthFromRootAttribute,
+		PV::FBranchPointsAttributeConstView BranchPointsAttribute,
+		PV::FBranchParentNumberAttributeConstView BranchParentNumberAttribute,
+		PV::FBranchChildrenAttributeConstView BranchChildrenAttribute,
+		PV::FBranchNumberAttributeConstView BranchNumberAttribute
+	);
+
+	struct FSplitFusedPointsAttributes
+	{
+		FBranchParentNumberAttributeConstView BranchParentNumber;
+		FBranchNumberAttributeConstView BranchNumber;
+		FBranchPointsAttributeView BranchPoints;
+
+		FPointPositionAttributeAccessor PointPosition;
+		FPointLengthFromRootAttributeAccessor PointLengthFromRoot;
+		FPointLengthFromSeedAttributeAccessor PointLengthFromSeed;
+		FPointScaleGradientAttributeAccessor PointScaleGradient;
+		FPointHullGradientAttributeAccessor PointHullGradient;
+		FPointMainTrunkGradientAttributeAccessor PointMainTrunkGradient;
+		FPointGroundGradientAttributeAccessor PointGroundGradient;
+		FPointScaleAttributeAccessor PointScale;
+		FBudDirectionAttributeAccessor BudDirection;
+		FBudStatusAttributeAccessor BudStatus;
+		FBudLateralMeristemAttributeAccessor BudLateralMeristem;
+		FPointBranchGradientAttributeAccessor PointBranchGradient;
+		FBudHormoneLevelsAttributeAccessor BudHormoneLevels;
+		FPointPlantGradientAttributeAccessor PointPlantGradient;
+		FPointBudNumberAttributeAccessor PointBudNumber;
+		FPointNjordPixelIndexAttributeAccessor PointNjordPixelIndex;
+		FBudLightDetectedAttributeAccessor BudLightDetected;
+		FBudDevelopmentAttributeAccessor BudDevelopment;
+
+		FSplitFusedPointsAttributes(FSplitFusedPointsAttributes&&) = default;
+		FSplitFusedPointsAttributes(FManagedArrayCollection& InCollection)
+			: BranchParentNumber(FBranchParentNumberAttribute::FindAttribute(InCollection))
+			, BranchNumber(FBranchNumberAttribute::FindAttribute(InCollection))
+			, BranchPoints(FBranchPointsAttribute::FindAttribute(InCollection))
+			, PointPosition(InCollection)
+			, PointLengthFromRoot(InCollection)
+			, PointLengthFromSeed(InCollection)
+			, PointScaleGradient(InCollection)
+			, PointHullGradient(InCollection)
+			, PointMainTrunkGradient(InCollection)
+			, PointGroundGradient(InCollection)
+			, PointScale(InCollection)
+			, BudDirection(InCollection)
+			, BudStatus(InCollection)
+			, BudLateralMeristem(InCollection)
+			, PointBranchGradient(InCollection)
+			, BudHormoneLevels(InCollection)
+			, PointPlantGradient(InCollection)
+			, PointBudNumber(InCollection)
+			, PointNjordPixelIndex(InCollection)
+			, BudLightDetected(InCollection)
+			, BudDevelopment(InCollection)
+		{}
+
+		bool IsValid() const
+		{
+			return BranchParentNumber.IsValid()
+				&& BranchNumber.IsValid()
+				&& BranchPoints.IsValid()
+				&& PointPosition.IsValid()
+				&& PointLengthFromRoot.IsValid()
+				&& PointLengthFromSeed.IsValid()
+				&& PointScaleGradient.IsValid()
+				&& PointHullGradient.IsValid()
+				&& PointMainTrunkGradient.IsValid()
+				&& PointGroundGradient.IsValid()
+				&& PointScale.IsValid()
+				&& BudDirection.IsValid()
+				&& BudStatus.IsValid()
+				&& BudLateralMeristem.IsValid()
+				&& PointBranchGradient.IsValid()
+				&& BudHormoneLevels.IsValid()
+				&& PointPlantGradient.IsValid()
+				&& PointBudNumber.IsValid()
+				&& PointNjordPixelIndex.IsValid()
+				&& BudLightDetected.IsValid()
+				&& BudDevelopment.IsValid();
+		}
+	};
+
+	// Converts a growth data collection with fused points to a split points collection.
+	// Will not validate that the input collection has fused points.
+	// NOTE: Assumes input collection is valid growth data.
+	PROCEDURALVEGETATION_API extern bool SplitFusedPoints(FSplitFusedPointsAttributes InAttributes);
+
+	struct FRecomputeGrowthDataAttributes
+	{
+		FPointPositionAttributeConstView PointPosition;
+		FBranchPointsAttributeConstView BranchPoints;
+		FBranchParentNumberAttributeConstView BranchParentNumber;
+		FBranchParentsAttributeConstView BranchParents;
+		FBranchNumberAttributeConstView BranchNumber;
+		FBranchChildrenAttributeConstView BranchChildren;
+		FBranchHierarchyNumberAttributeConstView BranchHierarchyNumber;
+		FBranchPlantNumberAttributeConstView BranchPlantNumber;
+		FPointScaleAttributeView PointScale;
+		FPointBudNumberAttributeView PointBudNumber;
+		FBranchSourceBudNumberAttributeView BranchSourceBudNumber;
+		FPointNjordPixelIndexAttributeView PointNjordPixelIndex;
+		FPointLengthFromRootAttributeView PointLengthFromRoot;
+		FPointLengthFromSeedAttributeView PointLengthFromSeed;
+		FBudDirectionAttributeView BudDirection;
+		FBudLightDetectedAttributeView BudLightDetected;
+		FBudLateralMeristemAttributeView BudLateralMeristem;
+		FPointGroundGradientAttributeView PointGroundGradient;
+		FPointScaleGradientAttributeView PointScaleGradient;
+		FPointPlantGradientAttributeView PointPlantGradient;
+		FPointBranchGradientAttributeView PointBranchGradient;
+		FBudDevelopmentAttributeView BudDevelopment;
+		FBudHormoneLevelsAttributeView BudHormoneLevels;
+		FBudStatusAttributeView BudStatus;
+
+		FRecomputeGrowthDataAttributes() = default;
+		FRecomputeGrowthDataAttributes(FManagedArrayCollection& InCollection)
+			: PointPosition(FPointPositionAttribute::FindAttribute(InCollection))
+			, BranchPoints(FBranchPointsAttribute::FindAttribute(InCollection))
+			, BranchParentNumber(FBranchParentNumberAttribute::FindAttribute(InCollection))
+			, BranchParents(FBranchParentsAttribute::FindAttribute(InCollection))
+			, BranchNumber(FBranchNumberAttribute::FindAttribute(InCollection))
+			, BranchChildren(FBranchChildrenAttribute::FindAttribute(InCollection))
+			, BranchHierarchyNumber(FBranchHierarchyNumberAttribute::FindAttribute(InCollection))
+			, BranchPlantNumber(FBranchPlantNumberAttribute::FindAttribute(InCollection))
+			, PointScale(FPointScaleAttribute::FindAttribute(InCollection))
+			, PointBudNumber(FPointBudNumberAttribute::FindAttribute(InCollection))
+			, BranchSourceBudNumber(FBranchSourceBudNumberAttribute::FindAttribute(InCollection))
+			, PointNjordPixelIndex(FPointNjordPixelIndexAttribute::FindAttribute(InCollection))
+			, PointLengthFromRoot(FPointLengthFromRootAttribute::FindAttribute(InCollection))
+			, PointLengthFromSeed(FPointLengthFromSeedAttribute::FindAttribute(InCollection))
+			, BudDirection(FBudDirectionAttribute::FindAttribute(InCollection))
+			, BudLightDetected(FBudLightDetectedAttribute::FindAttribute(InCollection))
+			, BudLateralMeristem(FBudLateralMeristemAttribute::FindAttribute(InCollection))
+			, PointGroundGradient(FPointGroundGradientAttribute::FindAttribute(InCollection))
+			, PointScaleGradient(FPointScaleGradientAttribute::FindAttribute(InCollection))
+			, PointPlantGradient(FPointPlantGradientAttribute::FindAttribute(InCollection))
+			, PointBranchGradient(FPointBranchGradientAttribute::FindAttribute(InCollection))
+			, BudDevelopment(FBudDevelopmentAttribute::FindAttribute(InCollection))
+			, BudHormoneLevels(FBudHormoneLevelsAttribute::FindAttribute(InCollection))
+			, BudStatus(FBudStatusAttribute::FindAttribute(InCollection))
+		{}
+
+		bool IsValid() const
+		{
+			return PointPosition.IsValid()
+				&& BranchPoints.IsValid()
+				&& BranchParentNumber.IsValid()
+				&& BranchParents.IsValid()
+				&& BranchNumber.IsValid()
+				&& BranchChildren.IsValid()
+				&& BranchHierarchyNumber.IsValid()
+				&& BranchPlantNumber.IsValid()
+				&& PointScale.IsValid()
+				&& PointBudNumber.IsValid()
+				&& BranchSourceBudNumber.IsValid()
+				&& PointNjordPixelIndex.IsValid()
+				&& PointLengthFromRoot.IsValid()
+				&& PointLengthFromSeed.IsValid()
+				&& BudDirection.IsValid()
+				&& BudLightDetected.IsValid()
+				&& BudLateralMeristem.IsValid()
+				&& PointGroundGradient.IsValid()
+				&& PointScaleGradient.IsValid()
+				&& PointPlantGradient.IsValid()
+				&& PointBranchGradient.IsValid()
+				&& BudDevelopment.IsValid()
+				&& BudHormoneLevels.IsValid()
+				&& BudStatus.IsValid();
+		}
+	};
+
+	PROCEDURALVEGETATION_API extern bool RecomputeAllGrowthDataAttributes(const FRecomputeGrowthDataAttributes& InAttributes);
+}
diff --git a/Engine/Plugins/Experimental/ProceduralVegetationEditor/Source/ProceduralVegetation/Private/Helpers/PVImportHelpers.cpp b/Engine/Plugins/Experimental/ProceduralVegetationEditor/Source/ProceduralVegetation/Private/Helpers/PVImportHelpers.cpp
--- a/Engine/Plugins/Experimental/ProceduralVegetationEditor/Source/ProceduralVegetation/Private/Helpers/PVImportHelpers.cpp
+++ b/Engine/Plugins/Experimental/ProceduralVegetationEditor/Source/ProceduralVegetation/Private/Helpers/PVImportHelpers.cpp
@@ -1235,111 +1235,13 @@ bool PV::ImportHelper::GenerateGrowthDataFromBranchHierarchies(FManagedArrayColl
 		InBranchHierarchies
 	);
 
-	auto PointBudNumberAttribute = FPointBudNumberAttribute::GetAttribute(OutCollection);
-	auto BranchSourceBudNumberAttribute = FBranchSourceBudNumberAttribute::GetAttribute(OutCollection);
-	ComputeBudNumbers(
-		PointBudNumberAttribute,
-		BranchSourceBudNumberAttribute, 
-		BranchPointsAttribute
-	);
-
-	auto PointNjordPixelIndexAttribute = FPointNjordPixelIndexAttribute::GetAttribute(OutCollection);
-	ComputeNjordPixelIndex(
-		PointNjordPixelIndexAttribute,
-		BranchPointsAttribute,
-		BranchParentNumberAttribute, 
-		BranchNumberAttribute,
-		BranchChildrenAttribute
-	);
-
-	auto PointLengthFromRootAttribute = FPointLengthFromRootAttribute::GetAttribute(OutCollection);
-	auto PointLengthFromSeedAttribute = FPointLengthFromSeedAttribute::GetAttribute(OutCollection);
-	ComputeLengthFromRoot(
-		PointLengthFromRootAttribute,
-		PointLengthFromSeedAttribute,
-		PointPositionAttribute,
-		BranchPointsAttribute,
-		BranchParentNumberAttribute,
-		BranchNumberAttribute,
-		BranchChildrenAttribute
-	);
-
-	auto BudDirectionAttribute = FBudDirectionAttribute::GetAttribute(OutCollection);
-	ComputeBudDirections(
-		BudDirectionAttribute,
-		PointPositionAttribute,
-		BranchChildrenAttribute,
-		BranchParentNumberAttribute,
-		BranchNumberAttribute,
-		BranchPointsAttribute
-	);
-
-	auto BudLightDetectedAttribute = FBudLightDetectedAttribute::GetAttribute(OutCollection);
-	ComputeBudLightDetected(BudLightDetectedAttribute);
-
-	InterpolatePointsScales(
-		PointScaleAttribute,
-		PointLengthFromRootAttribute,
-		BranchPointsAttribute,
-		BranchParentNumberAttribute,
-		BranchChildrenAttribute,
-		BranchNumberAttribute
-	);
-
-	auto BudLateralMeristemAttribute = FBudLateralMeristemAttribute::GetAttribute(OutCollection);
-	ComputeBudLateralMeristem(
-		BudLateralMeristemAttribute,
-		PointPositionAttribute,
-		PointScaleAttribute,
-		PointLengthFromRootAttribute,
-		BudDirectionAttribute,
-		BranchPointsAttribute,
-		BranchParentNumberAttribute,
-		BranchChildrenAttribute,
-		BranchNumberAttribute
-	);
-
-	auto GroundGradientAttribute = PV::FPointGroundGradientAttribute::GetAttribute(OutCollection);
-	auto PointScaleGradientAttribute = PV::FPointScaleGradientAttribute::GetAttribute(OutCollection);
-	auto PointPlantGradientAttribute = PV::FPointPlantGradientAttribute::GetAttribute(OutCollection);
-	auto PointBranchGradientAttribute = PV::FPointBranchGradientAttribute::GetAttribute(OutCollection);
-	ComputeGradients(
-		GroundGradientAttribute,
-		PointScaleGradientAttribute,
-		PointPlantGradientAttribute,
-		PointBranchGradientAttribute,
-		PointLengthFromRootAttribute,
-		PointScaleAttribute,
-		BranchPointsAttribute,
-		BranchParentNumberAttribute,
-		BranchNumberAttribute,
-		BranchChildrenAttribute
-	);
-
-	auto BudDevelopmentAttribute = FBudDevelopmentAttribute::GetAttribute(OutCollection);
-	ComputeBudDevelopment(
-		BudDevelopmentAttribute,
-		PointPlantGradientAttribute,
-		BranchPointsAttribute,
-		BranchParentNumberAttribute,
-		BranchParentsAttribute,
-		BranchNumberAttribute
-	);
-
-	auto BudHormoneLevelsAttribute = FBudHormoneLevelsAttribute::GetAttribute(OutCollection);
-	ComputeBudHormoneLevels(
-		BudHormoneLevelsAttribute,
-		PointPlantGradientAttribute
-	);
+	const FRecomputeGrowthDataAttributes GrowthDataAttributes(OutCollection);
+	if (!GrowthDataAttributes.IsValid())
+	{
+		return false;
+	}
 
-	auto BudStatusAttribute = FBudStatusAttribute::GetAttribute(OutCollection);
-	ComputeBudStatus(
-		BudStatusAttribute,
-		BranchPointsAttribute,
-		BranchParentNumberAttribute,
-		BranchNumberAttribute,
-		BranchChildrenAttribute
-	);
+	RecomputeAllGrowthDataAttributes(GrowthDataAttributes);
 
 	if (bSplitPoints)
 	{
diff --git a/Engine/Plugins/Experimental/ProceduralVegetationEditor/Source/ProceduralVegetation/Private/Helpers/PVPlantTraversalHelper.h b/Engine/Plugins/Experimental/ProceduralVegetationEditor/Source/ProceduralVegetation/Private/Helpers/PVPlantTraversalHelper.h
new file mode 100644
--- /dev/null
+++ b/Engine/Plugins/Experimental/ProceduralVegetationEditor/Source/ProceduralVegetation/Private/Helpers/PVPlantTraversalHelper.h
@@ -0,0 +1,803 @@
+// Copyright Epic Games, Inc. All Rights Reserved.
+#pragma once
+
+#include "Utils/PVAttributes.h"
+
+namespace PV::PlantTraversalHelper
+{
+	inline int32 GetBranchParentIndex(
+		FBranchParentNumberAttributeConstView BranchParentNumberAttribute,
+		FBranchNumberAttributeConstView BranchNumberAttribute,
+		int32 BranchIndex
+	)
+	{
+		return BranchParentNumberAttribute.IsValidIndex(BranchIndex) 
+			? BranchNumberAttribute.Find(BranchParentNumberAttribute[BranchIndex])
+			: INDEX_NONE;
+	}
+
+	inline TArray<int32> GetPlantBranchIndices(
+		PV::FBranchPlantNumberAttributeConstView BranchPlantNumberAttribute,
+		int32 PlantNumber
+	)
+	{
+		TArray<int32> OutBranchIndices;
+
+		for (int32 i = 0; i < BranchPlantNumberAttribute.Num(); i++)
+		{
+			if (BranchPlantNumberAttribute[i] == PlantNumber)
+			{
+				OutBranchIndices.Add(i);
+			}
+		}
+
+		return OutBranchIndices;
+	}
+
+	inline TArray<int32> GetTrunkIndices(FBranchParentNumberAttributeConstView BranchParentNumberAttribute)
+	{
+		TArray<int32> TrunkIndicesResult;
+		for (int32 i = 0; i < BranchParentNumberAttribute.Num(); i++)
+		{
+			if (BranchParentNumberAttribute[i] == 0)
+			{
+				TrunkIndicesResult.Add(i);
+			}
+		}
+		return TrunkIndicesResult;
+	}
+
+	inline TArray<int32> GetUniquePlantNumbers(PV::FBranchPlantNumberAttributeConstView BranchPlantNumberAttribute)
+	{
+		TSet<int32> PlantNumberElements;
+		for (int32 i = 0; i < BranchPlantNumberAttribute.Num(); i++)
+		{
+			PlantNumberElements.Add(BranchPlantNumberAttribute[i]);
+		}
+		return PlantNumberElements.Array();
+	}
+
+	inline bool HasFusedPoints(
+		PV::FBranchPointsAttributeConstView BranchPointsAttribute,
+		PV::FBranchParentNumberAttributeConstView BranchParentNumberAttribute,
+		PV::FBranchNumberAttributeConstView BranchNumberAttribute
+	)
+	{
+		for (int32 BranchIndex = 0; BranchIndex < BranchPointsAttribute.Num(); ++BranchIndex)
+		{
+			const int32 BranchParentIndex = GetBranchParentIndex(BranchParentNumberAttribute, BranchNumberAttribute, BranchIndex);
+			if (BranchParentIndex != INDEX_NONE)
+			{
+				const TArray<int32>& BranchPoints = BranchPointsAttribute[BranchIndex];
+				const TArray<int32>& ParentBranchPoints = BranchPointsAttribute[BranchParentIndex];
+
+				if (BranchPoints.Num() > 0 && ParentBranchPoints.Num() > 0)
+				{
+					return ParentBranchPoints.Find(BranchPoints[0]) != INDEX_NONE;
+				}
+			}
+		}
+
+		return true;
+	}
+
+	enum class EForEachResult
+	{
+		Break    = 0,
+		Continue = 1
+	};
+
+	// Walks the branches in hierarchical order (e.g. start with the trunk and recursivly visit the child branches)
+	template <typename StateType>
+	void RecursiveWalkBranches(
+		FBranchParentNumberAttributeConstView BranchParentNumberAttribute,
+		FBranchChildrenAttributeConstView BranchChildrenAttribute,
+		FBranchNumberAttributeConstView BranchNumberAttribute,
+		const StateType& RootBranchState,
+		const TFunction<EForEachResult(int32, StateType&)>& Body
+	)
+	{
+		const TFunction<EForEachResult(int32, const StateType&)> WalkBranches_Recursive = [&](int32 BranchIndex, const StateType& State) -> EForEachResult
+		{
+			StateType BranchState = State;
+			if (Body(BranchIndex, BranchState) == EForEachResult::Break)
+			{
+				return EForEachResult::Break;
+			}
+
+			const TArray<int32>& BranchChildren = BranchChildrenAttribute[BranchIndex];
+			for (int32 BranchChild : BranchChildren)
+			{
+				const int32 ChildBranchIndex = BranchNumberAttribute.Find(BranchChild);
+				check(ChildBranchIndex != INDEX_NONE);
+
+				if (WalkBranches_Recursive(ChildBranchIndex, BranchState) == EForEachResult::Break)
+				{
+					return EForEachResult::Break;
+				}
+			}
+
+			return EForEachResult::Continue;
+		};
+
+		for (const int32 TrunkIndex : GetTrunkIndices(BranchParentNumberAttribute))
+		{
+			if (WalkBranches_Recursive(TrunkIndex, RootBranchState) == EForEachResult::Break)
+			{
+				return;
+			}
+		}
+	}
+
+	inline void RecursiveWalkBranches(
+		FBranchParentNumberAttributeConstView BranchParentNumberAttribute,
+		FBranchChildrenAttributeConstView BranchChildrenAttribute,
+		FBranchNumberAttributeConstView BranchNumberAttribute,
+		const TFunction<EForEachResult(int32 BranchIndex)>& Body
+	)
+	{
+		RecursiveWalkBranches<int32>(
+			BranchParentNumberAttribute,
+			BranchChildrenAttribute,
+			BranchNumberAttribute,
+			INDEX_NONE,
+			[&](int32 BranchIndex, int32&) -> EForEachResult
+			{
+				if (Body(BranchIndex) == EForEachResult::Break)
+				{
+					return EForEachResult::Break;
+				}
+				return EForEachResult::Continue;
+			}
+		);
+	}
+
+	// Walks the branches in reversed hierarchical order (e.g. first visit all leaf nodes, then their parents, and so on, visiting the tunk last)
+	inline void RecursiveWalkBranches_Reversed(
+		FBranchParentNumberAttributeConstView BranchParentNumberAttribute,
+		FBranchChildrenAttributeConstView BranchChildrenAttribute,
+		FBranchNumberAttributeConstView BranchNumberAttribute,
+		const TFunction<EForEachResult(int32 BranchIndex)>& Body
+	)
+	{
+		TArray<int32> BranchIndices;
+		BranchIndices.Reserve(BranchParentNumberAttribute.Num());
+
+		RecursiveWalkBranches(
+			BranchParentNumberAttribute,
+			BranchChildrenAttribute,
+			BranchNumberAttribute,
+			[&](int32 BranchIndex) -> EForEachResult
+			{
+				BranchIndices.Add(BranchIndex);
+				return EForEachResult::Continue;
+			}
+		);
+
+		for (int32 i = BranchIndices.Num() - 1; i >= 0; --i)
+		{
+			if (Body(BranchIndices[i]) == EForEachResult::Break)
+			{
+				break;
+			}
+		}
+	}
+
+	inline TArray<int32> GetBranchPointFoliageEntryIds(
+		PV::FBranchPointsAttributeConstView BranchPointsAttribute,
+		PV::FPointLengthFromRootAttributeConstView PointLengthFromRootAttribute,
+		PV::FBranchFoliageIDsAttributeConstView BranchFoliageIDsAttribute,
+		PV::FFoliageLengthFromRootAttributeConstView FoliageLengthFromRootAttribute,
+		int32 BranchIndex,
+		int32 PointIndex)
+	{
+		const TArray<int32> BranchPoints = BranchPointsAttribute[BranchIndex];
+		const int32 PointBranchIndex = BranchPoints.Find(PointIndex);
+		if (PointBranchIndex == INDEX_NONE)
+		{
+			return TArray<int32>();
+		}
+
+		const int32 PreviousPointBranchIndex = PointBranchIndex > 0 ? PointBranchIndex - 1 : INDEX_NONE;
+		if (PreviousPointBranchIndex == INDEX_NONE)
+		{
+			return TArray<int32>();
+		}
+
+		const int32 PreviousPointIndex = BranchPoints[PreviousPointBranchIndex];
+
+		const float BranchPointLengthFromRoot = PointLengthFromRootAttribute[PointIndex];
+		const float PreviousBranchPointLengthFromRoot = PointLengthFromRootAttribute[PreviousPointIndex];
+
+		const TArray<int32> FoliageEntryIds = BranchFoliageIDsAttribute[BranchIndex];
+
+		TArray<int32> OutFoliageEntryIds;
+		OutFoliageEntryIds.Reserve(FoliageEntryIds.Num());
+
+		for (const int32 FoliageEntryId : FoliageEntryIds)
+		{
+			const float FoliageEntryLengthFromRoot = FoliageLengthFromRootAttribute[FoliageEntryId];
+
+			if (FoliageEntryLengthFromRoot <= BranchPointLengthFromRoot && FoliageEntryLengthFromRoot > PreviousBranchPointLengthFromRoot)
+			{
+				OutFoliageEntryIds.Add(FoliageEntryId);
+			}
+		}
+
+		return OutFoliageEntryIds;
+	}
+
+	inline TArray<int32> GetBranchPointChildBranchIndices_SplitPoints(
+		PV::FBranchChildrenAttributeConstView BranchChildrenAttribute,
+		PV::FBranchParentNumberAttributeConstView BranchParentNumberAttribute,
+		PV::FBranchNumberAttributeConstView BranchNumberAttribute,
+		PV::FBranchPointsAttributeConstView BranchPointsAttribute,
+		PV::FPointLengthFromRootAttributeConstView PointLengthFromRootAttribute,
+		int32 BranchIndex,
+		int32 PointIndex
+	)
+	{
+		const TArray<int32> BranchPoints = BranchPointsAttribute[BranchIndex];
+		const int32 PointBranchIndex = BranchPoints.Find(PointIndex);
+		if (PointBranchIndex == INDEX_NONE)
+		{
+			return TArray<int32>();
+		}
+
+		const int32 PreviousPointBranchIndex = PointBranchIndex > 0 ? PointBranchIndex - 1 : INDEX_NONE;
+		if (PreviousPointBranchIndex == INDEX_NONE)
+		{
+			return TArray<int32>();
+		}
+
+		const int32 PreviousPointIndex = BranchPoints[PreviousPointBranchIndex];
+
+		const TArray<int32>& BranchChildren = BranchChildrenAttribute[BranchIndex];
+
+		const float BranchPointLengthFromRoot = PointLengthFromRootAttribute[PointIndex];
+		const float PreviousBranchPointLengthFromRoot = PointLengthFromRootAttribute[PreviousPointIndex];
+
+		TArray<int32> ChildrenIndices;
+		ChildrenIndices.Reserve(BranchChildren.Num());
+
+		for (const int32 BranchChild : BranchChildren)
+		{
+			const int32 ChildIndex = BranchNumberAttribute.Find(BranchChild);
+			if (ChildIndex == INDEX_NONE)
+			{
+				continue;
+			}
+
+			const int32 ParentIndex = GetBranchParentIndex(BranchParentNumberAttribute, BranchNumberAttribute, ChildIndex);
+			if (ParentIndex != BranchIndex)
+			{
+				continue;
+			}
+
+			const TArray<int32>& ChildBranchPoints = BranchPointsAttribute[ChildIndex];
+			if (ChildBranchPoints.Num() == 0)
+			{
+				continue;
+			}
+
+			const int32 FirstChildPointIndex = ChildBranchPoints[0];
+			const float ChildLengthFromRoot = PointLengthFromRootAttribute[FirstChildPointIndex];
+			if (ChildLengthFromRoot > BranchPointLengthFromRoot || ChildLengthFromRoot <= PreviousBranchPointLengthFromRoot)
+			{
+				continue;
+			}
+
+			ChildrenIndices.Add(ChildIndex);
+		}
+
+		return ChildrenIndices;
+	}
+
+	inline TArray<int32> GetBranchPointChildBranchIndices_FusedPoints(
+		PV::FBranchChildrenAttributeConstView BranchChildrenAttribute,
+		PV::FBranchParentNumberAttributeConstView BranchParentNumberAttribute,
+		PV::FBranchNumberAttributeConstView BranchNumberAttribute,
+		PV::FBranchPointsAttributeConstView BranchPointsAttribute,
+		int32 BranchIndex,
+		int32 PointIndex
+	)
+	{
+		const TArray<int32>& BranchChildren = BranchChildrenAttribute[BranchIndex];
+
+		TArray<int32> ChildrenIndices;
+		ChildrenIndices.Reserve(BranchChildren.Num());
+
+		for (const int32 BranchChild : BranchChildren)
+		{
+			const int32 ChildIndex = BranchNumberAttribute.Find(BranchChild);
+			if (ChildIndex == INDEX_NONE)
+			{
+				continue;
+			}
+
+			const int32 ParentIndex = GetBranchParentIndex(BranchParentNumberAttribute, BranchNumberAttribute, ChildIndex);
+			if (ParentIndex != BranchIndex)
+			{
+				continue;
+			}
+
+			const TArray<int32>& ChildBranchPoints = BranchPointsAttribute[ChildIndex];
+			if (ChildBranchPoints.Num() == 0)
+			{
+				continue;
+			}
+
+			const int32 FirstChildPointIndex = ChildBranchPoints[0];
+			if (FirstChildPointIndex == PointIndex)
+			{
+				ChildrenIndices.Add(ChildIndex);
+			}
+		}
+
+		return ChildrenIndices;
+	}
+
+	inline TArray<int32> GetBranchPointChildBranchIndices(
+		PV::FBranchChildrenAttributeConstView BranchChildrenAttribute,
+		PV::FBranchParentNumberAttributeConstView BranchParentNumberAttribute,
+		PV::FBranchNumberAttributeConstView BranchNumberAttribute,
+		PV::FBranchPointsAttributeConstView BranchPointsAttribute,
+		PV::FPointLengthFromRootAttributeConstView PointLengthFromRootAttribute,
+		int32 BranchIndex,
+		int32 PointIndex
+	)
+	{
+		if (HasFusedPoints(BranchPointsAttribute, BranchParentNumberAttribute, BranchNumberAttribute))
+		{
+			return GetBranchPointChildBranchIndices_FusedPoints(
+				BranchChildrenAttribute,
+				BranchParentNumberAttribute,
+				BranchNumberAttribute,
+				BranchPointsAttribute,
+				BranchIndex,
+				PointIndex
+			);
+		}
+		else
+		{
+			return GetBranchPointChildBranchIndices_SplitPoints(
+				BranchChildrenAttribute,
+				BranchParentNumberAttribute,
+				BranchNumberAttribute,
+				BranchPointsAttribute,
+				PointLengthFromRootAttribute,
+				BranchIndex,
+				PointIndex
+			);
+		}
+	}
+
+	struct FRecursiveWalkBranchPointsParams
+	{
+		int32 BranchIndex;
+		int32 BranchPointIndex;
+		int32 PointIndex;
+		const TArray<int32>& ChildBranchIndices;
+		const TArray<int32>& ChildPointIndices;
+	};
+
+	template <typename StateType, bool bWalkAxillaryFirst = false>
+	inline EForEachResult RecursiveWalkBranchPoints_SplitPoints(
+		PV::FBranchChildrenAttributeConstView BranchChildrenAttribute,
+		PV::FBranchParentNumberAttributeConstView BranchParentNumberAttribute,
+		PV::FBranchNumberAttributeConstView BranchNumberAttribute,
+		PV::FBranchPointsAttributeConstView BranchPointsAttribute,
+		PV::FPointLengthFromRootAttributeConstView PointLengthFromRootAttribute,
+		int32 BranchIndex,
+		const StateType& RootPointState,
+		const TFunction<EForEachResult(const FRecursiveWalkBranchPointsParams& Params, StateType& PointState)>& Body
+	)
+	{
+		const TFunction<EForEachResult(int32, int32, const StateType&)> WalkPoints = [&](int32 BranchIndex, int32 BranchPointIndex, const StateType& ParentState) -> EForEachResult
+		{
+			StateType PointState = ParentState;
+			
+			const TArray<int32>& CurrentBranchPoints = BranchPointsAttribute[BranchIndex];
+			const int32 PointIndex = CurrentBranchPoints[BranchPointIndex];
+
+			const TArray<int32> ChildBranchIndices = GetBranchPointChildBranchIndices_SplitPoints(
+				BranchChildrenAttribute,
+				BranchParentNumberAttribute,
+				BranchNumberAttribute,
+				BranchPointsAttribute,
+				PointLengthFromRootAttribute,
+				BranchIndex,
+				PointIndex
+			);
+
+			TArray<int32> ChildPointIndices;
+			ChildPointIndices.SetNum(ChildBranchIndices.Num());
+			for (int32 i = 0; i < ChildBranchIndices.Num(); ++i)
+			{
+				const int32 ChildBranchIndex = ChildBranchIndices[i];
+				ChildPointIndices[i] = BranchPointsAttribute[ChildBranchIndex].Num() > 0 ? BranchPointsAttribute[ChildBranchIndex][0] : INDEX_NONE;
+			}
+
+			if (Body({ BranchIndex, BranchPointIndex, PointIndex, ChildBranchIndices, ChildPointIndices }, PointState) == EForEachResult::Break)
+			{
+				return EForEachResult::Break;
+			}
+
+			struct FChildPoint
+			{
+				int32 BranchIndex;
+				int32 BranchPointIndex;
+			};
+			TArray<FChildPoint, TInlineAllocator<3>> Children;
+			Children.Reserve(1 + ChildBranchIndices.Num());
+
+			if constexpr (bWalkAxillaryFirst)
+			{
+				for (int32 ChildBranchIndex : ChildBranchIndices)
+				{
+					Children.Emplace(ChildBranchIndex, 0);
+				}
+
+				if (BranchPointIndex < CurrentBranchPoints.Num() - 1)
+				{
+					Children.Emplace(BranchIndex, BranchPointIndex + 1);
+				}
+			}
+			else
+			{
+				if (BranchPointIndex < CurrentBranchPoints.Num() - 1)
+				{
+					Children.Emplace(BranchIndex, BranchPointIndex + 1);
+				}
+
+				for (int32 ChildBranchIndex : ChildBranchIndices)
+				{
+					Children.Emplace(ChildBranchIndex, 0);
+				}
+			}
+
+			for (const auto& [ChildBranchIndex, ChildBranchPointIndex] : Children)
+			{
+				if (WalkPoints(ChildBranchIndex, ChildBranchPointIndex, PointState) == EForEachResult::Break)
+				{
+					return EForEachResult::Break;
+				}
+			}
+
+			return EForEachResult::Continue;
+		};
+
+		return WalkPoints(BranchIndex, 0, RootPointState);
+	}
+
+	template <typename StateType, bool bWalkAxillaryFirst = false>
+	inline EForEachResult RecursiveWalkBranchPoints_FusedPoints(
+		PV::FBranchChildrenAttributeConstView BranchChildrenAttribute,
+		PV::FBranchParentNumberAttributeConstView BranchParentNumberAttribute,
+		PV::FBranchNumberAttributeConstView BranchNumberAttribute,
+		PV::FBranchPointsAttributeConstView BranchPointsAttribute,
+		int32 BranchIndex,
+		const StateType& RootPointState,
+		const TFunction<EForEachResult(const FRecursiveWalkBranchPointsParams& Params, StateType& PointState)>& Body
+	)
+	{
+		const TFunction<EForEachResult(int32, int32, const StateType&)> WalkPoints = [&](int32 BranchIndex, int32 BranchPointIndex, const StateType& ParentState) -> EForEachResult
+		{
+			StateType PointState = ParentState;
+			
+			const TArray<int32>& CurrentBranchPoints = BranchPointsAttribute[BranchIndex];
+			const int32 PointIndex = CurrentBranchPoints[BranchPointIndex];
+
+			const TArray<int32> ChildBranchIndices = GetBranchPointChildBranchIndices_FusedPoints(
+				BranchChildrenAttribute,
+				BranchParentNumberAttribute,
+				BranchNumberAttribute,
+				BranchPointsAttribute,
+				BranchIndex,
+				PointIndex
+			);
+
+			TArray<int32> ChildPointIndices;
+			ChildPointIndices.SetNum(ChildBranchIndices.Num());
+			for (int32 i = 0; i < ChildBranchIndices.Num(); ++i)
+			{
+				const int32 ChildBranchIndex = ChildBranchIndices[i];
+				ChildPointIndices[i] = BranchPointsAttribute[ChildBranchIndex].Num() > 1 ? BranchPointsAttribute[ChildBranchIndex][1] : INDEX_NONE;
+			}
+
+			if (Body({ BranchIndex, BranchPointIndex, PointIndex, ChildBranchIndices, ChildPointIndices }, PointState) == EForEachResult::Break)
+			{
+				return EForEachResult::Break;
+			}
+
+			struct FChildPoint
+			{
+				int32 BranchIndex;
+				int32 BranchPointIndex;
+			};
+			TArray<FChildPoint, TInlineAllocator<3>> Children;
+			Children.Reserve(1 + ChildBranchIndices.Num());
+
+			if constexpr (bWalkAxillaryFirst)
+			{
+				for (int32 ChildBranchIndex : ChildBranchIndices)
+				{
+					Children.Emplace(ChildBranchIndex, 1);
+				}
+
+				if (BranchPointIndex < CurrentBranchPoints.Num() - 1)
+				{
+					Children.Emplace(BranchIndex, BranchPointIndex + 1);
+				}
+			}
+			else
+			{
+				if (BranchPointIndex < CurrentBranchPoints.Num() - 1)
+				{
+					Children.Emplace(BranchIndex, BranchPointIndex + 1);
+				}
+
+				for (int32 ChildBranchIndex : ChildBranchIndices)
+				{
+					Children.Emplace(ChildBranchIndex, 1);
+				}
+			}
+
+			for (const auto& [ChildBranchIndex, ChildBranchPointIndex] : Children)
+			{
+				if (WalkPoints(ChildBranchIndex, ChildBranchPointIndex, PointState) == EForEachResult::Break)
+				{
+					return EForEachResult::Break;
+				}
+			}
+
+			return EForEachResult::Continue;
+		};
+
+		return WalkPoints(BranchIndex, 0, RootPointState);
+	}
+
+	template<typename StateType, bool bWalkAxillaryFirst = false>
+	inline EForEachResult RecursiveWalkBranchPoints(
+		PV::FBranchChildrenAttributeConstView BranchChildrenAttribute,
+		PV::FBranchParentNumberAttributeConstView BranchParentNumberAttribute,
+		PV::FBranchNumberAttributeConstView BranchNumberAttribute,
+		PV::FBranchPointsAttributeConstView BranchPointsAttribute,
+		PV::FPointLengthFromRootAttributeConstView PointLengthFromRootAttribute,
+		int32 BranchIndex,
+		const StateType& RootPointState,
+		const TFunction<EForEachResult(const FRecursiveWalkBranchPointsParams& Params, StateType& PointState)>& Body
+	)
+	{
+		if (HasFusedPoints(BranchPointsAttribute, BranchParentNumberAttribute, BranchNumberAttribute))
+		{
+			return RecursiveWalkBranchPoints_FusedPoints<StateType, bWalkAxillaryFirst>(
+				BranchChildrenAttribute,
+				BranchParentNumberAttribute,
+				BranchNumberAttribute,
+				BranchPointsAttribute,
+				BranchIndex,
+				RootPointState,
+				Body
+			);
+		}
+		else
+		{
+			return RecursiveWalkBranchPoints_SplitPoints<StateType, bWalkAxillaryFirst>(
+				BranchChildrenAttribute,
+				BranchParentNumberAttribute,
+				BranchNumberAttribute,
+				BranchPointsAttribute,
+				PointLengthFromRootAttribute,
+				BranchIndex,
+				RootPointState,
+				Body
+			);
+		}
+	}
+
+	template<bool bWalkAxillaryFirst = false>
+	inline EForEachResult RecursiveWalkBranchPoints(
+		PV::FBranchChildrenAttributeConstView BranchChildrenAttribute,
+		PV::FBranchParentNumberAttributeConstView BranchParentNumberAttribute,
+		PV::FBranchNumberAttributeConstView BranchNumberAttribute,
+		PV::FBranchPointsAttributeConstView BranchPointsAttribute,
+		PV::FPointLengthFromRootAttributeConstView PointLengthFromRootAttribute,
+		int32 BranchIndex,
+		const TFunction<EForEachResult(const FRecursiveWalkBranchPointsParams& Params)>& Body
+	)
+	{
+		if (HasFusedPoints(BranchPointsAttribute, BranchParentNumberAttribute, BranchNumberAttribute))
+		{
+			return RecursiveWalkBranchPoints_FusedPoints<int32, bWalkAxillaryFirst>(
+				BranchChildrenAttribute,
+				BranchParentNumberAttribute,
+				BranchNumberAttribute,
+				BranchPointsAttribute,
+				BranchIndex,
+				0,
+				[&](const FRecursiveWalkBranchPointsParams& Params, int32&) -> EForEachResult
+				{
+					return Body(Params);
+				}
+			);
+		}
+		else
+		{
+			return RecursiveWalkBranchPoints_SplitPoints<int32, bWalkAxillaryFirst>(
+				BranchChildrenAttribute,
+				BranchParentNumberAttribute,
+				BranchNumberAttribute,
+				BranchPointsAttribute,
+				PointLengthFromRootAttribute,
+				BranchIndex,
+				0,
+				[&](const FRecursiveWalkBranchPointsParams& Params, int32&) -> EForEachResult
+				{
+					return Body(Params);
+				}
+			);
+		}
+	}
+
+	template <typename StateType, bool bWalkAxillaryBranchesFirst = false>
+	void RecursiveWalkPlantPoints(
+		PV::FBranchChildrenAttributeConstView BranchChildrenAttribute,
+		PV::FBranchParentNumberAttributeConstView BranchParentNumberAttribute,
+		PV::FBranchNumberAttributeConstView BranchNumberAttribute,
+		PV::FBranchPointsAttributeConstView BranchPointsAttribute,
+		PV::FPointLengthFromRootAttributeConstView PointLengthFromRootAttribute,
+		const StateType& RootPointState,
+		const TFunction<EForEachResult(const FRecursiveWalkBranchPointsParams& Params, StateType& PointState)>& Body
+	)
+	{
+		if (HasFusedPoints(BranchPointsAttribute, BranchParentNumberAttribute, BranchNumberAttribute))
+		{
+			for (const int32 TrunkIndex : GetTrunkIndices(BranchParentNumberAttribute))
+			{
+				auto Result = RecursiveWalkBranchPoints_FusedPoints<StateType, bWalkAxillaryBranchesFirst>(
+					BranchChildrenAttribute,
+					BranchParentNumberAttribute,
+					BranchNumberAttribute,
+					BranchPointsAttribute,
+					TrunkIndex,
+					RootPointState,
+					Body
+				);
+				if (Result == EForEachResult::Break)
+				{
+					return;
+				}
+			}
+		}
+		else
+		{
+			for (const int32 TrunkIndex : GetTrunkIndices(BranchParentNumberAttribute))
+			{
+				auto Result = RecursiveWalkBranchPoints_SplitPoints<StateType, bWalkAxillaryBranchesFirst>(
+					BranchChildrenAttribute,
+					BranchParentNumberAttribute,
+					BranchNumberAttribute,
+					BranchPointsAttribute,
+					PointLengthFromRootAttribute,
+					TrunkIndex,
+					RootPointState,
+					Body
+				);
+				if (Result == EForEachResult::Break)
+				{
+					return;
+				}
+			}
+		}
+	}
+
+	template <bool bWalkAxillaryBranchesFirst = false>
+	inline void RecursiveWalkPlantPoints(
+		PV::FBranchChildrenAttributeConstView BranchChildrenAttribute,
+		PV::FBranchParentNumberAttributeConstView BranchParentNumberAttribute,
+		PV::FBranchNumberAttributeConstView BranchNumberAttribute,
+		PV::FBranchPointsAttributeConstView BranchPointsAttribute,
+		PV::FPointLengthFromRootAttributeConstView PointLengthFromRootAttribute,
+		const TFunction<EForEachResult(const FRecursiveWalkBranchPointsParams& Params)>& Body
+	)
+	{
+		RecursiveWalkPlantPoints<int32, bWalkAxillaryBranchesFirst>(
+			BranchChildrenAttribute,
+			BranchParentNumberAttribute,
+			BranchNumberAttribute,
+			BranchPointsAttribute,
+			PointLengthFromRootAttribute,
+			INDEX_NONE,
+			[&](const FRecursiveWalkBranchPointsParams& Params, int32&)
+			{
+				return Body(Params);
+			}
+		);
+	}
+
+	inline void ForEachUniquePointOnBranches(
+		PV::FBranchPointsAttributeConstView BranchPointsAttribute,
+		PV::FBranchParentNumberAttributeConstView BranchParentNumberAttribute,
+		PV::FBranchNumberAttributeConstView BranchNumberAttribute,
+		const TFunction<EForEachResult(int32 BranchIndex, int32 BranchPointIndex, int32 PointIndex)>& Body
+	)
+	{
+		const bool bHasFusedPoints = HasFusedPoints(BranchPointsAttribute, BranchParentNumberAttribute, BranchNumberAttribute);
+		for (int32 BranchIndex = 0; BranchIndex < BranchPointsAttribute.Num(); ++BranchIndex)
+		{
+			const bool bSkipFirstPoint = bHasFusedPoints 
+				&& GetBranchParentIndex(BranchParentNumberAttribute, BranchNumberAttribute, BranchIndex) != INDEX_NONE;
+
+			const int32 StartBranchPointIndex = bSkipFirstPoint ? 1 : 0;
+			const TArray<int32>& PointIndices = BranchPointsAttribute[BranchIndex];
+			for (int32 BranchPointIndex = StartBranchPointIndex; BranchPointIndex < PointIndices.Num(); ++BranchPointIndex)
+			{
+				if (Body(BranchIndex, BranchPointIndex, PointIndices[BranchPointIndex]) == EForEachResult::Break)
+				{
+					return;
+				}
+			}
+		}
+	}
+
+	inline void ForEachPlant(
+		PV::FBranchPlantNumberAttributeConstView BranchPlantNumberAttribute,
+		const TFunction<EForEachResult(int32 PlantNumber)>& Body
+	)
+	{
+		for (int32 PlantNumber : GetUniquePlantNumbers(BranchPlantNumberAttribute))
+		{
+			if (Body(PlantNumber) == EForEachResult::Break)
+			{
+				return;
+			}
+		}
+	}
+
+	inline void ForEachPlant(
+		const FManagedArrayCollection& Collection,
+		const TFunction<EForEachResult(int32 PlantNumber)>& Body
+	)
+	{
+		if (auto BranchPlantNumberAttribute = PV::FBranchPlantNumberAttribute::FindAttribute(Collection))
+		{
+			ForEachPlant(BranchPlantNumberAttribute, Body);
+		}
+	}
+
+	inline void ForEachPlantPoint(
+		PV::FBranchPlantNumberAttributeConstView BranchPlantNumberAttribute,
+		PV::FBranchPointsAttributeConstView BranchPointsAttribute,
+		int32 PlantNumber,
+		const TFunction<EForEachResult(int32 BranchIndex, int32 PointIndex)>& Body
+	)
+	{
+		const TArray<int32> BranchIndices = GetPlantBranchIndices(BranchPlantNumberAttribute, PlantNumber);
+		for (int32 BranchIndex : BranchIndices)
+		{
+			for (int32 PointIndex : BranchPointsAttribute[BranchIndex])
+			{
+				if (Body(BranchIndex, PointIndex) == EForEachResult::Break)
+				{
+					return;
+				}
+			}
+		}
+	}
+
+	inline void ForEachPlantPoint(
+		const FManagedArrayCollection& Collection,
+		int32 PlantNumber,
+		const TFunction<EForEachResult(int32 BranchIndex, int32 PointIndex)>& Body
+	)
+	{
+		const PV::FBranchPlantNumberAttributeConstView BranchPlantNumberAttribute = FBranchPlantNumberAttribute::FindAttribute(Collection);
+		const PV::FBranchPointsAttributeConstView BranchPointsAttribute = FBranchPointsAttribute::FindAttribute(Collection);
+
+		if (BranchPlantNumberAttribute && BranchPointsAttribute)
+		{
+			ForEachPlantPoint(BranchPlantNumberAttribute, BranchPointsAttribute, PlantNumber, Body);
+		}
+	}
+}
\ No newline at end of file
diff --git a/Engine/Plugins/Experimental/ProceduralVegetationEditor/Source/ProceduralVegetation/Private/Helpers/PVTransformHelper.cpp b/Engine/Plugins/Experimental/ProceduralVegetationEditor/Source/ProceduralVegetation/Private/Helpers/PVTransformHelper.cpp
--- a/Engine/Plugins/Experimental/ProceduralVegetationEditor/Source/ProceduralVegetation/Private/Helpers/PVTransformHelper.cpp
+++ b/Engine/Plugins/Experimental/ProceduralVegetationEditor/Source/ProceduralVegetation/Private/Helpers/PVTransformHelper.cpp
@@ -1,85 +1,6 @@
 // Copyright Epic Games, Inc. All Rights Reserved.
 
 #include "Helpers/PVTransformHelper.h"
-#include "Helpers/PVAttributesHelper.h"
-
-bool PV::Transform::RecursiveRotateBranchPoints_SplitPoints(
-	const FRotateBranchPointsAttributeCollection& Attributes,
-	int32 BranchIndex,
-	const FRotateBranchPointsFunction& RotatePointsBody,
-	const FVector3f& TrunkPivot
-)
-{
-	using namespace PV::AttributesHelper;
-
-	if (!Attributes.IsValid())
-	{
-		return false;
-	}
-
-	struct FPointState
-	{
-		int32     ParentPointIndex = INDEX_NONE;
-		FVector3f ParentPivotPoint = FVector3f::ZeroVector;
-		FVector3f ParentIdentityPosition = FVector3f::ZeroVector;
-		FMatrix44f TransformMatrix = FMatrix44f::Identity;
-	};
-	RecursiveWalkBranchPoints_SplitPoints<FPointState, false>(
-		Attributes.BranchChildren,
-		Attributes.BranchParentNumber,
-		Attributes.BranchNumber,
-		Attributes.BranchPoints,
-		Attributes.PointLengthFromRoot,
-		BranchIndex,
-		FPointState(),
-		[&](int32 BranchIndex, int32 BranchPointIndex, int32 PointIndex, const TArray<int32>& ChildBranchIndices, FPointState& PointState)
-		{
-			const FRotateBranchPointParams Params = FRotateBranchPointParams(
-				BranchIndex,
-				BranchPointIndex,
-				PointIndex,
-				PointState.ParentPointIndex,
-				ChildBranchIndices
-			);
-			FRotateBranchPointResult Result = {
-				.Rotation = FQuat4f::Identity,
-				.bAbsoluteRotation = false,
-			};
-			RotatePointsBody(Params, Result);
-
-			const FVector3f& PointPosition = Attributes.PointPosition[PointIndex];
-			const FVector3f RelativePosition = PointState.ParentPointIndex == INDEX_NONE
-				? TrunkPivot
-				: PointPosition - PointState.ParentIdentityPosition;
-
-			const FMatrix44f LocalRotationMatrix = Result.Rotation.ToMatrix();
-			const FMatrix44f LocalTranslationMatrix = FTranslationMatrix44f(RelativePosition);
-			const FMatrix44f LocalTransformMatrix = LocalTranslationMatrix * LocalRotationMatrix;
-
-			if (Result.bAbsoluteRotation)
-			{
-				const FMatrix44f TranslationOnly = FTranslationMatrix44f(PointState.TransformMatrix.GetOrigin() + LocalTransformMatrix.GetOrigin());
-				PointState.TransformMatrix = LocalRotationMatrix * TranslationOnly;
-			}
-			else
-			{
-				PointState.TransformMatrix = LocalTransformMatrix * PointState.TransformMatrix;
-			}
-
-			PointState.ParentIdentityPosition = PointPosition;
-
-			const FVector3f PointTranslation = PointState.TransformMatrix.GetOrigin() - PointPosition;
-			const FQuat4f PointRotation = PointState.TransformMatrix.ToQuat();
-			TransformPoint(Attributes, BranchIndex, PointIndex, PointTranslation, PointRotation);
-
-			PointState.ParentPointIndex = PointIndex;
-
-			return EForEachResult::Continue;
-		}
-	);
-
-	return true;
-}
 
 void PV::Transform::TransformPoint(
 	const FTransformPointAttributeCollection& Attributes,
@@ -111,7 +32,7 @@ void PV::Transform::TransformPoint(
 		&& Attributes.OptionalFoliageLengthFromRoot.IsValid();
 
 	const TArray<int32> BranchSegmentFoliageEntryIds = bUpdateFoliateRotation
-		? PV::AttributesHelper::GetBranchPointFoliageEntryIds(
+		? PV::PlantTraversalHelper::GetBranchPointFoliageEntryIds(
 			Attributes.OptionalBranchPoints,
 			Attributes.OptionalPointLengthFromRoot,
 			Attributes.OptionalBranchFoliageIDs,
diff --git a/Engine/Plugins/Experimental/ProceduralVegetationEditor/Source/ProceduralVegetation/Private/Helpers/PVTransformHelper.h b/Engine/Plugins/Experimental/ProceduralVegetationEditor/Source/ProceduralVegetation/Private/Helpers/PVTransformHelper.h
--- a/Engine/Plugins/Experimental/ProceduralVegetationEditor/Source/ProceduralVegetation/Private/Helpers/PVTransformHelper.h
+++ b/Engine/Plugins/Experimental/ProceduralVegetationEditor/Source/ProceduralVegetation/Private/Helpers/PVTransformHelper.h
@@ -2,43 +2,10 @@
 #pragma once
 
 #include "Utils/PVAttributes.h"
+#include "Helpers/PVPlantTraversalHelper.h"
 
 namespace PV::Transform
 {
-
-struct FRotateBranchPointParams
-{
-	int32 BranchIndex;
-	int32 BranchPointIndex; // The index of the point in the branch points attribute
-	int32 PointIndex;		// The global point index of this point
-	int32 ParentPointIndex; // INDEX_NONE for the first point on the trunk
-	const TArray<int32>& PointChildBranchIndices;
-
-	FRotateBranchPointParams() = delete;
-	FRotateBranchPointParams(
-		int32 InBranchIndex,
-		int32 InBranchPointIndex,
-		int32 InPointIndex,
-		int32 InParentPointIndex,
-		const TArray<int32>& InPointChildBranchIndices
-	)
-		: BranchIndex(InBranchIndex)
-		, BranchPointIndex(InBranchPointIndex)
-		, PointIndex(InPointIndex)
-		, ParentPointIndex(InParentPointIndex)
-		, PointChildBranchIndices(InPointChildBranchIndices)
-	{
-	}
-};
-
-struct FRotateBranchPointResult
-{
-	FQuat4f Rotation;
-	bool bAbsoluteRotation;  // If true, then the rotation of the parent point will be ignored when rotating this point.
-};
-
-using FRotateBranchPointsFunction = TFunction<void(const FRotateBranchPointParams& Params, FRotateBranchPointResult& Result)>;
-
 /*
 * Attributes required to traverse a tree and rotate the points around its parent point.
 * Attributes named as Optional are not required to peform this operation, but if present will
@@ -89,38 +56,6 @@
 	}
 };
 
-/**
- * Recursively rotates all points in the branch hierarchy starting from the specified BranchIndex.
- * 
- * Parent rotations are multiplied into child transforms ensuring that child points inherit all 
- * parent rotations (unless bAbsoluteRotation is set to true). The behavior mirrors bone 
- * transform accumulation in a skeletal hierarchy.
- *
- * This function performs a depth-first traversal of the branch hierarchy:
- *  1. It walks all points in the current branch in order.
- *  2. For each point, it invokes RotatePointsBody, allowing the caller to define
- *     a local rotation for that point.
- *  3. The computed rotation is accumulated with the parent’s transform, ensuring
- *     that rotations propagate hierarchically (similar to skeletal animation).
- *     Each point is effectively rotated around its parent pivot (or TrunkPivot
- *     for the root).
- *  4. After finishing all points in the current branch, the function recursively
- *     processes any axillary (child) branches originating from points in this branch.
- *
- * @param Attributes       The attributes requried to traverse the tree and update the point rotations.
- * @param BranchIndex      Index of the branch to start traversal from.
- * @param RotatePointsBody Callback invoked for each point to compute its rotation.
- * @param TrunkPivot       Pivot used for the root branch.
- *
- * @return true if the traversal and rotation succeeded; false otherwise.
- */
-PROCEDURALVEGETATION_API extern bool RecursiveRotateBranchPoints_SplitPoints(
-	const FRotateBranchPointsAttributeCollection& Attributes,
-	int32 BranchIndex,
-	const FRotateBranchPointsFunction& RotatePointsBody,
-	const FVector3f& TrunkPivot = FVector3f::ZeroVector
-);
-
 struct FTransformPointAttributeCollection
 {
 	PV::FPointPositionAttributeView PointPosition;
@@ -174,4 +109,152 @@
 	const FQuat4f& Rotation
 );
 
+struct FRotateBranchPointParams
+{
+	int32 BranchIndex;
+	int32 BranchPointIndex; // The index of the point in the branch points attribute
+	int32 PointIndex;		// The global point index of this point
+	int32 ParentPointIndex; // INDEX_NONE for the first point on the trunk
+	const TArray<int32>& ChildBranchIndices;
+	const TArray<int32>& ChildPointIndices;
 };
+
+struct FRotateBranchPointResult
+{
+	FQuat4f Rotation;
+	bool bAbsoluteRotation;  // If true, then the rotation of the parent point will be ignored when rotating this point.
+};
+
+/**
+ * Recursively rotates all points in the branch hierarchy starting from the specified BranchIndex.
+ * 
+ * Parent rotations are multiplied into child transforms ensuring that child points inherit all 
+ * parent rotations (unless bAbsoluteRotation is set to true). The behavior mirrors bone 
+ * transform accumulation in a skeletal hierarchy.
+ *
+ * This function performs a depth-first traversal of the branch hierarchy:
+ *  1. It walks all points in the current branch in order.
+ *  2. For each point, it invokes RotatePointsBody, allowing the caller to define
+ *     a local rotation for that point.
+ *  3. The computed rotation is accumulated with the parent’s transform, ensuring
+ *     that rotations propagate hierarchically (similar to skeletal animation).
+ *     Each point is effectively rotated around its parent pivot (or TrunkPivot
+ *     for the root).
+ *  4. After finishing all points in the current branch, the function recursively
+ *     processes any axillary (child) branches originating from points in this branch.
+ *
+ * @param Attributes       The attributes requried to traverse the tree and update the point rotations.
+ * @param BranchIndex      Index of the branch to start traversal from.
+ * @param RotatePointsBody Callback invoked for each point to compute its rotation.
+ * @param TrunkPivot       Pivot used for the root branch.
+ *
+ * @return true if the traversal and rotation succeeded; false otherwise.
+ */
+template<typename StateType>
+inline bool RecursiveRotateBranchPoints(
+	const FRotateBranchPointsAttributeCollection& Attributes,
+	int32 BranchIndex,
+	const StateType& RootPointState,
+	const TFunction<void(const FRotateBranchPointParams& Params, FRotateBranchPointResult& Result, StateType& State)>& RotatePointsBody,
+	const FVector3f& TrunkPivot = FVector3f::ZeroVector
+)
+{
+	using namespace PV::PlantTraversalHelper;
+
+	if (!Attributes.IsValid())
+	{
+		return false;
+	}
+
+	struct FPointState
+	{
+		int32      ParentPointIndex;
+		FVector3f  ParentPivotPoint;
+		FVector3f  ParentIdentityPosition;
+		FMatrix44f TransformMatrix;
+		StateType  UserState;
+	};
+
+	const FPointState DefaultPointState = {
+		INDEX_NONE,
+		FVector3f::ZeroVector,
+		FVector3f::ZeroVector,
+		FMatrix44f::Identity,
+		RootPointState
+	};
+
+	RecursiveWalkBranchPoints<FPointState, false>(
+		Attributes.BranchChildren,
+		Attributes.BranchParentNumber,
+		Attributes.BranchNumber,
+		Attributes.BranchPoints,
+		Attributes.PointLengthFromRoot,
+		BranchIndex,
+		DefaultPointState,
+		[&](const FRecursiveWalkBranchPointsParams& InParams, FPointState& PointState)
+		{
+			const FRotateBranchPointParams Params = FRotateBranchPointParams(
+				InParams.BranchIndex,
+				InParams.BranchPointIndex,
+				InParams.PointIndex,
+				PointState.ParentPointIndex,
+				InParams.ChildBranchIndices,
+				InParams.ChildPointIndices
+			);
+			FRotateBranchPointResult Result = {
+				.Rotation = FQuat4f::Identity,
+				.bAbsoluteRotation = false,
+			};
+			RotatePointsBody(Params, Result, PointState.UserState);
+
+			const FVector3f& PointPosition = Attributes.PointPosition[InParams.PointIndex];
+			const FVector3f RelativePosition = PointState.ParentPointIndex == INDEX_NONE
+				? TrunkPivot
+				: PointPosition - PointState.ParentIdentityPosition;
+
+			const FMatrix44f LocalRotationMatrix = Result.Rotation.ToMatrix();
+			const FMatrix44f LocalTranslationMatrix = FTranslationMatrix44f(RelativePosition);
+			const FMatrix44f LocalTransformMatrix = LocalTranslationMatrix * LocalRotationMatrix;
+
+			if (Result.bAbsoluteRotation)
+			{
+				const FMatrix44f TranslationOnly = FTranslationMatrix44f(PointState.TransformMatrix.GetOrigin() + LocalTransformMatrix.GetOrigin());
+				PointState.TransformMatrix = LocalRotationMatrix * TranslationOnly;
+			}
+			else
+			{
+				PointState.TransformMatrix = LocalTransformMatrix * PointState.TransformMatrix;
+			}
+
+			PointState.ParentIdentityPosition = PointPosition;
+
+			const FVector3f PointTranslation = PointState.TransformMatrix.GetOrigin() - PointPosition;
+			const FQuat4f PointRotation = PointState.TransformMatrix.ToQuat();
+			TransformPoint(Attributes, InParams.BranchIndex, InParams.PointIndex, PointTranslation, PointRotation);
+
+			PointState.ParentPointIndex = InParams.PointIndex;
+
+			return EForEachResult::Continue;
+		}
+	);
+
+	return true;
+}
+
+inline bool RecursiveRotateBranchPoints(
+	const FRotateBranchPointsAttributeCollection& Attributes,
+	int32 BranchIndex,
+	const TFunction<void(const FRotateBranchPointParams& Params, FRotateBranchPointResult& Result)>& RotatePointsBody,
+	const FVector3f& TrunkPivot = FVector3f::ZeroVector
+)
+{
+	return RecursiveRotateBranchPoints<int32>(
+		Attributes, 
+		BranchIndex, 
+		0,
+		[&](const FRotateBranchPointParams& Params, FRotateBranchPointResult& Result, int32&) { RotatePointsBody(Params, Result); },
+		TrunkPivot
+	);
+}
+
+};
diff --git a/Engine/Plugins/Experimental/ProceduralVegetationEditor/Source/ProceduralVegetation/Private/Implementations/PVGrafter.cpp b/Engine/Plugins/Experimental/ProceduralVegetationEditor/Source/ProceduralVegetation/Private/Implementations/PVGrafter.cpp
--- a/Engine/Plugins/Experimental/ProceduralVegetationEditor/Source/ProceduralVegetation/Private/Implementations/PVGrafter.cpp
+++ b/Engine/Plugins/Experimental/ProceduralVegetationEditor/Source/ProceduralVegetation/Private/Implementations/PVGrafter.cpp
@@ -10,7 +10,7 @@
 #include "Facades/PVBranchFacade.h"
 #include "Facades/PVBudVectorsFacade.h"
 #include "Facades/PVPointFacade.h"
-#include "Helpers/PVAttributesHelper.h"
+#include "Helpers/PVPlantTraversalHelper.h"
 #include "Helpers/PVFoliageJSONHelper.h"
 #include "Helpers/PVUtilities.h"
 
@@ -62,7 +62,7 @@ void FPVGrafter::ComputeBranchGradientAttribute(
 	PV::FBranchNumberAttributeConstView BranchNumberAttribute,
 	PV::FBranchPointsAttributeConstView BranchPointsAttribute)
 {
-	using namespace PV::AttributesHelper;
+	using namespace PV::PlantTraversalHelper;
 
 	RecursiveWalkBranches(
 		BranchParentNumberAttribute,
diff --git a/Engine/Plugins/Experimental/ProceduralVegetationEditor/Source/ProceduralVegetation/Private/Implementations/PVImporter_Texture2D.cpp b/Engine/Plugins/Experimental/ProceduralVegetationEditor/Source/ProceduralVegetation/Private/Implementations/PVImporter_Texture2D.cpp
--- a/Engine/Plugins/Experimental/ProceduralVegetationEditor/Source/ProceduralVegetation/Private/Implementations/PVImporter_Texture2D.cpp
+++ b/Engine/Plugins/Experimental/ProceduralVegetationEditor/Source/ProceduralVegetation/Private/Implementations/PVImporter_Texture2D.cpp
@@ -7,7 +7,7 @@
 #include "IntVectorTypes.h"
 #include "PCGSettings.h"
 #include "GeometryCollection/ManagedArrayCollection.h"
-#include "Helpers/PVAttributesHelper.h"
+#include "Helpers/PVPlantTraversalHelper.h"
 
 #include "Engine/Texture.h"
 #include "Engine/Texture2D.h"
@@ -260,7 +260,7 @@ namespace Internal
 		const FVector& InOffset = FVector::ZeroVector
 	)
 	{
-		using namespace PV::AttributesHelper;
+		using namespace PV::PlantTraversalHelper;
 
 		TArray<TTuple<FVector3f, FString, float>> Labels;
 		Labels.Reserve(PlantNumberToPlantIndices.Num());
diff --git a/Engine/Plugins/Experimental/ProceduralVegetationEditor/Source/ProceduralVegetation/Private/Implementations/PVMeshBuilder.cpp b/Engine/Plugins/Experimental/ProceduralVegetationEditor/Source/ProceduralVegetation/Private/Implementations/PVMeshBuilder.cpp
--- a/Engine/Plugins/Experimental/ProceduralVegetationEditor/Source/ProceduralVegetation/Private/Implementations/PVMeshBuilder.cpp
+++ b/Engine/Plugins/Experimental/ProceduralVegetationEditor/Source/ProceduralVegetation/Private/Implementations/PVMeshBuilder.cpp
@@ -66,43 +66,13 @@ void FPVMeshBuilder::GenerateGeometryCollection(const FManagedArrayCollection& I
 
 	RemoveNonBudPoints(BranchFacade, PointFacade);
 	
-	auto PointNjordPixelIndexAttribute = PV::FPointNjordPixelIndexAttribute::GetAttribute(OutGeometryCollection);
-	auto BranchPointsAttribute = PV::FBranchPointsAttribute::GetAttribute(OutGeometryCollection);
-	auto BranchParentNumberAttribute = PV::FBranchParentNumberAttribute::GetAttribute(OutGeometryCollection);
-	auto BranchNumberAttribute = PV::FBranchNumberAttribute::GetAttribute(OutGeometryCollection);
-	auto BranchChildrenAttribute = PV::FBranchChildrenAttribute::GetAttribute(OutGeometryCollection);
-
 	BranchFacade.RecomputeBranchChildren();
 	
-	PV::AttributesHelper::ComputeNjordPixelIndex(
-		PointNjordPixelIndexAttribute,
-		BranchPointsAttribute,
-		BranchParentNumberAttribute, 
-		BranchNumberAttribute,
-		BranchChildrenAttribute
-	);
+	PV::AttributesHelper::ComputeNjordPixelIndex(OutGeometryCollection);
 	
 	AddPointsToSkeleton(OutGeometryCollection);
 	
-	auto GroundGradientAttribute = PV::FPointGroundGradientAttribute::GetAttribute(OutGeometryCollection);
-	auto PointScaleGradientAttribute = PV::FPointScaleGradientAttribute::GetAttribute(OutGeometryCollection);
-	auto PointPlantGradientAttribute = PV::FPointPlantGradientAttribute::GetAttribute(OutGeometryCollection);
-	auto PointBranchGradientAttribute = PV::FPointBranchGradientAttribute::GetAttribute(OutGeometryCollection);
-	auto PointLengthFromRootAttribute = PV::FPointLengthFromRootAttribute::GetAttribute(OutGeometryCollection);
-	auto PointScaleAttribute = PV::FPointScaleAttribute::GetAttribute(OutGeometryCollection);
-	
-	PV::AttributesHelper::ComputeGradients(
-		GroundGradientAttribute,
-		PointScaleGradientAttribute,
-		PointPlantGradientAttribute,
-		PointBranchGradientAttribute,
-		PointLengthFromRootAttribute,
-		PointScaleAttribute,
-		BranchPointsAttribute,
-		BranchParentNumberAttribute,
-		BranchNumberAttribute,
-		BranchChildrenAttribute
-	);
+	PV::AttributesHelper::ComputeGradients(OutGeometryCollection);
 
 	MeshBuilderParams.MaterialSettings.ApplyMaterialSettings(OutGeometryCollection);
 	GetLocalDynamicMeshData(OutGeometryCollection, MeshBuilderParams, MeshesDataCollection);
diff --git a/Engine/Plugins/Experimental/ProceduralVegetationEditor/Source/ProceduralVegetation/Private/Implementations/PVRotateAxillaryBranches.cpp b/Engine/Plugins/Experimental/ProceduralVegetationEditor/Source/ProceduralVegetation/Private/Implementations/PVRotateAxillaryBranches.cpp
--- a/Engine/Plugins/Experimental/ProceduralVegetationEditor/Source/ProceduralVegetation/Private/Implementations/PVRotateAxillaryBranches.cpp
+++ b/Engine/Plugins/Experimental/ProceduralVegetationEditor/Source/ProceduralVegetation/Private/Implementations/PVRotateAxillaryBranches.cpp
@@ -2,7 +2,7 @@
 
 #include "PVRotateAxillaryBranches.h"
 #include "Helpers/PVTransformHelper.h"
-#include "Helpers/PVAttributesHelper.h"
+#include "Helpers/PVPlantTraversalHelper.h"
 
 void FPVRotateAxillaryBranches::ApplyRotateAxillaryBranches(const FPVRotateAxillaryBranchesParams& InRotateAxillaryBranchesParams, FManagedArrayCollection& OutCollection)
 {
@@ -36,7 +36,7 @@ void FPVRotateAxillaryBranches::ApplyRotateAxillaryBranches(const FPVRotateAxill
 	BranchesLocalBranchGenerations.SetNumZeroed(Attributes.BranchNumber.Num());
 
 	float RotationSign = 1;
-	for (const int32 TrunkIndex : PV::AttributesHelper::GetTrunkIndices(Attributes.BranchParentNumber))
+	for (const int32 TrunkIndex : PV::PlantTraversalHelper::GetTrunkIndices(Attributes.BranchParentNumber))
 	{
 		const TArray<int32>& CurrentBranchPoints = Attributes.BranchPoints[TrunkIndex];
 		if (CurrentBranchPoints.Num() == 0)
@@ -44,8 +44,13 @@ void FPVRotateAxillaryBranches::ApplyRotateAxillaryBranches(const FPVRotateAxill
 			continue;
 		}
 
+		struct FPointState
+		{
+			int32 PrevBranchIndex = INDEX_NONE;
+		};
+
 		FVector PivotLocation = FVector::ZeroVector;
-		const FRotateBranchPointsFunction CalcRotationFunction = [&](const FRotateBranchPointParams& Params, FRotateBranchPointResult& Result)
+		const auto CalcRotationFunction = [&](const FRotateBranchPointParams& Params, FRotateBranchPointResult& Result, FPointState& PointState)
 		{
 			if (Params.BranchIndex != CurrentBranchIndex)
 			{
@@ -53,17 +58,17 @@ void FPVRotateAxillaryBranches::ApplyRotateAxillaryBranches(const FPVRotateAxill
 				CurrentBranchIndex = Params.BranchIndex;
 			}
 
-			if (Params.PointChildBranchIndices.Num() > 0)
+			if (Params.ChildBranchIndices.Num() > 0)
 			{
 				CurrentLocalBranchGeneration += 1;
 
-				for (int32 ChildBranchIndex : Params.PointChildBranchIndices)
+				for (int32 ChildBranchIndex : Params.ChildBranchIndices)
 				{
 					BranchesLocalBranchGenerations[ChildBranchIndex] = CurrentLocalBranchGeneration;
 				}
 			}
 
-			if (Params.BranchPointIndex == 0 && Params.BranchIndex != TrunkIndex)
+			if (PointState.PrevBranchIndex != Params.BranchIndex && Params.BranchIndex != TrunkIndex)
 			{
 				const int32 Generation = BudDevelopmentAttribute[Params.PointIndex].Generation;
 				if (Generation < StartGeneration || Generation >= EndGeneration)
@@ -85,8 +90,10 @@ void FPVRotateAxillaryBranches::ApplyRotateAxillaryBranches(const FPVRotateAxill
 					RotationSign *= -1;
 				}
 			}
+
+			PointState.PrevBranchIndex = Params.BranchIndex;
 		};
 
-		RecursiveRotateBranchPoints_SplitPoints(Attributes, TrunkIndex, CalcRotationFunction);
+		RecursiveRotateBranchPoints<FPointState>(Attributes, TrunkIndex, FPointState(), CalcRotationFunction);
 	}
 }
\ No newline at end of file
diff --git a/Engine/Plugins/Experimental/ProceduralVegetationEditor/Source/ProceduralVegetation/Private/Implementations/PVSlope.cpp b/Engine/Plugins/Experimental/ProceduralVegetationEditor/Source/ProceduralVegetation/Private/Implementations/PVSlope.cpp
--- a/Engine/Plugins/Experimental/ProceduralVegetationEditor/Source/ProceduralVegetation/Private/Implementations/PVSlope.cpp
+++ b/Engine/Plugins/Experimental/ProceduralVegetationEditor/Source/ProceduralVegetation/Private/Implementations/PVSlope.cpp
@@ -2,7 +2,7 @@
 
 #include "PVSlope.h"
 #include "Helpers/PVTransformHelper.h"
-#include "Helpers/PVAttributesHelper.h"
+#include "Helpers/PVPlantTraversalHelper.h"
 
 namespace PVSlope
 {
@@ -25,7 +25,7 @@ namespace PVSlope
 		PV::FPointLengthFromRootAttributeConstView PointLengthFromRootAttribute,
 		int32 TrunkIndex)
 	{
-		using namespace PV::AttributesHelper;
+		using namespace PV::PlantTraversalHelper;
 
 		const TArray<int32>& CurrentBranchPoints = BranchPointsAttribute[TrunkIndex];
 		if (CurrentBranchPoints.Num() == 0)
@@ -36,16 +36,16 @@ namespace PVSlope
 		const FVector3f& RootPosition = PointPositionAttribute[CurrentBranchPoints[0]];
 
 		float BranchLengthSquared = 0;
-		RecursiveWalkBranchPoints_SplitPoints(
+		RecursiveWalkBranchPoints(
 			BranchChildrenAttribute,
 			BranchParentNumberAttribute,
 			BranchNumberAttribute,
 			BranchPointsAttribute,
 			PointLengthFromRootAttribute,
 			TrunkIndex,
-			[&](int32 BranchIndex, int32 BranchPointIndex, int32 PointIndex, const TArray<int32>& ChildPointIndices)->EForEachResult
+			[&](const FRecursiveWalkBranchPointsParams& Params)->EForEachResult
 			{
-				const FVector3f& PointPosition = PointPositionAttribute[PointIndex];
+				const FVector3f& PointPosition = PointPositionAttribute[Params.PointIndex];
 				const float PointDistSquared = (PointPosition - RootPosition).SizeSquared();
 				if (PointDistSquared > BranchLengthSquared)
 				{
@@ -69,7 +69,7 @@ void FPVSlope::ApplySlope(const FPVSlopeParams& InSlopeParams, FManagedArrayColl
 		return;
 	}
 
-	for (const int32 TrunkIndex : PV::AttributesHelper::GetTrunkIndices(Attributes.BranchParentNumber))
+	for (const int32 TrunkIndex : PV::PlantTraversalHelper::GetTrunkIndices(Attributes.BranchParentNumber))
 	{
 		const TArray<int32>& CurrentBranchPoints = Attributes.BranchPoints[TrunkIndex];
 		if (CurrentBranchPoints.Num() == 0)
@@ -88,7 +88,7 @@ void FPVSlope::ApplySlope(const FPVSlopeParams& InSlopeParams, FManagedArrayColl
 		);
 		const FVector3f TrunkPosition = Attributes.PointPosition[CurrentBranchPoints[0]];
 		const FVector3f TrunkPivot = InSlopeParams.TrunkPivotPoint == EPVSlopeTrunkPivotPoint::Origin ? FVector3f::ZeroVector : TrunkPosition;
-		const FRotateBranchPointsFunction CalcRotationFunction = [&](const FRotateBranchPointParams& Params, FRotateBranchPointResult& Result)
+		const auto CalcRotationFunction = [&](const FRotateBranchPointParams& Params, FRotateBranchPointResult& Result)
 		{
 			const FVector3f Position = Attributes.PointPosition[Params.PointIndex];
 			const float TreeAngleAlpha = 1.0 - (TreeHeight > 0 ? ((Position - TrunkPosition).Size() / TreeHeight) : 0);
@@ -98,6 +98,6 @@ void FPVSlope::ApplySlope(const FPVSlopeParams& InSlopeParams, FManagedArrayColl
 			Result.bAbsoluteRotation = true;
 		};
 
-		RecursiveRotateBranchPoints_SplitPoints(Attributes, TrunkIndex, CalcRotationFunction, TrunkPivot);
+		RecursiveRotateBranchPoints(Attributes, TrunkIndex, CalcRotationFunction, TrunkPivot);
 	}
 }
diff --git a/Engine/Plugins/Experimental/ProceduralVegetationEditor/Source/ProceduralVegetation/Private/Utils/PVAttributes.h b/Engine/Plugins/Experimental/ProceduralVegetationEditor/Source/ProceduralVegetation/Private/Utils/PVAttributes.h
--- a/Engine/Plugins/Experimental/ProceduralVegetationEditor/Source/ProceduralVegetation/Private/Utils/PVAttributes.h
+++ b/Engine/Plugins/Experimental/ProceduralVegetationEditor/Source/ProceduralVegetation/Private/Utils/PVAttributes.h
@@ -4,6 +4,7 @@
 #include "CoreMinimal.h"
 
 #include "GeometryCollection/ManagedArrayCollection.h"
+#include "GeometryCollection/ManagedArrayAccessor.h"
 
 #include "Facades/PVAttributesNames.h"
 
@@ -61,7 +62,8 @@ public:
 		InElementConstViewType,
 		InElementViewType>;
 
-	using NonConstAttributeViewType = TAttributeView<ElementType, InElementViewType, InElementConstViewType, InAttributeNameStr, InGroupNameStr>;
+	using MutableAttributeViewType = TAttributeView<ElementType, InElementViewType, InElementConstViewType, InAttributeNameStr, InGroupNameStr>;
+	using ConstAttributeViewType = TAttributeView<const ElementType, InElementViewType, InElementConstViewType, InAttributeNameStr, InGroupNameStr>;
 
 	using SizeType = TArray<ElementType>::SizeType;
 
@@ -78,10 +80,17 @@ public:
 		: ManagedArray(InManagedArray)
 	{}
 
-	TAttributeView(const NonConstAttributeViewType& Other)
+	TAttributeView(const MutableAttributeViewType& Other)
 		: ManagedArray(Other.ManagedArray)
 	{}
 
+	template<typename U = InElementType,
+		typename = std::enable_if_t<std::is_const_v<U>>>
+	TAttributeView(const ConstAttributeViewType& Other)
+		: ManagedArray(Other.ManagedArray)
+	{
+	}
+
 	bool IsValid() const
 	{
 		return ManagedArray != nullptr;
@@ -142,6 +151,59 @@ public:
 	FORCEINLINE AttributeRangedForIteratorType end() const { check(ManagedArray); return ManagedArray->end(); }
 };
 
+template<
+	typename InElementType,
+	typename InElementViewType,
+	typename InElementConstViewType,
+	const TCHAR* InAttributeNameStr,
+	const TCHAR* InGroupNameStr
+>
+class TAttributeAccessor : public TManagedArrayAccessor<InElementType>
+{
+	static_assert(!std::is_const_v<InElementType>, "InElementType cannot be const");
+
+public:
+	using ElementType = InElementType;
+	using ManagedArrayType = TManagedArray<ElementType>;
+	using AttributeView = TAttributeView<ElementType, InElementViewType, InElementConstViewType, InAttributeNameStr, InGroupNameStr>;
+	using AttributeConstView = TAttributeView<const ElementType, InElementViewType, InElementConstViewType, InAttributeNameStr, InGroupNameStr>;
+	using Super = TManagedArrayAccessor<ElementType>;
+	
+	inline static const FName AttributeName = FName(InAttributeNameStr);
+	inline static const FName GroupName = FName(InGroupNameStr);
+
+	TAttributeAccessor() = delete;
+
+	TAttributeAccessor(FManagedArrayCollection& InCollection)
+		: Super(InCollection, AttributeName, GroupName)
+	{}
+
+	TAttributeAccessor(const FManagedArrayCollection& InCollection)
+		: Super(InCollection, AttributeName, GroupName)
+	{}
+
+	AttributeView GetView()
+	{
+		return Super::IsValid() ? AttributeView(&Super::Get()) : AttributeView(nullptr);
+	}
+
+	AttributeConstView GetConstView() const
+	{
+		return Super::IsValid() ? AttributeConstView(&Super::Get()) : AttributeConstView(nullptr);
+	}
+
+	InElementViewType ModifyAt(int32 Index)
+	{
+		ManagedArrayType& MutableArray = Super::Modify();
+		return InElementViewType(MutableArray[Index]);
+	}
+
+	InElementConstViewType operator[](int32 Index) const
+	{
+		return InElementConstViewType(Super::operator[](Index));
+	}
+};
+
 template<
 	typename InAttributeType, 
 	typename InElementViewType,
@@ -155,11 +217,13 @@ class TAttributeDefinition
 
 public:
 	using AttributeType = InAttributeType;
-	using AttributeView = TAttributeView<AttributeType, InElementViewType, InElementConstViewType, InAttributeNameStr, InGroupNameStr>;
-	using AttributeConstView = TAttributeView<const AttributeType, InElementViewType, InElementConstViewType, InAttributeNameStr, InGroupNameStr>;
+	using AttributeAccessor = TAttributeAccessor<AttributeType, InElementViewType, InElementConstViewType, InAttributeNameStr, InGroupNameStr>;
+	using AttributeView = AttributeAccessor::AttributeView;
+	using AttributeConstView = AttributeAccessor::AttributeConstView;
 
 	inline static const FName AttributeName = FName(InAttributeNameStr);
 	inline static const FName GroupName = FName(InGroupNameStr);
+	inline static const FManagedArrayCollection::FManagedType ManagedType = FManagedArrayCollection::TManagedType<AttributeType>(AttributeName, GroupName);
 
 	/**
 	* Tries to find the attribute in the supplied collection. Returns a null view if attribute is not in collection.
@@ -238,6 +302,19 @@ public:
 	{
 		return InCollection.HasAttribute(AttributeName, GroupName);
 	}
+
+	/**
+	* Copy the attribute from one collection to another. Will perform an implicit group sync.
+	*/
+	static void CopyAttribute(const FManagedArrayCollection& SourceCollection, FManagedArrayCollection& DestCollection)
+	{
+		if (!SourceCollection.HasAttribute(AttributeName, GroupName))
+		{
+			return;
+		}
+
+		DestCollection.CopyAttribute(SourceCollection, AttributeName, GroupName);
+	}
 };
 
 // Helper class to allow for initialization, and selecton of, a fixed size array based of off default input array.
@@ -279,7 +356,8 @@ struct TFixedSizeArrayHelper
 #define PV_DEFINE_STRUCTURED_ATTRIBUTE(Name, Type, ElementViewType, ElementConstViewType, AttributeNameStr, AttributeGroupStr)          \
 	using F##Name##Attribute = TAttributeDefinition<Type, ElementViewType, ElementConstViewType, AttributeNameStr, AttributeGroupStr>;  \
 	using F##Name##AttributeView = F##Name##Attribute::AttributeView;                                                                   \
-	using F##Name##AttributeConstView = F##Name##Attribute::AttributeConstView;
+	using F##Name##AttributeConstView = F##Name##Attribute::AttributeConstView;                                                         \
+	using F##Name##AttributeAccessor = F##Name##Attribute::AttributeAccessor;
 
 #define PV_DEFINE_ATTRIBUTE(Name, Type, AttributeNameStr, AttributeGroupStr) \
 	PV_DEFINE_STRUCTURED_ATTRIBUTE(Name, Type, Type&, const Type&, AttributeNameStr, AttributeGroupStr)
diff --git a/Engine/Plugins/Experimental/ProceduralVegetationEditor/Source/ProceduralVegetation/Public/Helpers/PVAttributesHelper.h b/Engine/Plugins/Experimental/ProceduralVegetationEditor/Source/ProceduralVegetation/Public/Helpers/PVAttributesHelper.h
deleted file mode 100644
