Add debug nodes for importing growth data, force-recomputing all growth related attributes, and to diff two growth data inputs

- Add new debug node for importing growth data (PVGrowthDataJsonImporterSettings)
- Add new debug node for recomputing attributes (PVRecomputeGrowthDataAttributeSettings)
- Add new debug node for diffing two growth data inputs (PVGrowthDataDiffSettings)
- Move PVJSONHelper to private part of the module
- Add new function in PVJSONHelper which imports growth data only
- Move PVMockTreeImportSettings to private part of the module
- Fix BudHormoneLevels being incorrectly generated

#Synced-CL 51008641

diff --git a/Engine/Plugins/Experimental/ProceduralVegetationEditor/Source/ProceduralVegetation/Private/Facades/PVBoneFacade.cpp b/Engine/Plugins/Experimental/ProceduralVegetationEditor/Source/ProceduralVegetation/Private/Facades/PVBoneFacade.cpp
--- a/Engine/Plugins/Experimental/ProceduralVegetationEditor/Source/ProceduralVegetation/Private/Facades/PVBoneFacade.cpp
+++ b/Engine/Plugins/Experimental/ProceduralVegetationEditor/Source/ProceduralVegetation/Private/Facades/PVBoneFacade.cpp
@@ -253,7 +253,7 @@ namespace PV::Facades
 		FBranchFacade BranchFacade(VertexPointIds.GetConstCollection());
 		FPlantFacade PlantFacade(VertexPointIds.GetConstCollection());
 
-		if (!(Utilities::IsValidPVData(VertexPointIds.GetConstCollection()) && PlantFacade.IsValid()))
+		if (!(Utilities::IsValidGrowthData(VertexPointIds.GetConstCollection()) && PlantFacade.IsValid()))
 		{
 			return false;	
 		}
diff --git a/Engine/Plugins/Experimental/ProceduralVegetationEditor/Source/ProceduralVegetation/Private/Helpers/PVAttributesHelper.cpp b/Engine/Plugins/Experimental/ProceduralVegetationEditor/Source/ProceduralVegetation/Private/Helpers/PVAttributesHelper.cpp
--- a/Engine/Plugins/Experimental/ProceduralVegetationEditor/Source/ProceduralVegetation/Private/Helpers/PVAttributesHelper.cpp
+++ b/Engine/Plugins/Experimental/ProceduralVegetationEditor/Source/ProceduralVegetation/Private/Helpers/PVAttributesHelper.cpp
@@ -277,8 +277,8 @@ void PV::AttributesHelper::ComputeBudHormoneLevels(const FComputeBudHormoneLevel
 		const float Gradient = InAttributes.PointPlantGradient[PointIndex];
 		const float InverseGradient = 1.f - Gradient;
 
-		InAttributes.BudHormoneLevels[PointIndex].Apical = FMath::GetMappedRangeValueClamped(FVector2f(0.75, 1.0), FVector2f(0, 1), InverseGradient);
-		InAttributes.BudHormoneLevels[PointIndex].Axillary = 0;
+		InAttributes.BudHormoneLevels[PointIndex].Apical = 1;
+		InAttributes.BudHormoneLevels[PointIndex].Axillary = FMath::GetMappedRangeValueClamped(FVector2f(0, 1), FVector2f(0.35, 1.0), InverseGradient);
 		InAttributes.BudHormoneLevels[PointIndex].AxillaryInhibition = 0;
 		InAttributes.BudHormoneLevels[PointIndex].Radical = 0;
 		InAttributes.BudHormoneLevels[PointIndex].Ethylene = Gradient;
@@ -295,9 +295,20 @@ void PV::AttributesHelper::ComputeGradients(const FComputeGradientsAttributes& I
 		return;
 	}
 
+	ComputePointGroundGradient(InAttributes);
+	ComputePointScaleGradient(InAttributes);
+	ComputePointPlantAndBranchGradient(InAttributes);
+}
+
+void PV::AttributesHelper::ComputePointGroundGradient(const FComputePointGroundGradientAttributes& InAttributes)
+{
+	if (!InAttributes.IsValid())
+	{
+		return;
+	}
+
 	const float* MaxLengthFromRoot = Algo::MaxElement(InAttributes.PointLengthFromRoot);
-	const float* MaxPointScale = Algo::MaxElement(InAttributes.PointScale);
-	if (MaxLengthFromRoot == nullptr || MaxPointScale == nullptr || *MaxLengthFromRoot == 0 || *MaxPointScale == 0)
+	if (MaxLengthFromRoot == nullptr || *MaxLengthFromRoot == 0)
 	{
 		return;
 	}
@@ -307,11 +318,38 @@ void PV::AttributesHelper::ComputeGradients(const FComputeGradientsAttributes& I
 		const float LengthFromRoot = InAttributes.PointLengthFromRoot[PointIndex];
 		const float LengthFromRootAlpha = LengthFromRoot / *MaxLengthFromRoot;
 		InAttributes.PointGroundGradient[PointIndex] = 1.0 - LengthFromRootAlpha;
+	}
+}
 
+void PV::AttributesHelper::ComputePointScaleGradient(const FComputePointScaleGradientAttributes& InAttributes)
+{
+	if (!InAttributes.IsValid())
+	{
+		return;
+	}
+
+	const float* MaxPointScale = Algo::MaxElement(InAttributes.PointScale);
+	if (MaxPointScale == nullptr || *MaxPointScale == 0)
+	{
+		return;
+	}
+
+	for (int32 PointIndex = 0; PointIndex < InAttributes.PointScaleGradient.Num(); PointIndex++)
+	{
 		const float PointScale = InAttributes.PointScale[PointIndex];
 		const float PointScaleAlpha = PointScale / *MaxPointScale;
 		InAttributes.PointScaleGradient[PointIndex] = PointScaleAlpha;
 	}
+}
+
+void PV::AttributesHelper::ComputePointPlantAndBranchGradient(const FComputePointPlantAndBranchGradientAttributes& InAttributes)
+{
+	using namespace PV::PlantTraversalHelper;
+
+	if (!InAttributes.IsValid())
+	{
+		return;
+	}
 
 	// PlantGradient - A gradient from 1-0 where the tip of each branch is 0, and the base of each branch is whatever the gradient of the parent branch is at that point.
 	// BranchGradient - A gradient from 1-0 where 1 = base of the branch and 0 = tip of the branch
@@ -379,7 +417,11 @@ void PV::AttributesHelper::ComputeBudStatus(const FComputeBudStatusAttributes& I
 			else if (InAttributes.BranchPoints[BranchIndex][0] == PointIndex
 				&& GetBranchParentIndex(InAttributes.BranchParentNumber, InAttributes.BranchNumber, BranchIndex) == INDEX_NONE) // Is plant root point
 			{
+				BudStatusView.Axillary = 1;
 				BudStatusView.Seed = 1;
+				BudStatusView.Dormant = 1;
+				BudStatusView.Triggered = 0;
+				BudStatusView.NumTriggered = 0;
 			}
 			else
 			{
diff --git a/Engine/Plugins/Experimental/ProceduralVegetationEditor/Source/ProceduralVegetation/Private/Helpers/PVAttributesHelper.h b/Engine/Plugins/Experimental/ProceduralVegetationEditor/Source/ProceduralVegetation/Private/Helpers/PVAttributesHelper.h
--- a/Engine/Plugins/Experimental/ProceduralVegetationEditor/Source/ProceduralVegetation/Private/Helpers/PVAttributesHelper.h
+++ b/Engine/Plugins/Experimental/ProceduralVegetationEditor/Source/ProceduralVegetation/Private/Helpers/PVAttributesHelper.h
@@ -316,6 +316,124 @@ namespace PV::AttributesHelper
 	// Re-computes the point- ground gradient, scale gradient, plant gradient, and branch gradient.
 	PROCEDURALVEGETATION_API extern void ComputeGradients(const FComputeGradientsAttributes& InAttributes);
 
+	struct FComputePointGroundGradientAttributes
+	{
+		PV::FPointGroundGradientAttributeView PointGroundGradient;
+		PV::FPointLengthFromRootAttributeConstView PointLengthFromRoot;
+
+		FComputePointGroundGradientAttributes(const FComputeGradientsAttributes& InComputeGradientsAttributes)
+			: PointGroundGradient(InComputeGradientsAttributes.PointGroundGradient)
+			, PointLengthFromRoot(InComputeGradientsAttributes.PointLengthFromRoot)
+		{
+		}
+
+		FComputePointGroundGradientAttributes(
+			PV::FPointGroundGradientAttributeView InPointGroundGradient,
+			PV::FPointLengthFromRootAttributeConstView InPointLengthFromRoot)
+			: PointGroundGradient(InPointGroundGradient)
+			, PointLengthFromRoot(InPointLengthFromRoot)
+		{
+		}
+
+		FComputePointGroundGradientAttributes(FManagedArrayCollection& InCollection)
+			: PointGroundGradient(PV::FPointGroundGradientAttribute::FindAttribute(InCollection))
+			, PointLengthFromRoot(PV::FPointLengthFromRootAttribute::FindAttribute(InCollection))
+		{
+		}
+
+		bool IsValid() const { return PointGroundGradient && PointLengthFromRoot; }
+	};
+
+	PROCEDURALVEGETATION_API extern void ComputePointGroundGradient(const FComputePointGroundGradientAttributes& InAttributes);
+
+	struct FComputePointScaleGradientAttributes
+	{
+		PV::FPointScaleGradientAttributeView PointScaleGradient;
+		PV::FPointScaleAttributeConstView PointScale;
+
+		FComputePointScaleGradientAttributes(const FComputeGradientsAttributes& InComputeGradientsAttributes)
+			: PointScaleGradient(InComputeGradientsAttributes.PointScaleGradient)
+			, PointScale(InComputeGradientsAttributes.PointScale)
+		{
+		}
+
+		FComputePointScaleGradientAttributes(
+			PV::FPointScaleGradientAttributeView InPointScaleGradient,
+			PV::FPointScaleAttributeConstView InPointScale)
+			: PointScaleGradient(InPointScaleGradient)
+			, PointScale(InPointScale)
+		{
+		}
+
+		FComputePointScaleGradientAttributes(FManagedArrayCollection& InCollection)
+			: PointScaleGradient(PV::FPointScaleGradientAttribute::FindAttribute(InCollection))
+			, PointScale(PV::FPointScaleAttribute::FindAttribute(InCollection))
+		{
+		}
+
+		bool IsValid() const { return PointScaleGradient && PointScale; }
+	};
+
+	PROCEDURALVEGETATION_API extern void ComputePointScaleGradient(const FComputePointScaleGradientAttributes& InAttributes);
+
+	struct FComputePointPlantAndBranchGradientAttributes
+	{
+		PV::FPointPlantGradientAttributeView PointPlantGradient;
+		PV::FPointBranchGradientAttributeView PointBranchGradient;
+		PV::FBranchPointsAttributeConstView BranchPoints;
+		PV::FBranchParentNumberAttributeConstView BranchParentNumber;
+		PV::FBranchNumberAttributeConstView BranchNumber;
+		PV::FBranchChildrenAttributeConstView BranchChildren;
+
+		FComputePointPlantAndBranchGradientAttributes(const FComputeGradientsAttributes& InComputeGradientsAttributes)
+			: PointPlantGradient(InComputeGradientsAttributes.PointPlantGradient)
+			, PointBranchGradient(InComputeGradientsAttributes.PointBranchGradient)
+			, BranchPoints(InComputeGradientsAttributes.BranchPoints)
+			, BranchParentNumber(InComputeGradientsAttributes.BranchParentNumber)
+			, BranchNumber(InComputeGradientsAttributes.BranchNumber)
+			, BranchChildren(InComputeGradientsAttributes.BranchChildren)
+		{
+		}
+
+		FComputePointPlantAndBranchGradientAttributes(
+			PV::FPointPlantGradientAttributeView InPointPlantGradient,
+			PV::FPointBranchGradientAttributeView InPointBranchGradient,
+			PV::FBranchPointsAttributeConstView InBranchPoints,
+			PV::FBranchParentNumberAttributeConstView InBranchParentNumber,
+			PV::FBranchNumberAttributeConstView InBranchNumber,
+			PV::FBranchChildrenAttributeConstView InBranchChildren)
+			: PointPlantGradient(InPointPlantGradient)
+			, PointBranchGradient(InPointBranchGradient)
+			, BranchPoints(InBranchPoints)
+			, BranchParentNumber(InBranchParentNumber)
+			, BranchNumber(InBranchNumber)
+			, BranchChildren(InBranchChildren)
+		{
+		}
+
+		FComputePointPlantAndBranchGradientAttributes(FManagedArrayCollection& InCollection)
+			: PointPlantGradient(PV::FPointPlantGradientAttribute::FindAttribute(InCollection))
+			, PointBranchGradient(PV::FPointBranchGradientAttribute::FindAttribute(InCollection))
+			, BranchPoints(PV::FBranchPointsAttribute::FindAttribute(InCollection))
+			, BranchParentNumber(PV::FBranchParentNumberAttribute::FindAttribute(InCollection))
+			, BranchNumber(PV::FBranchNumberAttribute::FindAttribute(InCollection))
+			, BranchChildren(PV::FBranchChildrenAttribute::FindAttribute(InCollection))
+		{
+		}
+
+		bool IsValid() const
+		{
+			return PointPlantGradient
+				&& PointBranchGradient
+				&& BranchPoints
+				&& BranchParentNumber
+				&& BranchNumber
+				&& BranchChildren;
+		}
+	};
+
+	PROCEDURALVEGETATION_API extern void ComputePointPlantAndBranchGradient(const FComputePointPlantAndBranchGradientAttributes& InAttributes);
+
 	struct FComputeBudStatusAttributes
 	{ 
 		PV::FBudStatusAttributeView BudStatus;
diff --git a/Engine/Plugins/Experimental/ProceduralVegetationEditor/Source/ProceduralVegetation/Private/Helpers/PVJSONHelper.h b/Engine/Plugins/Experimental/ProceduralVegetationEditor/Source/ProceduralVegetation/Private/Helpers/PVJSONHelper.h
new file mode 100644
--- /dev/null
+++ b/Engine/Plugins/Experimental/ProceduralVegetationEditor/Source/ProceduralVegetation/Private/Helpers/PVJSONHelper.h
@@ -0,0 +1,655 @@
+// Copyright Epic Games, Inc. All Rights Reserved.
+
+#pragma once
+
+#include "CoreMinimal.h"
+#include "Chaos/Deformable/MuscleActivationConstraints.h"
+#include "GeometryCollection/ManagedArrayCollection.h"
+#include "Facades/PVAttributesNames.h"
+#include "Dom/JsonObject.h"
+#include "Dom/JsonValue.h"
+#include "Facades/PVBranchFacade.h"
+#include "Facades/PVMetaInfoFacade.h"
+#include "Facades/PVFoliageFacade.h"
+#include "Facades/PVProfileFacade.h"
+#include "Misc/FileHelper.h"
+#include "Serialization/JsonSerializer.h"
+#include "Misc/Paths.h"
+#include "Misc/PackageName.h"
+
+namespace PV
+{
+	static void FillAttributes(FManagedArrayCollection& Collection, const FName Group, const TSharedPtr<FJsonObject>& AttributesObject)
+	{
+		for (const auto& [Key, Attribute] : AttributesObject->Values)
+		{
+			const TSharedPtr<FJsonObject>& AttributeObject = Attribute->AsObject();
+
+			const FString Type = AttributeObject->GetStringField(TEXT("type"));
+			const int32 ElementCount = AttributeObject->GetNumberField(TEXT("size"));
+			const bool IsArray = AttributeObject->GetBoolField(TEXT("isArray"));
+			const TArray<TSharedPtr<FJsonValue>>& Values = AttributeObject->GetArrayField(TEXT("values"));
+
+			if (IsArray)
+			{
+				if (ElementCount == 1)
+				{
+					if (Type == "int")
+					{
+						TManagedArray<TArray<int32>>& Attr = Collection.AddAttribute<TArray<int32>>(*Key, Group);
+						for (int32 Index = 0; Index < Values.Num(); ++Index)
+						{
+							TArray<int32>& CurrentAttrib = Attr[Index];
+							const TArray<TSharedPtr<FJsonValue>>& AttribValues = Values[Index]->AsArray();
+							for (const TSharedPtr<FJsonValue>& AttribValue : AttribValues)
+							{
+								CurrentAttrib.Add(AttribValue->AsNumber());
+							}
+						}
+					}
+					else if (Type == "float")
+					{
+						TManagedArray<TArray<float>>& Attr = Collection.AddAttribute<TArray<float>>(*Key, Group);
+						for (int32 Index = 0; Index < Values.Num(); ++Index)
+						{
+							TArray<float>& CurrentAttrib = Attr[Index];
+							const TArray<TSharedPtr<FJsonValue>>& AttribValues = Values[Index]->AsArray();
+							for (const TSharedPtr<FJsonValue>& AttribValue : AttribValues)
+							{
+								CurrentAttrib.Add(AttribValue->AsNumber());
+							}
+						}
+					}
+					else if (Type == "string")
+					{
+						// TODO: Figure out how to manage string arrays
+					}
+				}
+				else if (ElementCount == 3)
+				{
+					if (Type == "int")
+					{
+						TManagedArray<TArray<FIntVector>>& Attr = Collection.AddAttribute<TArray<FIntVector>>(*Key, Group);
+						for (int32 Index = 0; Index < Values.Num(); ++Index)
+						{
+							TArray<FIntVector>& CurrentAttrib = Attr[Index];
+							const TArray<TSharedPtr<FJsonValue>>& AttribValues = Values[Index]->AsArray();
+							for (int ElemIndex = 0; ElemIndex < AttribValues.Num(); ElemIndex += ElementCount)
+							{
+								CurrentAttrib.Emplace(AttribValues[ElemIndex]->AsNumber(), AttribValues[ElemIndex + 2]->AsNumber(),
+									AttribValues[ElemIndex + 1]->AsNumber());
+							}
+						}
+					}
+					else if (Type == "float")
+					{
+						TManagedArray<TArray<FVector3f>>& Attr = Collection.AddAttribute<TArray<FVector3f>>(*Key, Group);
+						for (int32 Index = 0; Index < Values.Num(); ++Index)
+						{
+							TArray<FVector3f>& CurrentAttrib = Attr[Index];
+							const TArray<TSharedPtr<FJsonValue>>& AttribValues = Values[Index]->AsArray();
+							for (int ElemIndex = 0; ElemIndex < AttribValues.Num(); ElemIndex += ElementCount)
+							{
+								CurrentAttrib.Emplace(AttribValues[ElemIndex]->AsNumber(), AttribValues[ElemIndex + 2]->AsNumber(),
+									AttribValues[ElemIndex + 1]->AsNumber());
+							}
+						}
+					}
+				}
+			}
+			else
+			{
+				if (ElementCount == 1)
+				{
+					if (Type == "int")
+					{
+						TManagedArray<int32>& Attr = Collection.AddAttribute<int32>(*Key, Group);
+						for (int32 Index = 0; Index < Values.Num(); ++Index)
+						{
+							Attr[Index] = Values[Index]->AsNumber();
+						}
+					}
+					else if (Type == "float")
+					{
+						TManagedArray<float>& Attr = Collection.AddAttribute<float>(*Key, Group);
+						for (int32 Index = 0; Index < Values.Num(); ++Index)
+						{
+							Attr[Index] = Values[Index]->AsNumber();
+						}
+					}
+					else if (Type == "string")
+					{
+						TManagedArray<FString>& Attr = Collection.AddAttribute<FString>(*Key, Group);
+						for (int32 Index = 0; Index < Values.Num(); ++Index)
+						{
+							Attr[Index] = Values[Index]->AsString();
+						}
+					}
+				}
+				else if (ElementCount == 3)
+				{
+					if (Type == "int")
+					{
+						TManagedArray<FIntVector>& Attr = Collection.AddAttribute<FIntVector>(*Key, Group);
+						for (int32 Index = 0; Index < Values.Num(); ++Index)
+						{
+							const TArray<TSharedPtr<FJsonValue>>& VectorAttribute = Values[Index]->AsArray();
+							Attr[Index] = FIntVector(VectorAttribute[0]->AsNumber(), VectorAttribute[2]->AsNumber(), VectorAttribute[1]->AsNumber());
+						}
+					}
+					else if (Type == "float")
+					{
+						TManagedArray<FVector3f>& Attr = Collection.AddAttribute<FVector3f>(*Key, Group);
+						for (int32 Index = 0; Index < Values.Num(); ++Index)
+						{
+							const TArray<TSharedPtr<FJsonValue>>& VectorAttribute = Values[Index]->AsArray();
+							Attr[Index] = FVector3f(VectorAttribute[0]->AsNumber(), VectorAttribute[2]->AsNumber(), VectorAttribute[1]->AsNumber());
+						}
+					}
+				}
+			}
+		}
+	}
+
+	static void FillDetailsAttributes(FManagedArrayCollection& Collection, const FName Group, const TSharedPtr<FJsonObject>& AttributesObject)
+	{
+		for (const auto& [Key, Attribute] : AttributesObject->Values)
+		{
+			const TSharedPtr<FJsonObject>& AttributeObject = Attribute->AsObject();
+
+			const FString Type = AttributeObject->GetStringField(TEXT("type"));
+			const bool IsArray = AttributeObject->GetBoolField(TEXT("isArray"));
+
+			if (IsArray)
+			{
+				const TArray<TSharedPtr<FJsonValue>>& Values = AttributeObject->GetArrayField(TEXT("value"));
+				if (Type == "int")
+				{
+					TManagedArray<TArray<int32>>& Attr = Collection.AddAttribute<TArray<int32>>(*Key, Group);
+					TArray<int32>& CurrentAttrib = Attr[0];
+					for (const TSharedPtr<FJsonValue>& AttribValue : Values)
+					{
+						CurrentAttrib.Add(AttribValue->AsNumber());
+					}
+				}
+				else if (Type == "float")
+				{
+					TManagedArray<TArray<float>>& Attr = Collection.AddAttribute<TArray<float>>(*Key, Group);
+					TArray<float>& CurrentAttrib = Attr[0];
+					for (const TSharedPtr<FJsonValue>& AttribValue : Values)
+					{
+						CurrentAttrib.Add(AttribValue->AsNumber());
+					}
+				}
+				else if (Type == "string")
+				{
+					// TODO: Figure out how to manage string arrays
+				}
+			}
+			else
+			{
+				if (Type == "int")
+				{
+					TManagedArray<int32>& Attr = Collection.AddAttribute<int32>(*Key, Group);
+					Attr[0] = AttributeObject->GetNumberField(TEXT("value"));
+				}
+				else if (Type == "float")
+				{
+					TManagedArray<float>& Attr = Collection.AddAttribute<float>(*Key, Group);
+					Attr[0] = AttributeObject->GetNumberField(TEXT("value"));
+				}
+				else if (Type == "string")
+				{
+					TManagedArray<FString>& Attr = Collection.AddAttribute<FString>(*Key, Group);
+					Attr[0] = AttributeObject->GetStringField(TEXT("value"));
+				}
+			}
+		}
+	}
+
+	static void FillPlantProfilesData(FManagedArrayCollection& Collection, const FName Group, const TSharedPtr<FJsonObject>& AttributesObject)
+	{
+		Facades::FPlantProfileFacade Facade(Collection);
+
+		for (const auto& [Key, Attribute] : AttributesObject->Values)
+		{
+			if (const TSharedPtr<FJsonObject>& AttributeObject = Attribute->AsObject();
+				AttributeObject
+				&& AttributeObject->HasTypedField<EJson::Array>(TEXT("value"))
+				&& Key.StartsWith(TEXT("plantProfile_")))
+			{
+				const TArray<TSharedPtr<FJsonValue>>& Values = AttributeObject->GetArrayField(TEXT("value"));
+				TArray<float> Points;
+				Points.Reserve(Values.Num());
+
+				for (const TSharedPtr<FJsonValue>& AttribValue : Values)
+				{
+					if (double Point; AttribValue->TryGetNumber(Point))
+					{
+						Points.Add(static_cast<float>(AttribValue->AsNumber()));
+					}
+				}
+
+				Facade.AddProfileEntry(Points);
+			}
+		}
+	}
+
+	static void FillFoliageData(FManagedArrayCollection& Collection, const TSharedPtr<FJsonObject>& PrimitiveAttributesObject, const FString& InPath)
+	{
+		const TSharedPtr<FJsonObject>& InstancerNameObject = PrimitiveAttributesObject->GetObjectField(TEXT("instancer_name"));
+		const TArray<TSharedPtr<FJsonValue>>& InstancerNameValues = InstancerNameObject->GetArrayField(TEXT("values"));
+
+		const TSharedPtr<FJsonObject>& InstancerPivotObject = PrimitiveAttributesObject->GetObjectField(TEXT("instancer_pivot"));
+		const TArray<TSharedPtr<FJsonValue>>& InstancerPivotValues = InstancerPivotObject->GetArrayField(TEXT("values"));
+
+		const TSharedPtr<FJsonObject>& InstancerUpVectorObject = PrimitiveAttributesObject->GetObjectField(TEXT("instancer_UP"));
+		const TArray<TSharedPtr<FJsonValue>>& InstancerUpVectorValues = InstancerUpVectorObject->GetArrayField(TEXT("values"));
+
+		const TSharedPtr<FJsonObject>& InstancerNormalVectorObject = PrimitiveAttributesObject->GetObjectField(TEXT("instancer_N"));
+		const TArray<TSharedPtr<FJsonValue>>& InstancerNormalVectorValues = InstancerNormalVectorObject->GetArrayField(TEXT("values"));
+
+		const TSharedPtr<FJsonObject>& InstancerScaleObject = PrimitiveAttributesObject->GetObjectField(TEXT("instancer_scale"));
+		const TArray<TSharedPtr<FJsonValue>>& InstancerScaleValues = InstancerScaleObject->GetArrayField(TEXT("values"));
+
+		const TSharedPtr<FJsonObject>& InstancerLFRObject = PrimitiveAttributesObject->GetObjectField(TEXT("instancer_LFR"));
+		const TArray<TSharedPtr<FJsonValue>>& InstancerLFRValues = InstancerLFRObject->GetArrayField(TEXT("values"));
+
+		// Compute unique names and number of elements for group
+		int32 TotalValues = 0;
+		int32 NextFoliageIndex = 0;
+		TMap<FString, int32> FoliageNamesToIndex;
+		TArray<FString> FoliageNames;
+		for (int32 i = 0; i < InstancerNameValues.Num(); ++i)
+		{
+			const TArray<TSharedPtr<FJsonValue>>& InstancerNames = InstancerNameValues[i]->AsArray();
+			for (int32 Index = 0; Index < InstancerNames.Num(); ++Index)
+			{
+				const FString& Name = InstancerNames[Index]->AsString();
+				if (!FoliageNamesToIndex.Contains(Name))
+				{
+					FoliageNamesToIndex.Emplace(Name, NextFoliageIndex++);
+					FoliageNames.Emplace(Name);
+				}
+			}
+			TotalValues += InstancerNames.Num();
+		}
+
+		Facades::FFoliageFacade Facade(Collection, TotalValues);
+
+		for (FString& FoliageName : FoliageNames)
+		{
+			FoliageName = FoliageName + '.' + FoliageName;
+		}
+		
+		Facade.SetFoliageNames(FoliageNames);
+
+		int32 FoliageEntryIndex = 0;
+		for (int32 Id = 0; Id < InstancerNameValues.Num(); ++Id)
+		{
+			const TArray<TSharedPtr<FJsonValue>>& InstancerNames = InstancerNameValues[Id]->AsArray();
+			const TArray<TSharedPtr<FJsonValue>>& PivotPoints = InstancerPivotValues[Id]->AsArray();
+			const TArray<TSharedPtr<FJsonValue>>& UpVectors = InstancerUpVectorValues[Id]->AsArray();
+			const TArray<TSharedPtr<FJsonValue>>& NormalVectors = InstancerNormalVectorValues[Id]->AsArray();
+			const TArray<TSharedPtr<FJsonValue>>& Scales = InstancerScaleValues[Id]->AsArray();
+			const TArray<TSharedPtr<FJsonValue>>& LFRs = InstancerLFRValues[Id]->AsArray();
+
+			TArray<int32> FoliageEntryIds;
+			for (int32 j = 0; j < InstancerNames.Num(); ++j)
+			{
+				const FString Name = InstancerNames[j]->AsString();
+				const FVector3f PivotPoint = FVector3f(
+					PivotPoints[j * 3]->AsNumber(),
+					PivotPoints[(j * 3) + 2]->AsNumber(),
+					PivotPoints[(j * 3) + 1]->AsNumber()) * 100.0f;
+				const FVector3f UpVector = FVector3f(
+					UpVectors[j * 3]->AsNumber(),
+					UpVectors[(j * 3) + 2]->AsNumber(),
+					UpVectors[(j * 3) + 1]->AsNumber());
+				const FVector3f NormalVector = FVector3f(
+					NormalVectors[j * 3]->AsNumber(),
+					NormalVectors[(j * 3) + 2]->AsNumber(),
+					NormalVectors[(j * 3) + 1]->AsNumber());
+				const float Scale = Scales[j]->AsNumber();
+				const float LengthFromRoot = LFRs[j]->AsNumber();
+
+				Facade.SetFoliageEntry(FoliageEntryIndex, {
+					.NameId = FoliageNamesToIndex[Name],
+					.BranchId = Id,
+					.PivotPoint = PivotPoint,
+					.UpVector = UpVector,
+					.NormalVector = NormalVector,
+					.Scale = Scale,
+					.LengthFromRoot = LengthFromRoot,
+				});
+
+				FoliageEntryIds.Add(FoliageEntryIndex);
+				FoliageEntryIndex++;
+			}
+
+			Facade.SetFoliageIdsArray(Id, FoliageEntryIds);
+		}
+	}
+
+	static void SetFoliagePaths(FManagedArrayCollection& Collection,const FString& FilePath)
+	{
+		Facades::FFoliageFacade Facade(Collection);
+
+		FString AbsolutePath = FPaths::ConvertRelativePathToFull(FilePath);
+		FString PackagePath;
+
+		if (FPackageName::TryConvertFilenameToLongPackageName(AbsolutePath, PackagePath))
+		{
+			TArray<FString> FoliageNames = Facade.GetFoliageNames();
+			FString PackageFolder = FPackageName::GetLongPackagePath(PackagePath);
+			Facade.SetFoliagePath(PackageFolder);
+			for (FString& FoliageName : FoliageNames)
+			{
+				FoliageName = PackageFolder / FoliageName;
+			}
+			Facade.SetFoliageNames(FoliageNames);
+		}
+		
+	}
+
+	static bool HasJsonFieldPath(const TSharedPtr<FJsonObject>& JsonObject, const FString& Path)
+	{
+		if (!JsonObject.IsValid())
+		{
+			return false;
+		}
+
+		TArray<FString> Keys;
+		Path.ParseIntoArray(Keys, TEXT("."), true);
+
+		TSharedPtr<FJsonObject> CurrentObject = JsonObject;
+
+		for (int32 i = 0; i < Keys.Num(); ++i)
+		{
+			const FString& CurrentKey = Keys[i];
+			if (i == Keys.Num() - 1)
+			{
+				if (!CurrentObject->HasField(CurrentKey))
+				{
+					return false;
+				}
+			}
+			else
+			{
+				const TSharedPtr<FJsonObject>* NextObject;
+				if (CurrentObject->TryGetObjectField(CurrentKey, NextObject))
+				{
+					CurrentObject = *NextObject;
+				}
+				else
+				{
+					return false;
+				}
+			}
+		}
+
+		return true;
+	}
+
+	static bool LoadMegaPlantsJsonToCollection(FManagedArrayCollection& Collection, const FString& FilePath, FString& OutErrorMessage)
+	{
+		FString MegaPlantsData;
+		if (!FFileHelper::LoadFileToString(MegaPlantsData, *FilePath))
+		{
+			OutErrorMessage = TEXT("Failed to open skeleton file");
+			return false;
+		}
+
+		TSharedPtr<FJsonObject> LoadedData = MakeShareable(new FJsonObject);
+		if (!FJsonSerializer::Deserialize(TJsonReaderFactory<>::Create(MegaPlantsData), LoadedData))
+		{
+			OutErrorMessage = TEXT("Failed to deserialize JSON file for skeleton");
+			return false;
+		}
+
+		TArray<FString> RequiredJSONPaths = {
+			TEXT("points.attributes.pscale"),
+			TEXT("points.positions"),
+			TEXT("points.attributes.lengthFromRoot.values"),
+			TEXT("points.attributes.LOD_totalPscaleGradient.values"),
+			TEXT("points.attributes.budDirection.values"),
+			TEXT("primitives.points"),
+			TEXT("primitives.attributes.instancer_name.values"),
+			TEXT("primitives.attributes.instancer_pivot.values"),
+			TEXT("primitives.attributes.instancer_UP.values"),
+			TEXT("primitives.attributes.instancer_scale.values"),
+			TEXT("primitives.attributes.instancer_LFR.values"),
+			TEXT("primitives.attributes.parents.values"),
+			TEXT("primitives.attributes.children.values"),
+			TEXT("primitives.attributes.branchNumber.values"),
+			TEXT("globalAttributes.phyllotaxyLeaf.value")
+		};
+
+		for (const FString& Path : RequiredJSONPaths)
+		{
+			if (!HasJsonFieldPath(LoadedData, Path))
+			{
+				OutErrorMessage = FString::Printf(TEXT("Failed to find required path in JSON file: %s"), *Path);
+				return false;
+			}
+		}
+
+		const TSharedPtr<FJsonObject>& Points = LoadedData->GetObjectField(TEXT("points"));
+		const TSharedPtr<FJsonObject>& PointAttributes = Points->GetObjectField(TEXT("attributes"));
+
+		const TArray<TSharedPtr<FJsonValue>>& PointPositions = Points->GetArrayField(TEXT("positions"));
+		const TArray<TSharedPtr<FJsonValue>>& PointScales = PointAttributes->GetObjectField(TEXT("pscale"))->GetArrayField(TEXT("values"));
+
+		Collection.AddGroup(GroupNames::PointGroup);
+		Collection.AddElements(PointPositions.Num(), GroupNames::PointGroup);
+
+		TManagedArray<FVector3f>& PositionArray = Collection.AddAttribute<FVector3f>(AttributeNames::PointPosition, GroupNames::PointGroup);
+		TManagedArray<float>& ScaleArray = Collection.AddAttribute<float>("Scale", GroupNames::PointGroup);
+
+		for (int32 PointIndex = 0; PointIndex < PointPositions.Num(); ++PointIndex)
+		{
+			const TArray<TSharedPtr<FJsonValue>>& PointPosition = PointPositions[PointIndex]->AsArray();
+
+			const FVector3f Position = FVector3f(PointPosition[0]->AsNumber(), PointPosition[2]->AsNumber(), PointPosition[1]->AsNumber()) * 100.0f;
+			const float Scale = PointScales[PointIndex]->AsNumber() * 100.0f;
+
+			PositionArray[PointIndex] = Position;
+			ScaleArray[PointIndex] = Scale;
+		}
+
+		const TSharedPtr<FJsonObject>& Primitives = LoadedData->GetObjectField(TEXT("primitives"));
+		const TSharedPtr<FJsonObject>& PrimitiveAttributes = Primitives->GetObjectField(TEXT("attributes"));
+
+		const TSharedPtr<FJsonObject>& GlobalAttributes = LoadedData->GetObjectField(TEXT("globalAttributes"));
+
+		const TArray<TSharedPtr<FJsonValue>>& PrimitivesPoints = Primitives->GetArrayField(TEXT("points"));
+
+		Collection.AddGroup(GroupNames::BranchGroup);
+		Collection.AddElements(PrimitivesPoints.Num(), GroupNames::BranchGroup);
+
+		TManagedArray<TArray<int32>>& PointsArray = Collection.AddAttribute<TArray<int32>>(
+			PV::AttributeNames::BranchPoints, PV::GroupNames::BranchGroup);
+
+		for (int32 PrimitiveIndex = 0; PrimitiveIndex < PrimitivesPoints.Num(); ++PrimitiveIndex)
+		{
+			const TArray<TSharedPtr<FJsonValue>>& PrimitivePoints = PrimitivesPoints[PrimitiveIndex]->AsArray();
+			for (const TSharedPtr<FJsonValue>& PrimitivePoint : PrimitivePoints)
+			{
+				PointsArray[PrimitiveIndex].Add(PrimitivePoint->AsNumber());
+			}
+		}
+
+		Collection.AddGroup(GroupNames::DetailsGroup);
+		Collection.AddElements(1, GroupNames::DetailsGroup);
+
+		// Please note that LFR (length from root) for both skeleton points and foliage points are not being scaled by
+		// 100 at the moment, which is OK for now since they're only being compared relative to each other for now,
+		// however if they're used in any computation in the future in any meaningful way where other quantities are
+		// involved they will end up giving erroneous results.
+		FillAttributes(Collection, GroupNames::PointGroup, PointAttributes);
+		FillAttributes(Collection, GroupNames::BranchGroup, PrimitiveAttributes);
+		FillDetailsAttributes(Collection, GroupNames::DetailsGroup, GlobalAttributes);
+		FillPlantProfilesData(Collection, GroupNames::PlantProfilesGroup, GlobalAttributes);
+
+		FillFoliageData(Collection, PrimitiveAttributes, FilePath);
+
+		Facades::FMetaInfoFacade MetaInfoFacade(Collection);
+		MetaInfoFacade.CreateGuid(FilePath);
+
+		OutErrorMessage = "";
+		return true;
+	}
+
+	static bool LoadGrowthDataJsonToCollection(FManagedArrayCollection& Collection, const FString& FilePath, FString& OutErrorMessage)
+	{
+		FString MegaPlantsData;
+		if (!FFileHelper::LoadFileToString(MegaPlantsData, *FilePath))
+		{
+			OutErrorMessage = TEXT("Failed to open skeleton file");
+			return false;
+		}
+
+		TSharedPtr<FJsonObject> LoadedData = MakeShareable(new FJsonObject);
+		if (!FJsonSerializer::Deserialize(TJsonReaderFactory<>::Create(MegaPlantsData), LoadedData))
+		{
+			OutErrorMessage = TEXT("Failed to deserialize JSON file for skeleton");
+			return false;
+		}
+
+		TArray<FString> RequiredJSONPaths = {
+			TEXT("points.positions"),
+			TEXT("points.attributes.budDirection.values"),
+			TEXT("primitives.points"),
+			TEXT("primitives.attributes.parents.values"),
+			TEXT("primitives.attributes.children.values"),
+			TEXT("primitives.attributes.branchNumber.values"),
+		};
+
+		for (const FString& Path : RequiredJSONPaths)
+		{
+			if (!HasJsonFieldPath(LoadedData, Path))
+			{
+				OutErrorMessage = FString::Printf(TEXT("Failed to find required path in JSON file: %s"), *Path);
+				return false;
+			}
+		}
+
+		const TSharedPtr<FJsonObject>& Points = LoadedData->GetObjectField(TEXT("points"));
+		const TSharedPtr<FJsonObject>& PointAttributes = Points->GetObjectField(TEXT("attributes"));
+
+		const TArray<TSharedPtr<FJsonValue>>& PointPositions = Points->GetArrayField(TEXT("positions"));
+
+		const TArray<TSharedPtr<FJsonValue>>* PointScales = Invoke([&]()
+		{
+			const TArray<TSharedPtr<FJsonValue>>* OutPointScales = nullptr;
+			const TSharedPtr<FJsonObject>* PScaleObject = nullptr;
+			if (PointAttributes->TryGetObjectField(TEXT("pscale"), PScaleObject))
+			{
+				(*PScaleObject)->TryGetArrayField(TEXT("values"), OutPointScales);
+			}
+			return OutPointScales;
+		});
+
+		Collection.AddGroup(GroupNames::PointGroup);
+		Collection.AddElements(PointPositions.Num(), GroupNames::PointGroup);
+
+		TManagedArray<FVector3f>& PositionArray = Collection.AddAttribute<FVector3f>(AttributeNames::PointPosition, GroupNames::PointGroup);
+		TManagedArray<float>& ScaleArray = Collection.AddAttribute<float>("Scale", GroupNames::PointGroup);
+
+		for (int32 PointIndex = 0; PointIndex < PointPositions.Num(); ++PointIndex)
+		{
+			const TArray<TSharedPtr<FJsonValue>>& PointPosition = PointPositions[PointIndex]->AsArray();
+
+			const FVector3f Position = FVector3f(PointPosition[0]->AsNumber(), PointPosition[2]->AsNumber(), PointPosition[1]->AsNumber()) * 100.0f;
+			const float Scale = PointScales ? (*PointScales)[PointIndex]->AsNumber() * 100.0f : 0.f;
+
+			PositionArray[PointIndex] = Position;
+			ScaleArray[PointIndex] = Scale;
+		}
+
+		const TSharedPtr<FJsonObject>& Primitives = LoadedData->GetObjectField(TEXT("primitives"));
+		const TSharedPtr<FJsonObject>& PrimitiveAttributes = Primitives->GetObjectField(TEXT("attributes"));
+		const TArray<TSharedPtr<FJsonValue>>& PrimitivesPoints = Primitives->GetArrayField(TEXT("points"));
+
+		Collection.AddGroup(GroupNames::BranchGroup);
+		Collection.AddElements(PrimitivesPoints.Num(), GroupNames::BranchGroup);
+
+		TManagedArray<TArray<int32>>& PointsArray = Collection.AddAttribute<TArray<int32>>(
+			PV::AttributeNames::BranchPoints, PV::GroupNames::BranchGroup);
+
+		for (int32 PrimitiveIndex = 0; PrimitiveIndex < PrimitivesPoints.Num(); ++PrimitiveIndex)
+		{
+			const TArray<TSharedPtr<FJsonValue>>& PrimitivePoints = PrimitivesPoints[PrimitiveIndex]->AsArray();
+			for (const TSharedPtr<FJsonValue>& PrimitivePoint : PrimitivePoints)
+			{
+				PointsArray[PrimitiveIndex].Add(PrimitivePoint->AsNumber());
+			}
+		}
+
+		Collection.AddGroup(GroupNames::DetailsGroup);
+		Collection.AddElements(1, GroupNames::DetailsGroup);
+
+		FillAttributes(Collection, GroupNames::PointGroup, PointAttributes);
+		FillAttributes(Collection, GroupNames::BranchGroup, PrimitiveAttributes);
+		FillDetailsAttributes(Collection, GroupNames::DetailsGroup, LoadedData->GetObjectField(TEXT("globalAttributes")));
+
+		Facades::FMetaInfoFacade MetaInfoFacade(Collection);
+		MetaInfoFacade.CreateGuid(FilePath);
+
+		OutErrorMessage = "";
+		return true;
+	}
+
+	static TSharedPtr<FJsonObject> LoadMetaFileIntoJsonObject(const FString& FilePath, FString& OutErrorMessage)
+	{
+		FString MetaData;
+		if (!FFileHelper::LoadFileToString(MetaData, *FilePath))
+		{
+			OutErrorMessage = TEXT("Failed to open skeleton file");
+			return nullptr;
+		}
+
+		TSharedPtr<FJsonObject> LoadedData = MakeShareable(new FJsonObject);
+		if (!FJsonSerializer::Deserialize(TJsonReaderFactory<>::Create(MetaData), LoadedData))
+		{
+			OutErrorMessage = TEXT("Failed to deserialize JSON file for skeleton");
+			return nullptr;
+		}
+
+		TArray<FString> RequiredJSONPaths = {
+			TEXT("apix_uv"),
+			TEXT("apix_texture_channels")
+		};
+
+		for (const FString& Path : RequiredJSONPaths)
+		{
+			if (!HasJsonFieldPath(LoadedData, Path))
+			{
+				OutErrorMessage = FString::Printf(TEXT("Failed to find required path in JSON file: %s"), *Path);
+				return nullptr;
+			}
+		}
+
+		OutErrorMessage = "";
+		return LoadedData;
+	}
+
+	static bool LoadMetaJsonToCollection(FManagedArrayCollection& Collection, TSharedPtr<FJsonObject> LoadedData)
+	{
+		
+		const TArray<TSharedPtr<FJsonValue>>& MaterialURanges = LoadedData->GetArrayField(TEXT("apix_uv"));
+
+		TArray<FVector2f> URanges;
+		URanges.Reserve(MaterialURanges.Num());
+		
+		for (int i=0; (i + 1) < MaterialURanges.Num(); i+= 2)
+		{
+			URanges.Add(FVector2f(MaterialURanges[i]->AsNumber() , MaterialURanges[i+1]->AsNumber()));
+		}
+
+		Facades::FBranchFacade Facade = Facades::FBranchFacade(Collection);
+		Facade.SetTrunkURange(URanges);
+		
+		return true;
+	}
+}
diff --git a/Engine/Plugins/Experimental/ProceduralVegetationEditor/Source/ProceduralVegetation/Private/Helpers/PVUtilities.cpp b/Engine/Plugins/Experimental/ProceduralVegetationEditor/Source/ProceduralVegetation/Private/Helpers/PVUtilities.cpp
--- a/Engine/Plugins/Experimental/ProceduralVegetationEditor/Source/ProceduralVegetation/Private/Helpers/PVUtilities.cpp
+++ b/Engine/Plugins/Experimental/ProceduralVegetationEditor/Source/ProceduralVegetation/Private/Helpers/PVUtilities.cpp
@@ -283,7 +283,7 @@ namespace PV::Utilities
 		return FLinearColor(HueDeg, 1, 1).HSVToLinearRGB();
 	};
  
-	bool IsValidPVData(const FManagedArrayCollection& Collection)
+	bool IsValidGrowthData(const FManagedArrayCollection& Collection)
 	{
 		Facades::FPointFacade PointFacade = Facades::FPointFacade(Collection);
 		Facades::FBranchFacade BranchFacade = Facades::FBranchFacade(Collection);
diff --git a/Engine/Plugins/Experimental/ProceduralVegetationEditor/Source/ProceduralVegetation/Private/Nodes/Development/PVGrowthDataDiffSettings.cpp b/Engine/Plugins/Experimental/ProceduralVegetationEditor/Source/ProceduralVegetation/Private/Nodes/Development/PVGrowthDataDiffSettings.cpp
new file mode 100644
--- /dev/null
+++ b/Engine/Plugins/Experimental/ProceduralVegetationEditor/Source/ProceduralVegetation/Private/Nodes/Development/PVGrowthDataDiffSettings.cpp
@@ -0,0 +1,465 @@
+// Copyright Epic Games, Inc. All Rights Reserved.
+ 
+#include "PVGrowthDataDiffSettings.h"
+ 
+#include "ProceduralVegetationModule.h"
+ 
+#include "DataTypes/PVGrowthData.h"
+#include "Utils/PVAttributes.h"
+#include "GeometryCollection/ManagedArrayCollection.h"
+ 
+#define LOCTEXT_NAMESPACE "PVGrowthDataDiffSettings"
+
+namespace PV::DiffHelpers
+{
+	struct FGrowthDataHierarchyAttributes
+	{
+		FPointPositionAttributeConstView PointPosition;
+		FBranchPointsAttributeConstView BranchPoints;
+		FBranchParentNumberAttributeConstView BranchParentNumber;
+		FBranchParentsAttributeConstView BranchParents;
+		FBranchNumberAttributeConstView BranchNumber;
+		FBranchChildrenAttributeConstView BranchChildren;
+		FBranchHierarchyNumberAttributeConstView BranchHierarchyNumber;
+		FBranchPlantNumberAttributeConstView BranchPlantNumber;
+
+		FGrowthDataHierarchyAttributes() = default;
+		FGrowthDataHierarchyAttributes(const FManagedArrayCollection& InCollection)
+			: PointPosition(FPointPositionAttribute::FindAttribute(InCollection))
+			, BranchPoints(FBranchPointsAttribute::FindAttribute(InCollection))
+			, BranchParentNumber(FBranchParentNumberAttribute::FindAttribute(InCollection))
+			, BranchParents(FBranchParentsAttribute::FindAttribute(InCollection))
+			, BranchNumber(FBranchNumberAttribute::FindAttribute(InCollection))
+			, BranchChildren(FBranchChildrenAttribute::FindAttribute(InCollection))
+			, BranchHierarchyNumber(FBranchHierarchyNumberAttribute::FindAttribute(InCollection))
+			, BranchPlantNumber(FBranchPlantNumberAttribute::FindAttribute(InCollection))
+		{}
+	};
+
+	bool HasSameGrowthDataHierarchy(
+		const FGrowthDataHierarchyAttributes& AttributesA,
+		const FGrowthDataHierarchyAttributes& AttributesB)
+	{
+		if (AttributesA.PointPosition.IsValid() != AttributesB.PointPosition.IsValid()
+			|| AttributesA.BranchPoints.IsValid() != AttributesB.BranchPoints.IsValid()
+			|| AttributesA.BranchParentNumber.IsValid() != AttributesB.BranchParentNumber.IsValid()
+			|| AttributesA.BranchParents.IsValid() != AttributesB.BranchParents.IsValid()
+			|| AttributesA.BranchNumber.IsValid() != AttributesB.BranchNumber.IsValid()
+			|| AttributesA.BranchChildren.IsValid() != AttributesB.BranchChildren.IsValid()
+			|| AttributesA.BranchHierarchyNumber.IsValid() != AttributesB.BranchHierarchyNumber.IsValid()
+			|| AttributesA.BranchPlantNumber.IsValid() != AttributesB.BranchPlantNumber.IsValid())
+		{
+			return false;
+		}
+
+		if (AttributesA.PointPosition.Num() != AttributesB.PointPosition.Num()
+			|| AttributesA.BranchPoints.Num() != AttributesB.BranchPoints.Num()
+			|| AttributesA.BranchParentNumber.Num() != AttributesB.BranchParentNumber.Num()
+			|| AttributesA.BranchParents.Num() != AttributesB.BranchParents.Num()
+			|| AttributesA.BranchNumber.Num() != AttributesB.BranchNumber.Num()
+			|| AttributesA.BranchChildren.Num() != AttributesB.BranchChildren.Num()
+			|| AttributesA.BranchHierarchyNumber.Num() != AttributesB.BranchHierarchyNumber.Num()
+			|| AttributesA.BranchPlantNumber.Num() != AttributesB.BranchPlantNumber.Num())
+		{
+			return false;
+		}
+
+		for (int32 i = 0; i < AttributesA.PointPosition.Num(); ++i)
+		{
+			if (!AttributesA.PointPosition[i].Equals(AttributesB.PointPosition[i]))
+			{
+				return false;
+			}
+		}
+
+		for (int32 i = 0; i < AttributesA.BranchPoints.Num(); ++i)
+		{
+			if (AttributesA.BranchPoints[i] != AttributesB.BranchPoints[i])
+			{
+				return false;
+			}
+		}
+
+		for (int32 i = 0; i < AttributesA.BranchParentNumber.Num(); ++i)
+		{
+			if (AttributesA.BranchParentNumber[i] != AttributesB.BranchParentNumber[i])
+			{
+				return false;
+			}
+		}
+
+		for (int32 i = 0; i < AttributesA.BranchParents.Num(); ++i)
+		{
+			if (AttributesA.BranchParents[i] != AttributesB.BranchParents[i])
+			{
+				return false;
+			}
+		}
+
+		for (int32 i = 0; i < AttributesA.BranchNumber.Num(); ++i)
+		{
+			if (AttributesA.BranchNumber[i] != AttributesB.BranchNumber[i])
+			{
+				return false;
+			}
+		}
+
+		for (int32 i = 0; i < AttributesA.BranchChildren.Num(); ++i)
+		{
+			if (AttributesA.BranchChildren[i] != AttributesB.BranchChildren[i])
+			{
+				return false;
+			}
+		}
+
+		for (int32 i = 0; i < AttributesA.BranchHierarchyNumber.Num(); ++i)
+		{
+			if (AttributesA.BranchHierarchyNumber[i] != AttributesB.BranchHierarchyNumber[i])
+			{
+				return false;
+			}
+		}
+
+		for (int32 i = 0; i < AttributesA.BranchPlantNumber.Num(); ++i)
+		{
+			if (AttributesA.BranchPlantNumber[i] != AttributesB.BranchPlantNumber[i])
+			{
+				return false;
+			}
+		}
+
+		return true;
+	}
+
+	template<typename TAttributeType>
+	TAttributeType DiffAttributeValue(const TAttributeType& A, const TAttributeType& B) { return A - B; }
+
+	template<typename TAttributeType>
+	TArray<TAttributeType> DiffAttributeValue(const TArray<TAttributeType>& A, const TArray<TAttributeType>& B)
+	{ 
+		TArray<TAttributeType> DiffArray;
+		if (A.Num() != B.Num())
+		{
+			return DiffArray;
+		}
+
+		DiffArray.SetNum(A.Num());
+		for (int32 i = 0; i < A.Num(); ++i)
+		{
+			DiffArray[i] = DiffAttributeValue<TAttributeType>(A[i], B[i]);
+		}
+
+		return DiffArray;
+	}
+
+	template<> bool DiffAttributeValue(const bool& A, const bool& B) { return A == B; }
+	template<> FTransform DiffAttributeValue<FTransform>(const FTransform& A, const FTransform& B) 
+	{
+		return FTransform(
+			DiffAttributeValue(A.GetRotation(), B.GetRotation()),
+			DiffAttributeValue(A.GetLocation(), B.GetLocation()),
+			DiffAttributeValue(A.GetScale3D(), B.GetScale3D())
+		);
+	}
+	template<> FTransform3f DiffAttributeValue<FTransform3f>(const FTransform3f& A, const FTransform3f& B) 
+	{ 
+		return FTransform3f(
+			DiffAttributeValue(A.GetRotation(), B.GetRotation()),
+			DiffAttributeValue(A.GetLocation(), B.GetLocation()),
+			DiffAttributeValue(A.GetScale3D(), B.GetScale3D())
+		);
+	}
+	template<> TSet<int32> DiffAttributeValue<TSet<int32>>(const TSet<int32>& A, const TSet<int32>& B)
+	{ 
+		return A.Difference(B);
+	}
+
+	template<typename TAttributeType>
+	void DiffAttributes_Internal(
+		const FManagedArrayCollection& CollectionA,
+		const FManagedArrayCollection& CollectionB,
+		FManagedArrayCollection& ResultCollection,
+		const FName& InAttributeName,
+		const FName& InGroupName
+	)
+	{
+		if (CollectionA.NumElements(InGroupName) != CollectionB.NumElements(InGroupName))
+		{
+			return;
+		}
+
+		ResultCollection.CopyAttribute(CollectionA, InAttributeName, InGroupName);
+
+		const TManagedArray<TAttributeType>& AttributeA = CollectionA.GetAttribute<TAttributeType>(InAttributeName, InGroupName);
+		const TManagedArray<TAttributeType>& AttributeB = CollectionB.GetAttribute<TAttributeType>(InAttributeName, InGroupName);
+		TManagedArray<TAttributeType>& ResultAttribute = ResultCollection.ModifyAttribute<TAttributeType>(InAttributeName, InGroupName);
+
+		for (int32 i = 0; i < AttributeA.Num(); ++i)
+		{
+			ResultAttribute[i] = DiffAttributeValue(AttributeA[i], AttributeB[i]);
+		}
+	}
+
+	void DiffAttributes(
+		const FManagedArrayCollection& CollectionA, 
+		const FManagedArrayCollection& CollectionB, 
+		FManagedArrayCollection& ResultCollection,
+		const FName& InAttributeName, 
+		const FName& InGroupName
+	)
+	{
+		const FManagedArrayCollection::EArrayType ArrayTypeA = CollectionA.GetAttributeType(InAttributeName, InGroupName);
+		const FManagedArrayCollection::EArrayType ArrayTypeB = CollectionB.GetAttributeType(InAttributeName, InGroupName);
+		if (ArrayTypeA != ArrayTypeB)
+		{
+			return;
+		}
+
+		switch (ArrayTypeA)
+		{
+		case FManagedArrayCollection::EArrayType::FVectorType:
+			return DiffAttributes_Internal<FVector3f>(CollectionA, CollectionB, ResultCollection, InAttributeName, InGroupName);
+		case FManagedArrayCollection::EArrayType::FIntVectorType:
+			return DiffAttributes_Internal<FIntVector>(CollectionA, CollectionB, ResultCollection, InAttributeName, InGroupName);
+		case FManagedArrayCollection::EArrayType::FVector2DType:
+			return DiffAttributes_Internal<FVector2f>(CollectionA, CollectionB, ResultCollection, InAttributeName, InGroupName);
+		case FManagedArrayCollection::EArrayType::FLinearColorType:
+			return DiffAttributes_Internal<FLinearColor>(CollectionA, CollectionB, ResultCollection, InAttributeName, InGroupName);
+		case FManagedArrayCollection::EArrayType::FInt32Type:
+			return DiffAttributes_Internal<int32>(CollectionA, CollectionB, ResultCollection, InAttributeName, InGroupName);
+		case FManagedArrayCollection::EArrayType::FTransformType:
+			return DiffAttributes_Internal<FTransform>(CollectionA, CollectionB, ResultCollection, InAttributeName, InGroupName);
+		case FManagedArrayCollection::EArrayType::FFloatType:
+			return DiffAttributes_Internal<float>(CollectionA, CollectionB, ResultCollection, InAttributeName, InGroupName);
+		case FManagedArrayCollection::EArrayType::FQuatType:
+			return DiffAttributes_Internal<FQuat4f>(CollectionA, CollectionB, ResultCollection, InAttributeName, InGroupName);
+		case FManagedArrayCollection::EArrayType::FIntArrayType:
+			return DiffAttributes_Internal<TSet<int32>>(CollectionA, CollectionB, ResultCollection, InAttributeName, InGroupName);
+		case FManagedArrayCollection::EArrayType::FUInt8Type:
+			return DiffAttributes_Internal<uint8>(CollectionA, CollectionB, ResultCollection, InAttributeName, InGroupName);
+		case FManagedArrayCollection::EArrayType::FVector2DArrayType:
+			return DiffAttributes_Internal<TArray<FVector2f>>(CollectionA, CollectionB, ResultCollection, InAttributeName, InGroupName);
+		case FManagedArrayCollection::EArrayType::FDoubleType:
+			return DiffAttributes_Internal<double>(CollectionA, CollectionB, ResultCollection, InAttributeName, InGroupName);
+		case FManagedArrayCollection::EArrayType::FIntVector4Type:
+			return DiffAttributes_Internal<FIntVector4>(CollectionA, CollectionB, ResultCollection, InAttributeName, InGroupName);
+		case FManagedArrayCollection::EArrayType::FVector3dType:
+			return DiffAttributes_Internal<FVector3d>(CollectionA, CollectionB, ResultCollection, InAttributeName, InGroupName);
+		case FManagedArrayCollection::EArrayType::FIntVector2Type:
+			return DiffAttributes_Internal<FIntVector2>(CollectionA, CollectionB, ResultCollection, InAttributeName, InGroupName);
+		case FManagedArrayCollection::EArrayType::FIntVector2ArrayType:
+			return DiffAttributes_Internal<TArray<FIntVector2>>(CollectionA, CollectionB, ResultCollection, InAttributeName, InGroupName);
+		case FManagedArrayCollection::EArrayType::FInt32ArrayType:
+			return DiffAttributes_Internal<TArray<int32>>(CollectionA, CollectionB, ResultCollection, InAttributeName, InGroupName);
+		case FManagedArrayCollection::EArrayType::FFloatArrayType:
+			return DiffAttributes_Internal<TArray<float>>(CollectionA, CollectionB, ResultCollection, InAttributeName, InGroupName);
+		case FManagedArrayCollection::EArrayType::FVector4fType:
+			return DiffAttributes_Internal<FVector4f>(CollectionA, CollectionB, ResultCollection, InAttributeName, InGroupName);
+		case FManagedArrayCollection::EArrayType::FFVectorArrayType:
+			return DiffAttributes_Internal<TArray<FVector3f>>(CollectionA, CollectionB, ResultCollection, InAttributeName, InGroupName);
+		case FManagedArrayCollection::EArrayType::FTransform3fType:
+			return DiffAttributes_Internal<FTransform3f>(CollectionA, CollectionB, ResultCollection, InAttributeName, InGroupName);
+		case FManagedArrayCollection::EArrayType::FIntVector3ArrayType:
+			return DiffAttributes_Internal<TArray<FIntVector3>>(CollectionA, CollectionB, ResultCollection, InAttributeName, InGroupName);
+		case FManagedArrayCollection::EArrayType::FVector4fArrayType:
+			return DiffAttributes_Internal<TArray<FVector4f>>(CollectionA, CollectionB, ResultCollection, InAttributeName, InGroupName);
+		case FManagedArrayCollection::EArrayType::FUintVector2Type:
+			return DiffAttributes_Internal<FUintVector2>(CollectionA, CollectionB, ResultCollection, InAttributeName, InGroupName);
+		
+		// Not trivially diffable
+		case FManagedArrayCollection::EArrayType::FStringType:
+		case FManagedArrayCollection::EArrayType::FGuidType:
+		case FManagedArrayCollection::EArrayType::FVectorArrayPointerType:
+		case FManagedArrayCollection::EArrayType::FVectorArrayUniquePointerType:
+		case FManagedArrayCollection::EArrayType::FUObjectArrayType:
+		case FManagedArrayCollection::EArrayType::FNameType:
+		case FManagedArrayCollection::EArrayType::FSoftObjectPathType:
+		case FManagedArrayCollection::EArrayType::FBoolType:
+
+		// Unused in PVE
+		case FManagedArrayCollection::EArrayType::FMeshSectionType:
+		case FManagedArrayCollection::EArrayType::FBoneNodeType:
+		case FManagedArrayCollection::EArrayType::FFImplicitObject3PointerType:
+		case FManagedArrayCollection::EArrayType::FFImplicitObject3UniquePointerType:
+		case FManagedArrayCollection::EArrayType::FFImplicitObject3SerializablePtrType:
+		case FManagedArrayCollection::EArrayType::FFBVHParticlesFloat3PointerType:
+		case FManagedArrayCollection::EArrayType::FFBVHParticlesFloat3UniquePointerType:
+		case FManagedArrayCollection::EArrayType::FTPBDRigidParticleHandle3fPtrType:
+		case FManagedArrayCollection::EArrayType::FTPBDGeometryCollectionParticleHandle3fPtrType:
+		case FManagedArrayCollection::EArrayType::FTGeometryParticle3fUniquePtrType:
+		case FManagedArrayCollection::EArrayType::FFImplicitObject3ThreadSafeSharedPointerType:
+		case FManagedArrayCollection::EArrayType::FFImplicitObject3SharedPointerType:
+		case FManagedArrayCollection::EArrayType::FTPBDRigidClusteredParticleHandle3fPtrType:
+		case FManagedArrayCollection::EArrayType::FFConvexUniquePtrType:
+		case FManagedArrayCollection::EArrayType::FTPBDRigidParticle3fUniquePtrType:
+		case FManagedArrayCollection::EArrayType::FFImplicitObjectRefCountedPtrType:
+		case FManagedArrayCollection::EArrayType::FFConvexRefCountedPtrType:
+		case FManagedArrayCollection::EArrayType::FPMatrix33dType:
+		case FManagedArrayCollection::EArrayType::FLinearCurveType:
+		case FManagedArrayCollection::EArrayType::FFVector3fNestedArrayType:
+			return;
+		default:
+			ensureAlwaysMsgf(false, TEXT("Unknown attribute type"));
+			break;
+		}
+	}
+};
+
+namespace PVGrowthDataDiffPins
+{
+	const static FName InputA = FName(TEXT("A"));
+	const static FName InputB = FName(TEXT("B"));
+};
+
+#if WITH_EDITOR
+FName UPVGrowthDataDiffSettings::GetDefaultNodeName() const
+{ 
+	return FName(TEXT("GrowthDataDiff")); 
+}
+
+FText UPVGrowthDataDiffSettings::GetDefaultNodeTitle() const
+{ 
+	return LOCTEXT("NodeTitle", "Growth Data Diff"); 
+}
+
+FText UPVGrowthDataDiffSettings::GetNodeTooltipText() const
+{ 
+	return LOCTEXT("NodeTooltip", "Diff two sets of grwoth data");
+}
+#endif
+
+UPVGrowthDataDiffSettings::UPVGrowthDataDiffSettings()
+{
+#if WITH_EDITOR
+	if (HasAnyFlags(RF_ClassDefaultObject))
+	{
+		bExposeToLibrary = PV::Utilities::DebugModeEnabled();
+	}
+#endif
+}
+ 
+TArray<FPCGPinProperties> UPVGrowthDataDiffSettings::InputPinProperties() const
+{
+	TArray<FPCGPinProperties> Properties;
+
+	FPCGPinProperties& PinA = Properties.Emplace_GetRef(PVGrowthDataDiffPins::InputA, GetInputPinTypeIdentifier());
+	PinA.SetRequiredPin();
+	PinA.SetAllowMultipleConnections(false);
+	PinA.bAllowMultipleData = false;
+
+	FPCGPinProperties& PinB = Properties.Emplace_GetRef(PVGrowthDataDiffPins::InputB, GetInputPinTypeIdentifier());
+	PinB.SetRequiredPin();
+	PinB.SetAllowMultipleConnections(false);
+	PinB.bAllowMultipleData = false;
+
+	return Properties;
+}
+
+FPCGElementPtr UPVGrowthDataDiffSettings::CreateElement() const
+{
+	return MakeShared<FPVGrowthDataDiffElement>();
+}
+
+FPCGDataTypeIdentifier UPVGrowthDataDiffSettings::GetInputPinTypeIdentifier() const
+{
+	return FPCGDataTypeIdentifier{ FPVDataTypeInfoGrowth::AsId() };
+}
+ 
+FPCGDataTypeIdentifier UPVGrowthDataDiffSettings::GetOutputPinTypeIdentifier() const
+{
+	return FPCGDataTypeIdentifier{ FPVDataTypeInfoGrowth::AsId() };
+}
+ 
+bool FPVGrowthDataDiffElement::ExecuteInternal(FPCGContext* InContext) const
+{
+	TRACE_CPUPROFILER_EVENT_SCOPE(FPVGrowthDataDiffElement::Execute);
+ 
+	check(InContext);
+ 
+	const UPVGrowthDataDiffSettings* Settings = InContext->GetInputSettings<UPVGrowthDataDiffSettings>();
+	check(Settings);
+	
+	const TArray<FPCGTaggedData> PinInputA = InContext->InputData.GetInputsByPin(PVGrowthDataDiffPins::InputA);
+	const TArray<FPCGTaggedData> PinInputB = InContext->InputData.GetInputsByPin(PVGrowthDataDiffPins::InputB);
+
+	if (PinInputA.IsEmpty() || PinInputB.IsEmpty())
+	{
+		return true;
+	}
+
+	const UPVGrowthData* GrowthDataA = Cast<UPVGrowthData>(PinInputA[0].Data);
+	const UPVGrowthData* GrowthDataB = Cast<UPVGrowthData>(PinInputB[0].Data);
+	if (!GrowthDataA || !GrowthDataB)
+	{
+		PCGLog::InputOutput::LogInvalidInputDataError(InContext);
+		return true;
+	}
+
+	const FManagedArrayCollection& CollectionA = GrowthDataA->GetCollection();
+	const FManagedArrayCollection& CollectionB = GrowthDataB->GetCollection();
+
+	if (!PV::Utilities::IsValidGrowthData(CollectionA) || !PV::Utilities::IsValidGrowthData(CollectionB))
+	{
+		PCGLog::InputOutput::LogInvalidInputDataError(InContext);
+		return true;
+	}
+
+	if (!PV::DiffHelpers::HasSameGrowthDataHierarchy(CollectionA, CollectionB))
+	{
+		PCGLog::LogErrorOnGraph(LOCTEXT("DifferentHierarchy", "Plant hierarchy is not the same"), InContext);
+		return true;
+	}
+
+	FManagedArrayCollection DiffCollection;
+
+	// Copy core hierarchy attributes
+	PV::FPointPositionAttribute::CopyAttribute(CollectionA, DiffCollection);
+	PV::FBranchPointsAttribute::CopyAttribute(CollectionA, DiffCollection);
+	PV::FBranchParentNumberAttribute::CopyAttribute(CollectionA, DiffCollection);
+	PV::FBranchParentsAttribute::CopyAttribute(CollectionA, DiffCollection);
+	PV::FBranchNumberAttribute::CopyAttribute(CollectionA, DiffCollection);
+	PV::FBranchChildrenAttribute::CopyAttribute(CollectionA, DiffCollection);
+	PV::FBranchHierarchyNumberAttribute::CopyAttribute(CollectionA, DiffCollection);
+	PV::FBranchPlantNumberAttribute::CopyAttribute(CollectionA, DiffCollection);
+
+	// Diff required attributes
+	const static TArray<FManagedArrayCollection::FManagedType> AttributesToDiff = {
+		PV::FBranchSourceBudNumberAttribute::ManagedType,
+		PV::FPointLengthFromRootAttribute::ManagedType,
+		PV::FPointLengthFromSeedAttribute::ManagedType,
+		PV::FPointScaleGradientAttribute::ManagedType,
+		PV::FPointHullGradientAttribute::ManagedType,
+		PV::FPointMainTrunkGradientAttribute::ManagedType,
+		PV::FPointGroundGradientAttribute::ManagedType,
+		PV::FPointScaleAttribute::ManagedType,
+		PV::FBudDirectionAttribute::ManagedType,
+		PV::FBudStatusAttribute::ManagedType,
+		PV::FBudLateralMeristemAttribute::ManagedType,
+		PV::FPointBranchGradientAttribute::ManagedType,
+		PV::FBudHormoneLevelsAttribute::ManagedType,
+		PV::FPointPlantGradientAttribute::ManagedType,
+		PV::FPointBudNumberAttribute::ManagedType,
+		PV::FPointNjordPixelIndexAttribute::ManagedType,
+		PV::FBudLightDetectedAttribute::ManagedType,
+		PV::FBudDevelopmentAttribute::ManagedType,
+	};
+
+	for (const auto& ManagedType : AttributesToDiff)
+	{
+		if (CollectionA.GetAttributeType(ManagedType.Name, ManagedType.Group) != ManagedType.Type
+			|| CollectionB.GetAttributeType(ManagedType.Name, ManagedType.Group) != ManagedType.Type)
+		{
+			continue;
+		}
+
+		PV::DiffHelpers::DiffAttributes(CollectionA, CollectionB, DiffCollection, ManagedType.Name, ManagedType.Group);
+	}
+
+	UPVGrowthData* OutGrowthData = FPCGContext::NewObject_AnyThread<UPVGrowthData>(InContext);
+	OutGrowthData->Initialize(MoveTemp(DiffCollection));
+	InContext->OutputData.TaggedData.Emplace(OutGrowthData);
+
+	return true; 
+}
+
+#undef LOCTEXT_NAMESPACE
\ No newline at end of file
diff --git a/Engine/Plugins/Experimental/ProceduralVegetationEditor/Source/ProceduralVegetation/Private/Nodes/Development/PVGrowthDataDiffSettings.h b/Engine/Plugins/Experimental/ProceduralVegetationEditor/Source/ProceduralVegetation/Private/Nodes/Development/PVGrowthDataDiffSettings.h
new file mode 100644
--- /dev/null
+++ b/Engine/Plugins/Experimental/ProceduralVegetationEditor/Source/ProceduralVegetation/Private/Nodes/Development/PVGrowthDataDiffSettings.h
@@ -0,0 +1,38 @@
+// Copyright Epic Games, Inc. All Rights Reserved.
+#pragma once
+
+#include "CoreMinimal.h"
+
+#include "Nodes/PVBaseSettings.h"
+
+#include "PVGrowthDataDiffSettings.generated.h"
+
+UCLASS(BlueprintType, ClassGroup = (Procedural))
+class UPVGrowthDataDiffSettings : public UPVBaseSettings
+{
+	GENERATED_BODY()
+
+public:
+	UPVGrowthDataDiffSettings();
+
+#if WITH_EDITOR
+	virtual FName GetDefaultNodeName() const override;
+	virtual FText GetDefaultNodeTitle() const override;
+	virtual FText GetNodeTooltipText() const override;
+	virtual FLinearColor GetNodeTitleColor() const override { return FLinearColor(FColor::Black); }
+
+	virtual TArray<EPVRenderType> GetDefaultRenderType() const override { return TArray{ EPVRenderType::PointData }; }
+#endif
+
+protected:
+	virtual TArray<FPCGPinProperties> InputPinProperties() const override;
+	virtual FPCGDataTypeIdentifier GetInputPinTypeIdentifier() const override;
+	virtual FPCGDataTypeIdentifier GetOutputPinTypeIdentifier() const override;
+	virtual FPCGElementPtr CreateElement() const override;
+};
+
+class FPVGrowthDataDiffElement : public FPVBaseElement
+{
+protected:
+	virtual bool ExecuteInternal(FPCGContext* InContext) const override;
+};
diff --git a/Engine/Plugins/Experimental/ProceduralVegetationEditor/Source/ProceduralVegetation/Private/Nodes/Development/PVGrowthDataJsonImporterSettings.cpp b/Engine/Plugins/Experimental/ProceduralVegetationEditor/Source/ProceduralVegetation/Private/Nodes/Development/PVGrowthDataJsonImporterSettings.cpp
new file mode 100644
--- /dev/null
+++ b/Engine/Plugins/Experimental/ProceduralVegetationEditor/Source/ProceduralVegetation/Private/Nodes/Development/PVGrowthDataJsonImporterSettings.cpp
@@ -0,0 +1,138 @@
+// Copyright Epic Games, Inc. All Rights Reserved.
+ 
+#include "PVGrowthDataJsonImporterSettings.h"
+ 
+#include "ProceduralVegetationModule.h"
+ 
+#include "DataTypes/PVGrowthData.h"
+#include "GeometryCollection/ManagedArrayCollection.h"
+#include "Helpers/PVJSONHelper.h"
+#include "Utils/PVAttributes.h"
+#include "Helpers/PVAttributesHelper.h"
+
+#define LOCTEXT_NAMESPACE "PVGrowthDataJsonImporterSettings"
+
+#if WITH_EDITOR
+FName UPVGrowthDataJsonImporterSettings::GetDefaultNodeName() const
+{ 
+	return FName(TEXT("GrowthDataJsonImporter")); 
+}
+
+FText UPVGrowthDataJsonImporterSettings::GetDefaultNodeTitle() const
+{ 
+	return LOCTEXT("NodeTitle", "Growth Data Json Importer"); 
+}
+
+FText UPVGrowthDataJsonImporterSettings::GetNodeTooltipText() const
+{ 
+	return LOCTEXT("NodeTooltip", "Imports growth data from a json file");
+}
+#endif
+
+UPVGrowthDataJsonImporterSettings::UPVGrowthDataJsonImporterSettings()
+{
+#if WITH_EDITOR
+	if (HasAnyFlags(RF_ClassDefaultObject))
+	{
+		bExposeToLibrary = PV::Utilities::DebugModeEnabled();
+	}
+#endif
+}
+ 
+TArray<FPCGPinProperties> UPVGrowthDataJsonImporterSettings::InputPinProperties() const
+{
+	return TArray<FPCGPinProperties>();
+}
+
+FPCGElementPtr UPVGrowthDataJsonImporterSettings::CreateElement() const
+{
+	return MakeShared<FPVGrowthDataJsonImporterElement>();
+}
+ 
+FPCGDataTypeIdentifier UPVGrowthDataJsonImporterSettings::GetOutputPinTypeIdentifier() const
+{
+	return FPCGDataTypeIdentifier{ FPVDataTypeInfoGrowth::AsId() };
+}
+ 
+bool FPVGrowthDataJsonImporterElement::ExecuteInternal(FPCGContext* InContext) const
+{
+	TRACE_CPUPROFILER_EVENT_SCOPE(FPVGrowthDataJsonImporterElement::Execute);
+ 
+	check(InContext);
+ 
+	const UPVGrowthDataJsonImporterSettings* Settings = InContext->GetInputSettings<UPVGrowthDataJsonImporterSettings>();
+	check(Settings);
+	
+	FString ErrorMessage;
+	FManagedArrayCollection Collection;
+	bool bSuccess = PV::LoadGrowthDataJsonToCollection(Collection, Settings->GrowthDataFile.FilePath, ErrorMessage);
+	if (!bSuccess)
+	{
+		UE_LOG(LogProceduralVegetation, Error, TEXT("%s"), *ErrorMessage);
+		PCGLog::InputOutput::LogInvalidInputDataError(InContext);
+		return true;
+	}
+
+	if (!PV::FPointScaleAttribute::HasAttribute(Collection))
+	{
+		PV::FPointScaleAttribute::AddAttribute(Collection, 0);
+	}
+	
+	if (!PV::FPointLengthFromRootAttribute::HasAttribute(Collection)
+		|| !PV::FPointLengthFromSeedAttribute::HasAttribute(Collection))
+	{
+		PV::FPointLengthFromRootAttribute::AddAttribute(Collection);
+		PV::FPointLengthFromSeedAttribute::AddAttribute(Collection);
+		PV::AttributesHelper::ComputeLengthFromRoot(Collection);
+	}
+
+	if (!PV::FPointGroundGradientAttribute::HasAttribute(Collection))
+	{
+		PV::FPointGroundGradientAttribute::AddAttribute(Collection);
+		PV::AttributesHelper::ComputePointGroundGradient(Collection);
+	}
+
+	if (!PV::FPointScaleGradientAttribute::HasAttribute(Collection))
+	{
+		PV::FPointScaleGradientAttribute::AddAttribute(Collection);
+		PV::AttributesHelper::ComputePointScaleGradient(Collection);
+	}
+
+	if (!PV::FPointPlantGradientAttribute::HasAttribute(Collection) 
+		|| !PV::FPointBranchGradientAttribute::HasAttribute(Collection))
+	{
+		PV::FPointPlantGradientAttribute::AddAttribute(Collection);
+		PV::FPointBranchGradientAttribute::AddAttribute(Collection);
+		PV::AttributesHelper::ComputePointPlantAndBranchGradient(Collection);
+	}
+
+	if (!PV::FPointHullGradientAttribute::HasAttribute(Collection))
+	{
+		PV::FPointHullGradientAttribute::AddAttribute(Collection, 0);
+	}
+
+	if (!PV::FPointMainTrunkGradientAttribute::HasAttribute(Collection))
+	{
+		PV::FPointMainTrunkGradientAttribute::AddAttribute(Collection, 0);
+	}
+
+	if (!PV::FPointNjordPixelIndexAttribute::HasAttribute(Collection))
+	{
+		PV::FPointNjordPixelIndexAttribute::AddAttribute(Collection);
+		PV::AttributesHelper::ComputeNjordPixelIndex(Collection);
+	}
+
+	if (!PV::FBranchHierarchyNumberAttribute::HasAttribute(Collection))
+	{
+		PV::FBranchHierarchyNumberAttribute::AddAttribute(Collection);
+		PV::AttributesHelper::ComputeBranchHierarchyNumbers(Collection);
+	}
+
+	UPVGrowthData* OutGrowthData = FPCGContext::NewObject_AnyThread<UPVGrowthData>(InContext);
+	OutGrowthData->Initialize(MoveTemp(Collection));
+	InContext->OutputData.TaggedData.Emplace(OutGrowthData);
+
+	return true; 
+}
+
+#undef LOCTEXT_NAMESPACE
\ No newline at end of file
diff --git a/Engine/Plugins/Experimental/ProceduralVegetationEditor/Source/ProceduralVegetation/Private/Nodes/Development/PVGrowthDataJsonImporterSettings.h b/Engine/Plugins/Experimental/ProceduralVegetationEditor/Source/ProceduralVegetation/Private/Nodes/Development/PVGrowthDataJsonImporterSettings.h
new file mode 100644
--- /dev/null
+++ b/Engine/Plugins/Experimental/ProceduralVegetationEditor/Source/ProceduralVegetation/Private/Nodes/Development/PVGrowthDataJsonImporterSettings.h
@@ -0,0 +1,41 @@
+// Copyright Epic Games, Inc. All Rights Reserved.
+#pragma once
+
+#include "CoreMinimal.h"
+
+#include "Nodes/PVBaseSettings.h"
+
+#include "PVGrowthDataJsonImporterSettings.generated.h"
+
+UCLASS(BlueprintType, ClassGroup = (Procedural))
+class UPVGrowthDataJsonImporterSettings : public UPVBaseSettings
+{
+	GENERATED_BODY()
+
+public:
+	UPVGrowthDataJsonImporterSettings();
+
+#if WITH_EDITOR
+	virtual FName GetDefaultNodeName() const override;
+	virtual FText GetDefaultNodeTitle() const override;
+	virtual FText GetNodeTooltipText() const override;
+	virtual FLinearColor GetNodeTitleColor() const override { return FLinearColor(FColor::Black); }
+
+	virtual TArray<EPVRenderType> GetDefaultRenderType() const override { return TArray{ EPVRenderType::PointData }; }
+#endif
+
+protected:
+	virtual TArray<FPCGPinProperties> InputPinProperties() const override;
+	virtual FPCGDataTypeIdentifier GetOutputPinTypeIdentifier() const override;
+	virtual FPCGElementPtr CreateElement() const override;
+
+public:
+	UPROPERTY(EditAnywhere, Category = "Growth Data")
+	FFilePath GrowthDataFile;
+};
+
+class FPVGrowthDataJsonImporterElement : public FPVBaseElement
+{
+protected:
+	virtual bool ExecuteInternal(FPCGContext* InContext) const override;
+};
diff --git a/Engine/Plugins/Experimental/ProceduralVegetationEditor/Source/ProceduralVegetation/Private/Nodes/Development/PVMockTreeImporterSettings.cpp b/Engine/Plugins/Experimental/ProceduralVegetationEditor/Source/ProceduralVegetation/Private/Nodes/Development/PVMockTreeImporterSettings.cpp
new file mode 100644
--- /dev/null
+++ b/Engine/Plugins/Experimental/ProceduralVegetationEditor/Source/ProceduralVegetation/Private/Nodes/Development/PVMockTreeImporterSettings.cpp
@@ -0,0 +1,77 @@
+// Copyright Epic Games, Inc. All Rights Reserved.
+ 
+#include "PVMockTreeImporterSettings.h"
+ 
+#include "ProceduralVegetationModule.h"
+ 
+#include "DataTypes/PVGrowthData.h"
+#include "GeometryCollection/ManagedArrayCollection.h"
+#include "Tests/PVMockTree.h"
+ 
+#define LOCTEXT_NAMESPACE "PVMockTreeImporterSettings"
+
+#if WITH_EDITOR
+FName UPVMockTreeImporterSettings::GetDefaultNodeName() const 
+{ 
+	return FName(TEXT("ProceduralVegetationMockTreeImporter")); 
+}
+
+FText UPVMockTreeImporterSettings::GetDefaultNodeTitle() const 
+{ 
+	return LOCTEXT("NodeTitle", "Mock Tree Importer"); 
+}
+
+FText UPVMockTreeImporterSettings::GetNodeTooltipText() const 
+{ 
+	return LOCTEXT("NodeTooltip", "Create Procedural Vegetation Skeleton based of the Mock Tree"); 
+}
+#endif
+
+UPVMockTreeImporterSettings::UPVMockTreeImporterSettings()
+{
+#if WITH_EDITOR
+	if (HasAnyFlags(RF_ClassDefaultObject))
+	{
+		bExposeToLibrary = PV::Utilities::DebugModeEnabled();
+	}
+#endif
+}
+ 
+TArray<FPCGPinProperties> UPVMockTreeImporterSettings::InputPinProperties() const
+{
+	return TArray<FPCGPinProperties>();
+}
+ 
+FPCGElementPtr UPVMockTreeImporterSettings::CreateElement() const
+{
+	return MakeShared<FPVMockTreeImporterElement>();
+}
+ 
+FPCGDataTypeIdentifier UPVMockTreeImporterSettings::GetOutputPinTypeIdentifier() const
+{
+	return FPCGDataTypeIdentifier{ FPVDataTypeInfoGrowth::AsId() };
+}
+ 
+bool FPVMockTreeImporterElement::ExecuteInternal(FPCGContext* InContext) const
+{
+	TRACE_CPUPROFILER_EVENT_SCOPE(FPVMockTreeImporterElement::Execute);
+ 
+	check(InContext);
+ 
+	const UPVMockTreeImporterSettings* Settings = InContext->GetInputSettings<UPVMockTreeImporterSettings>();
+	check(Settings);
+	
+	FManagedArrayCollection Collection;
+#if WITH_DEV_AUTOMATION_TESTS
+	Collection = MoveTemp(*PVMockTreeCollection::CreateCollection());
+#endif
+
+	UPVGrowthData* OutManagedArrayCollectionData = FPCGContext::NewObject_AnyThread<UPVGrowthData>(InContext);
+	OutManagedArrayCollectionData->Initialize(MoveTemp(Collection));
+		
+	InContext->OutputData.TaggedData.Emplace(OutManagedArrayCollectionData);
+	
+	return true; 
+}
+
+#undef LOCTEXT_NAMESPACE
\ No newline at end of file
diff --git a/Engine/Plugins/Experimental/ProceduralVegetationEditor/Source/ProceduralVegetation/Private/Nodes/Development/PVMockTreeImporterSettings.h b/Engine/Plugins/Experimental/ProceduralVegetationEditor/Source/ProceduralVegetation/Private/Nodes/Development/PVMockTreeImporterSettings.h
new file mode 100644
--- /dev/null
+++ b/Engine/Plugins/Experimental/ProceduralVegetationEditor/Source/ProceduralVegetation/Private/Nodes/Development/PVMockTreeImporterSettings.h
@@ -0,0 +1,37 @@
+// Copyright Epic Games, Inc. All Rights Reserved.
+#pragma once
+
+#include "CoreMinimal.h"
+
+#include "Nodes/PVBaseSettings.h"
+
+#include "PVMockTreeImporterSettings.generated.h"
+
+UCLASS(BlueprintType, ClassGroup = (Procedural))
+class UPVMockTreeImporterSettings : public UPVBaseSettings
+{
+	GENERATED_BODY()
+
+public:
+	UPVMockTreeImporterSettings();
+
+#if WITH_EDITOR
+	virtual FName GetDefaultNodeName() const override;
+	virtual FText GetDefaultNodeTitle() const override;
+	virtual FText GetNodeTooltipText() const override;
+	virtual FLinearColor GetNodeTitleColor() const override { return FLinearColor(FColor::Black); }
+
+	virtual TArray<EPVRenderType> GetDefaultRenderType() const override { return TArray{ EPVRenderType::PointData }; }
+#endif
+
+protected:
+	virtual TArray<FPCGPinProperties> InputPinProperties() const override;
+	virtual FPCGDataTypeIdentifier GetOutputPinTypeIdentifier() const override;
+	virtual FPCGElementPtr CreateElement() const override;
+};
+
+class FPVMockTreeImporterElement : public FPVBaseElement
+{
+protected:
+	virtual bool ExecuteInternal(FPCGContext* InContext) const override;
+};
diff --git a/Engine/Plugins/Experimental/ProceduralVegetationEditor/Source/ProceduralVegetation/Private/Nodes/Development/PVRecomputeGrowthDataAttributesSettings.cpp b/Engine/Plugins/Experimental/ProceduralVegetationEditor/Source/ProceduralVegetation/Private/Nodes/Development/PVRecomputeGrowthDataAttributesSettings.cpp
new file mode 100644
--- /dev/null
+++ b/Engine/Plugins/Experimental/ProceduralVegetationEditor/Source/ProceduralVegetation/Private/Nodes/Development/PVRecomputeGrowthDataAttributesSettings.cpp
@@ -0,0 +1,80 @@
+// Copyright Epic Games, Inc. All Rights Reserved.
+ 
+#include "PVRecomputeGrowthDataAttributesSettings.h"
+ 
+#include "DataTypes/PVGrowthData.h"
+#include "Helpers/PVAttributesHelper.h"
+ 
+#define LOCTEXT_NAMESPACE "PVRecomputeGrowthDataAttributesSettings"
+
+#if WITH_EDITOR
+FName UPVRecomputeGrowthDataAttributesSettings::GetDefaultNodeName() const
+{ 
+	return FName(TEXT("RecomputeGrowthDataAttributes")); 
+}
+
+FText UPVRecomputeGrowthDataAttributesSettings::GetDefaultNodeTitle() const
+{ 
+	return LOCTEXT("NodeTitle", "Recompute Growth Data Attributes"); 
+}
+
+FText UPVRecomputeGrowthDataAttributesSettings::GetNodeTooltipText() const
+{ 
+	return LOCTEXT("NodeTooltip", "Recomputes all growth data attributes using the plant skeleton");
+}
+#endif
+
+UPVRecomputeGrowthDataAttributesSettings::UPVRecomputeGrowthDataAttributesSettings()
+{
+#if WITH_EDITOR
+	if (HasAnyFlags(RF_ClassDefaultObject))
+	{
+		bExposeToLibrary = PV::Utilities::DebugModeEnabled();
+	}
+#endif
+}
+ 
+FPCGElementPtr UPVRecomputeGrowthDataAttributesSettings::CreateElement() const
+{
+	return MakeShared<FPVRecomputeGrowthDataAttributesElement>();
+}
+
+FPCGDataTypeIdentifier UPVRecomputeGrowthDataAttributesSettings::GetInputPinTypeIdentifier() const
+{
+	return FPCGDataTypeIdentifier{ FPVDataTypeInfoGrowth::AsId() };
+}
+ 
+FPCGDataTypeIdentifier UPVRecomputeGrowthDataAttributesSettings::GetOutputPinTypeIdentifier() const
+{
+	return FPCGDataTypeIdentifier{ FPVDataTypeInfoGrowth::AsId() };
+}
+ 
+bool FPVRecomputeGrowthDataAttributesElement::ExecuteInternal(FPCGContext* InContext) const
+{
+	TRACE_CPUPROFILER_EVENT_SCOPE(FPVRecomputeGrowthDataAttributesElement::Execute);
+ 
+	check(InContext);
+ 
+	for (const FPCGTaggedData& Input : InContext->InputData.GetInputsByPin(PCGPinConstants::DefaultInputLabel))
+	{
+		if (const UPVGrowthData* InputData = Cast<UPVGrowthData>(Input.Data))
+		{
+			FManagedArrayCollection OutCollection = InputData->GetCollection();
+
+			PV::AttributesHelper::RecomputeAllGrowthDataAttributes(OutCollection);
+
+			UPVGrowthData* OutGrowthData = FPCGContext::NewObject_AnyThread<UPVGrowthData>(InContext);
+			OutGrowthData->Initialize(MoveTemp(OutCollection));
+			InContext->OutputData.TaggedData.Emplace(OutGrowthData);
+		}
+		else
+		{
+			PCGLog::InputOutput::LogInvalidInputDataError(InContext);
+			return true;
+		}
+	}
+	
+	return true; 
+}
+
+#undef LOCTEXT_NAMESPACE
\ No newline at end of file
diff --git a/Engine/Plugins/Experimental/ProceduralVegetationEditor/Source/ProceduralVegetation/Private/Nodes/Development/PVRecomputeGrowthDataAttributesSettings.h b/Engine/Plugins/Experimental/ProceduralVegetationEditor/Source/ProceduralVegetation/Private/Nodes/Development/PVRecomputeGrowthDataAttributesSettings.h
new file mode 100644
--- /dev/null
+++ b/Engine/Plugins/Experimental/ProceduralVegetationEditor/Source/ProceduralVegetation/Private/Nodes/Development/PVRecomputeGrowthDataAttributesSettings.h
@@ -0,0 +1,37 @@
+// Copyright Epic Games, Inc. All Rights Reserved.
+#pragma once
+
+#include "CoreMinimal.h"
+
+#include "Nodes/PVBaseSettings.h"
+
+#include "PVRecomputeGrowthDataAttributesSettings.generated.h"
+
+UCLASS(BlueprintType, ClassGroup = (Procedural))
+class UPVRecomputeGrowthDataAttributesSettings : public UPVBaseSettings
+{
+	GENERATED_BODY()
+
+public:
+	UPVRecomputeGrowthDataAttributesSettings();
+
+#if WITH_EDITOR
+	virtual FName GetDefaultNodeName() const override;
+	virtual FText GetDefaultNodeTitle() const override;
+	virtual FText GetNodeTooltipText() const override;
+	virtual FLinearColor GetNodeTitleColor() const override { return FLinearColor(FColor::Black); }
+
+	virtual TArray<EPVRenderType> GetDefaultRenderType() const override { return TArray{ EPVRenderType::PointData }; }
+#endif
+
+protected:
+	virtual FPCGDataTypeIdentifier GetInputPinTypeIdentifier() const override;
+	virtual FPCGDataTypeIdentifier GetOutputPinTypeIdentifier() const override;
+	virtual FPCGElementPtr CreateElement() const override;
+};
+
+class FPVRecomputeGrowthDataAttributesElement : public FPVBaseElement
+{
+protected:
+	virtual bool ExecuteInternal(FPCGContext* InContext) const override;
+};
diff --git a/Engine/Plugins/Experimental/ProceduralVegetationEditor/Source/ProceduralVegetation/Private/Nodes/PVImporterSettings.cpp b/Engine/Plugins/Experimental/ProceduralVegetationEditor/Source/ProceduralVegetation/Private/Nodes/PVImporterSettings.cpp
--- a/Engine/Plugins/Experimental/ProceduralVegetationEditor/Source/ProceduralVegetation/Private/Nodes/PVImporterSettings.cpp
+++ b/Engine/Plugins/Experimental/ProceduralVegetationEditor/Source/ProceduralVegetation/Private/Nodes/PVImporterSettings.cpp
@@ -63,7 +63,7 @@ bool FPVImporterElement::ExecuteInternal(FPCGContext* InContext) const
 		return true;
 	}
  
-	if (!PV::Utilities::IsValidPVData(Collection))
+	if (!PV::Utilities::IsValidGrowthData(Collection))
 	{
 		PCGLog::LogErrorOnGraph(FText::FromString(FString::Format(TEXT("Invalid data in the imported json {0}"), { SkeletonFilePath})), InContext);
 		return true;
diff --git a/Engine/Plugins/Experimental/ProceduralVegetationEditor/Source/ProceduralVegetation/Private/Nodes/PVMockTreeImporterSettings.cpp b/Engine/Plugins/Experimental/ProceduralVegetationEditor/Source/ProceduralVegetation/Private/Nodes/PVMockTreeImporterSettings.cpp
deleted file mode 100644
diff --git a/Engine/Plugins/Experimental/ProceduralVegetationEditor/Source/ProceduralVegetation/Private/Nodes/PVMockTreeImporterSettings.h b/Engine/Plugins/Experimental/ProceduralVegetationEditor/Source/ProceduralVegetation/Private/Nodes/PVMockTreeImporterSettings.h
deleted file mode 100644
diff --git a/Engine/Plugins/Experimental/ProceduralVegetationEditor/Source/ProceduralVegetation/Private/Nodes/PVPresetLoaderSettings.cpp b/Engine/Plugins/Experimental/ProceduralVegetationEditor/Source/ProceduralVegetation/Private/Nodes/PVPresetLoaderSettings.cpp
--- a/Engine/Plugins/Experimental/ProceduralVegetationEditor/Source/ProceduralVegetation/Private/Nodes/PVPresetLoaderSettings.cpp
+++ b/Engine/Plugins/Experimental/ProceduralVegetationEditor/Source/ProceduralVegetation/Private/Nodes/PVPresetLoaderSettings.cpp
@@ -59,7 +59,7 @@ TArray<FPCGPinProperties> UPVPresetLoaderSettings::OutputPinProperties() const
 	{
 		for (const auto& [VariantName, VariantData] : Preset->Variants)
 		{
-			if (PV::Utilities::IsValidPVData(VariantData))
+			if (PV::Utilities::IsValidGrowthData(VariantData))
 			{
 				FPCGPinProperties& Pin = Properties.Emplace_GetRef(*VariantName, GetOutputPinTypeIdentifier());
 				Pin.bAllowMultipleData = false;
@@ -116,7 +116,7 @@ bool FPVPresetLoaderElement::ExecuteInternal(FPCGContext* InContext) const
 	bool bValidDataExist = false;
 	for (const auto& [VariantName, VariantData] : Settings->Preset->Variants)
 	{
-		if (PV::Utilities::IsValidPVData(VariantData))
+		if (PV::Utilities::IsValidGrowthData(VariantData))
 		{
 			UPVGrowthData* OutVariantData = FPCGContext::NewObject_AnyThread<UPVGrowthData>(InContext);
 			OutVariantData->Initialize(CopyTemp(VariantData));
diff --git a/Engine/Plugins/Experimental/ProceduralVegetationEditor/Source/ProceduralVegetation/Private/Nodes/PVSubgraphSettings.cpp b/Engine/Plugins/Experimental/ProceduralVegetationEditor/Source/ProceduralVegetation/Private/Nodes/PVSubgraphSettings.cpp
--- a/Engine/Plugins/Experimental/ProceduralVegetationEditor/Source/ProceduralVegetation/Private/Nodes/PVSubgraphSettings.cpp
+++ b/Engine/Plugins/Experimental/ProceduralVegetationEditor/Source/ProceduralVegetation/Private/Nodes/PVSubgraphSettings.cpp
@@ -5,8 +5,6 @@
 #include "ProceduralVegetation.h"
 #include "ProceduralVegetationModule.h"
 #include "Helpers/PCGSettingsHelpers.h"
-#include "Helpers/PVFoliageJSONHelper.h"
-#include "Helpers/PVJSONHelper.h"
 
 #define LOCTEXT_NAMESPACE "PVSubgraph"
 
diff --git a/Engine/Plugins/Experimental/ProceduralVegetationEditor/Source/ProceduralVegetation/Public/Helpers/PVJSONHelper.h b/Engine/Plugins/Experimental/ProceduralVegetationEditor/Source/ProceduralVegetation/Public/Helpers/PVJSONHelper.h
deleted file mode 100644
diff --git a/Engine/Plugins/Experimental/ProceduralVegetationEditor/Source/ProceduralVegetation/Public/Helpers/PVUtilities.h b/Engine/Plugins/Experimental/ProceduralVegetationEditor/Source/ProceduralVegetation/Public/Helpers/PVUtilities.h
--- a/Engine/Plugins/Experimental/ProceduralVegetationEditor/Source/ProceduralVegetation/Public/Helpers/PVUtilities.h
+++ b/Engine/Plugins/Experimental/ProceduralVegetationEditor/Source/ProceduralVegetation/Public/Helpers/PVUtilities.h
@@ -34,7 +34,7 @@ namespace PV::Utilities
 	int32 PROCEDURALVEGETATION_API GetMeshTriangles(const FString InMeshPath);
 	FLinearColor PROCEDURALVEGETATION_API GetRandomHueColor(float Alpha);
 
-	bool IsValidPVData(const FManagedArrayCollection& Collection);
+	bool IsValidGrowthData(const FManagedArrayCollection& Collection);
 
 	template <typename T>
 	TArray<T> LerpArrayElements(const TArray<T>& InArray1, const TArray<T>& InArray2, const float Alpha);
