[PVE] Removed function definitions from PVJSONHelper header.

#Synced-CL 50934585

diff --git a/Engine/Plugins/Experimental/ProceduralVegetationEditor/Source/ProceduralVegetation/Private/Helpers/PVJSONHelper.cpp b/Engine/Plugins/Experimental/ProceduralVegetationEditor/Source/ProceduralVegetation/Private/Helpers/PVJSONHelper.cpp
new file mode 100644
--- /dev/null
+++ b/Engine/Plugins/Experimental/ProceduralVegetationEditor/Source/ProceduralVegetation/Private/Helpers/PVJSONHelper.cpp
@@ -0,0 +1,653 @@
+﻿// Copyright Epic Games, Inc. All Rights Reserved.
+ 
+#include "PVJSONHelper.h"
+#include "Dom/JsonObject.h"
+#include "Dom/JsonValue.h"
+#include "Facades/PVAttributesNames.h"
+#include "GeometryCollection/ManagedArray.h"
+#include "GeometryCollection/ManagedArrayCollection.h"
+#include "Misc/FileHelper.h"
+#include "Serialization/JsonSerializer.h"
+#include "Facades/PVBranchFacade.h"
+#include "Facades/PVMetaInfoFacade.h"
+#include "Facades/PVFoliageFacade.h"
+#include "Facades/PVProfileFacade.h"
+#include "Misc/PackageName.h"
+#include "Misc/Paths.h"
+
+namespace PV::JSON
+{
+	void FillAttributes(FManagedArrayCollection& Collection, const FName Group, const TSharedPtr<FJsonObject>& AttributesObject)
+	{
+		for (const auto& [Key, Attribute] : AttributesObject->Values)
+		{
+			const TSharedPtr<FJsonObject>& AttributeObject = Attribute->AsObject();
+
+			const FString Type = AttributeObject->GetStringField(TEXT("type"));
+			const int32 ElementCount = AttributeObject->GetNumberField(TEXT("size"));
+			const bool IsArray = AttributeObject->GetBoolField(TEXT("isArray"));
+			const TArray<TSharedPtr<FJsonValue>>& Values = AttributeObject->GetArrayField(TEXT("values"));
+
+			if (IsArray)
+			{
+				if (ElementCount == 1)
+				{
+					if (Type == "int")
+					{
+						TManagedArray<TArray<int32>>& Attr = Collection.AddAttribute<TArray<int32>>(*Key, Group);
+						for (int32 Index = 0; Index < Values.Num(); ++Index)
+						{
+							TArray<int32>& CurrentAttrib = Attr[Index];
+							const TArray<TSharedPtr<FJsonValue>>& AttribValues = Values[Index]->AsArray();
+							for (const TSharedPtr<FJsonValue>& AttribValue : AttribValues)
+							{
+								CurrentAttrib.Add(AttribValue->AsNumber());
+							}
+						}
+					}
+					else if (Type == "float")
+					{
+						TManagedArray<TArray<float>>& Attr = Collection.AddAttribute<TArray<float>>(*Key, Group);
+						for (int32 Index = 0; Index < Values.Num(); ++Index)
+						{
+							TArray<float>& CurrentAttrib = Attr[Index];
+							const TArray<TSharedPtr<FJsonValue>>& AttribValues = Values[Index]->AsArray();
+							for (const TSharedPtr<FJsonValue>& AttribValue : AttribValues)
+							{
+								CurrentAttrib.Add(AttribValue->AsNumber());
+							}
+						}
+					}
+					else if (Type == "string")
+					{
+						// TODO: Figure out how to manage string arrays
+					}
+				}
+				else if (ElementCount == 3)
+				{
+					if (Type == "int")
+					{
+						TManagedArray<TArray<FIntVector>>& Attr = Collection.AddAttribute<TArray<FIntVector>>(*Key, Group);
+						for (int32 Index = 0; Index < Values.Num(); ++Index)
+						{
+							TArray<FIntVector>& CurrentAttrib = Attr[Index];
+							const TArray<TSharedPtr<FJsonValue>>& AttribValues = Values[Index]->AsArray();
+							for (int ElemIndex = 0; ElemIndex < AttribValues.Num(); ElemIndex += ElementCount)
+							{
+								CurrentAttrib.Emplace(AttribValues[ElemIndex]->AsNumber(), AttribValues[ElemIndex + 2]->AsNumber(),
+									AttribValues[ElemIndex + 1]->AsNumber());
+							}
+						}
+					}
+					else if (Type == "float")
+					{
+						TManagedArray<TArray<FVector3f>>& Attr = Collection.AddAttribute<TArray<FVector3f>>(*Key, Group);
+						for (int32 Index = 0; Index < Values.Num(); ++Index)
+						{
+							TArray<FVector3f>& CurrentAttrib = Attr[Index];
+							const TArray<TSharedPtr<FJsonValue>>& AttribValues = Values[Index]->AsArray();
+							for (int ElemIndex = 0; ElemIndex < AttribValues.Num(); ElemIndex += ElementCount)
+							{
+								CurrentAttrib.Emplace(AttribValues[ElemIndex]->AsNumber(), AttribValues[ElemIndex + 2]->AsNumber(),
+									AttribValues[ElemIndex + 1]->AsNumber());
+							}
+						}
+					}
+				}
+			}
+			else
+			{
+				if (ElementCount == 1)
+				{
+					if (Type == "int")
+					{
+						TManagedArray<int32>& Attr = Collection.AddAttribute<int32>(*Key, Group);
+						for (int32 Index = 0; Index < Values.Num(); ++Index)
+						{
+							Attr[Index] = Values[Index]->AsNumber();
+						}
+					}
+					else if (Type == "float")
+					{
+						TManagedArray<float>& Attr = Collection.AddAttribute<float>(*Key, Group);
+						for (int32 Index = 0; Index < Values.Num(); ++Index)
+						{
+							Attr[Index] = Values[Index]->AsNumber();
+						}
+					}
+					else if (Type == "string")
+					{
+						TManagedArray<FString>& Attr = Collection.AddAttribute<FString>(*Key, Group);
+						for (int32 Index = 0; Index < Values.Num(); ++Index)
+						{
+							Attr[Index] = Values[Index]->AsString();
+						}
+					}
+				}
+				else if (ElementCount == 3)
+				{
+					if (Type == "int")
+					{
+						TManagedArray<FIntVector>& Attr = Collection.AddAttribute<FIntVector>(*Key, Group);
+						for (int32 Index = 0; Index < Values.Num(); ++Index)
+						{
+							const TArray<TSharedPtr<FJsonValue>>& VectorAttribute = Values[Index]->AsArray();
+							Attr[Index] = FIntVector(VectorAttribute[0]->AsNumber(), VectorAttribute[2]->AsNumber(), VectorAttribute[1]->AsNumber());
+						}
+					}
+					else if (Type == "float")
+					{
+						TManagedArray<FVector3f>& Attr = Collection.AddAttribute<FVector3f>(*Key, Group);
+						for (int32 Index = 0; Index < Values.Num(); ++Index)
+						{
+							const TArray<TSharedPtr<FJsonValue>>& VectorAttribute = Values[Index]->AsArray();
+							Attr[Index] = FVector3f(VectorAttribute[0]->AsNumber(), VectorAttribute[2]->AsNumber(), VectorAttribute[1]->AsNumber());
+						}
+					}
+				}
+			}
+		}
+	}
+
+	void FillDetailsAttributes(FManagedArrayCollection& Collection, const FName Group, const TSharedPtr<FJsonObject>& AttributesObject)
+	{
+		for (const auto& [Key, Attribute] : AttributesObject->Values)
+		{
+			const TSharedPtr<FJsonObject>& AttributeObject = Attribute->AsObject();
+
+			const FString Type = AttributeObject->GetStringField(TEXT("type"));
+			const bool IsArray = AttributeObject->GetBoolField(TEXT("isArray"));
+
+			if (IsArray)
+			{
+				const TArray<TSharedPtr<FJsonValue>>& Values = AttributeObject->GetArrayField(TEXT("value"));
+				if (Type == "int")
+				{
+					TManagedArray<TArray<int32>>& Attr = Collection.AddAttribute<TArray<int32>>(*Key, Group);
+					TArray<int32>& CurrentAttrib = Attr[0];
+					for (const TSharedPtr<FJsonValue>& AttribValue : Values)
+					{
+						CurrentAttrib.Add(AttribValue->AsNumber());
+					}
+				}
+				else if (Type == "float")
+				{
+					TManagedArray<TArray<float>>& Attr = Collection.AddAttribute<TArray<float>>(*Key, Group);
+					TArray<float>& CurrentAttrib = Attr[0];
+					for (const TSharedPtr<FJsonValue>& AttribValue : Values)
+					{
+						CurrentAttrib.Add(AttribValue->AsNumber());
+					}
+				}
+				else if (Type == "string")
+				{
+					// TODO: Figure out how to manage string arrays
+				}
+			}
+			else
+			{
+				if (Type == "int")
+				{
+					TManagedArray<int32>& Attr = Collection.AddAttribute<int32>(*Key, Group);
+					Attr[0] = AttributeObject->GetNumberField(TEXT("value"));
+				}
+				else if (Type == "float")
+				{
+					TManagedArray<float>& Attr = Collection.AddAttribute<float>(*Key, Group);
+					Attr[0] = AttributeObject->GetNumberField(TEXT("value"));
+				}
+				else if (Type == "string")
+				{
+					TManagedArray<FString>& Attr = Collection.AddAttribute<FString>(*Key, Group);
+					Attr[0] = AttributeObject->GetStringField(TEXT("value"));
+				}
+			}
+		}
+	}
+
+	void FillPlantProfilesData(FManagedArrayCollection& Collection, const FName Group, const TSharedPtr<FJsonObject>& AttributesObject)
+	{
+		Facades::FPlantProfileFacade Facade(Collection);
+
+		for (const auto& [Key, Attribute] : AttributesObject->Values)
+		{
+			if (const TSharedPtr<FJsonObject>& AttributeObject = Attribute->AsObject();
+				AttributeObject
+				&& AttributeObject->HasTypedField<EJson::Array>(TEXT("value"))
+				&& Key.StartsWith(TEXT("plantProfile_")))
+			{
+				const TArray<TSharedPtr<FJsonValue>>& Values = AttributeObject->GetArrayField(TEXT("value"));
+				TArray<float> Points;
+				Points.Reserve(Values.Num());
+
+				for (const TSharedPtr<FJsonValue>& AttribValue : Values)
+				{
+					if (double Point; AttribValue->TryGetNumber(Point))
+					{
+						Points.Add(static_cast<float>(AttribValue->AsNumber()));
+					}
+				}
+
+				Facade.AddProfileEntry(Points);
+			}
+		}
+	}
+
+	void FillFoliageData(FManagedArrayCollection& Collection, const TSharedPtr<FJsonObject>& PrimitiveAttributesObject, const FString& InPath)
+	{
+		const TSharedPtr<FJsonObject>& InstancerNameObject = PrimitiveAttributesObject->GetObjectField(TEXT("instancer_name"));
+		const TArray<TSharedPtr<FJsonValue>>& InstancerNameValues = InstancerNameObject->GetArrayField(TEXT("values"));
+
+		const TSharedPtr<FJsonObject>& InstancerPivotObject = PrimitiveAttributesObject->GetObjectField(TEXT("instancer_pivot"));
+		const TArray<TSharedPtr<FJsonValue>>& InstancerPivotValues = InstancerPivotObject->GetArrayField(TEXT("values"));
+
+		const TSharedPtr<FJsonObject>& InstancerUpVectorObject = PrimitiveAttributesObject->GetObjectField(TEXT("instancer_UP"));
+		const TArray<TSharedPtr<FJsonValue>>& InstancerUpVectorValues = InstancerUpVectorObject->GetArrayField(TEXT("values"));
+
+		const TSharedPtr<FJsonObject>& InstancerNormalVectorObject = PrimitiveAttributesObject->GetObjectField(TEXT("instancer_N"));
+		const TArray<TSharedPtr<FJsonValue>>& InstancerNormalVectorValues = InstancerNormalVectorObject->GetArrayField(TEXT("values"));
+
+		const TSharedPtr<FJsonObject>& InstancerScaleObject = PrimitiveAttributesObject->GetObjectField(TEXT("instancer_scale"));
+		const TArray<TSharedPtr<FJsonValue>>& InstancerScaleValues = InstancerScaleObject->GetArrayField(TEXT("values"));
+
+		const TSharedPtr<FJsonObject>& InstancerLFRObject = PrimitiveAttributesObject->GetObjectField(TEXT("instancer_LFR"));
+		const TArray<TSharedPtr<FJsonValue>>& InstancerLFRValues = InstancerLFRObject->GetArrayField(TEXT("values"));
+
+		// Compute unique names and number of elements for group
+		int32 TotalValues = 0;
+		int32 NextFoliageIndex = 0;
+		TMap<FString, int32> FoliageNamesToIndex;
+		TArray<FString> FoliageNames;
+		for (int32 i = 0; i < InstancerNameValues.Num(); ++i)
+		{
+			const TArray<TSharedPtr<FJsonValue>>& InstancerNames = InstancerNameValues[i]->AsArray();
+			for (int32 Index = 0; Index < InstancerNames.Num(); ++Index)
+			{
+				const FString& Name = InstancerNames[Index]->AsString();
+				if (!FoliageNamesToIndex.Contains(Name))
+				{
+					FoliageNamesToIndex.Emplace(Name, NextFoliageIndex++);
+					FoliageNames.Emplace(Name);
+				}
+			}
+			TotalValues += InstancerNames.Num();
+		}
+
+		Facades::FFoliageFacade Facade(Collection, TotalValues);
+
+		for (FString& FoliageName : FoliageNames)
+		{
+			FoliageName = FoliageName + '.' + FoliageName;
+		}
+		
+		Facade.SetFoliageNames(FoliageNames);
+
+		int32 FoliageEntryIndex = 0;
+		for (int32 Id = 0; Id < InstancerNameValues.Num(); ++Id)
+		{
+			const TArray<TSharedPtr<FJsonValue>>& InstancerNames = InstancerNameValues[Id]->AsArray();
+			const TArray<TSharedPtr<FJsonValue>>& PivotPoints = InstancerPivotValues[Id]->AsArray();
+			const TArray<TSharedPtr<FJsonValue>>& UpVectors = InstancerUpVectorValues[Id]->AsArray();
+			const TArray<TSharedPtr<FJsonValue>>& NormalVectors = InstancerNormalVectorValues[Id]->AsArray();
+			const TArray<TSharedPtr<FJsonValue>>& Scales = InstancerScaleValues[Id]->AsArray();
+			const TArray<TSharedPtr<FJsonValue>>& LFRs = InstancerLFRValues[Id]->AsArray();
+
+			TArray<int32> FoliageEntryIds;
+			for (int32 j = 0; j < InstancerNames.Num(); ++j)
+			{
+				const FString Name = InstancerNames[j]->AsString();
+				const FVector3f PivotPoint = FVector3f(
+					PivotPoints[j * 3]->AsNumber(),
+					PivotPoints[(j * 3) + 2]->AsNumber(),
+					PivotPoints[(j * 3) + 1]->AsNumber()) * 100.0f;
+				const FVector3f UpVector = FVector3f(
+					UpVectors[j * 3]->AsNumber(),
+					UpVectors[(j * 3) + 2]->AsNumber(),
+					UpVectors[(j * 3) + 1]->AsNumber());
+				const FVector3f NormalVector = FVector3f(
+					NormalVectors[j * 3]->AsNumber(),
+					NormalVectors[(j * 3) + 2]->AsNumber(),
+					NormalVectors[(j * 3) + 1]->AsNumber());
+				const float Scale = Scales[j]->AsNumber();
+				const float LengthFromRoot = LFRs[j]->AsNumber();
+
+				Facade.SetFoliageEntry(FoliageEntryIndex, {
+					.NameId = FoliageNamesToIndex[Name],
+					.BranchId = Id,
+					.PivotPoint = PivotPoint,
+					.UpVector = UpVector,
+					.NormalVector = NormalVector,
+					.Scale = Scale,
+					.LengthFromRoot = LengthFromRoot,
+				});
+
+				FoliageEntryIds.Add(FoliageEntryIndex);
+				FoliageEntryIndex++;
+			}
+
+			Facade.SetFoliageIdsArray(Id, FoliageEntryIds);
+		}
+	}
+
+	void SetFoliagePaths(FManagedArrayCollection& Collection,const FString& FilePath)
+	{
+		Facades::FFoliageFacade Facade(Collection);
+
+		FString AbsolutePath = FPaths::ConvertRelativePathToFull(FilePath);
+		FString PackagePath;
+
+		if (FPackageName::TryConvertFilenameToLongPackageName(AbsolutePath, PackagePath))
+		{
+			TArray<FString> FoliageNames = Facade.GetFoliageNames();
+			FString PackageFolder = FPackageName::GetLongPackagePath(PackagePath);
+			Facade.SetFoliagePath(PackageFolder);
+			for (FString& FoliageName : FoliageNames)
+			{
+				FoliageName = PackageFolder / FoliageName;
+			}
+			Facade.SetFoliageNames(FoliageNames);
+		}
+		
+	}
+
+	bool HasJsonFieldPath(const TSharedPtr<FJsonObject>& JsonObject, const FString& Path)
+	{
+		if (!JsonObject.IsValid())
+		{
+			return false;
+		}
+
+		TArray<FString> Keys;
+		Path.ParseIntoArray(Keys, TEXT("."), true);
+
+		TSharedPtr<FJsonObject> CurrentObject = JsonObject;
+
+		for (int32 i = 0; i < Keys.Num(); ++i)
+		{
+			const FString& CurrentKey = Keys[i];
+			if (i == Keys.Num() - 1)
+			{
+				if (!CurrentObject->HasField(CurrentKey))
+				{
+					return false;
+				}
+			}
+			else
+			{
+				const TSharedPtr<FJsonObject>* NextObject;
+				if (CurrentObject->TryGetObjectField(CurrentKey, NextObject))
+				{
+					CurrentObject = *NextObject;
+				}
+				else
+				{
+					return false;
+				}
+			}
+		}
+
+		return true;
+	}
+
+	bool LoadMegaPlantsJsonToCollection(FManagedArrayCollection& Collection, const FString& FilePath, FString& OutErrorMessage)
+	{
+		FString MegaPlantsData;
+		if (!FFileHelper::LoadFileToString(MegaPlantsData, *FilePath))
+		{
+			OutErrorMessage = TEXT("Failed to open skeleton file");
+			return false;
+		}
+
+		TSharedPtr<FJsonObject> LoadedData = MakeShareable(new FJsonObject);
+		if (!FJsonSerializer::Deserialize(TJsonReaderFactory<>::Create(MegaPlantsData), LoadedData))
+		{
+			OutErrorMessage = TEXT("Failed to deserialize JSON file for skeleton");
+			return false;
+		}
+
+		TArray<FString> RequiredJSONPaths = {
+			TEXT("points.attributes.pscale"),
+			TEXT("points.positions"),
+			TEXT("points.attributes.lengthFromRoot.values"),
+			TEXT("points.attributes.LOD_totalPscaleGradient.values"),
+			TEXT("points.attributes.budDirection.values"),
+			TEXT("primitives.points"),
+			TEXT("primitives.attributes.instancer_name.values"),
+			TEXT("primitives.attributes.instancer_pivot.values"),
+			TEXT("primitives.attributes.instancer_UP.values"),
+			TEXT("primitives.attributes.instancer_scale.values"),
+			TEXT("primitives.attributes.instancer_LFR.values"),
+			TEXT("primitives.attributes.parents.values"),
+			TEXT("primitives.attributes.children.values"),
+			TEXT("primitives.attributes.branchNumber.values"),
+			TEXT("globalAttributes.phyllotaxyLeaf.value")
+		};
+
+		for (const FString& Path : RequiredJSONPaths)
+		{
+			if (!HasJsonFieldPath(LoadedData, Path))
+			{
+				OutErrorMessage = FString::Printf(TEXT("Failed to find required path in JSON file: %s"), *Path);
+				return false;
+			}
+		}
+
+		const TSharedPtr<FJsonObject>& Points = LoadedData->GetObjectField(TEXT("points"));
+		const TSharedPtr<FJsonObject>& PointAttributes = Points->GetObjectField(TEXT("attributes"));
+
+		const TArray<TSharedPtr<FJsonValue>>& PointPositions = Points->GetArrayField(TEXT("positions"));
+		const TArray<TSharedPtr<FJsonValue>>& PointScales = PointAttributes->GetObjectField(TEXT("pscale"))->GetArrayField(TEXT("values"));
+
+		Collection.AddGroup(GroupNames::PointGroup);
+		Collection.AddElements(PointPositions.Num(), GroupNames::PointGroup);
+
+		TManagedArray<FVector3f>& PositionArray = Collection.AddAttribute<FVector3f>(AttributeNames::PointPosition, GroupNames::PointGroup);
+		TManagedArray<float>& ScaleArray = Collection.AddAttribute<float>("Scale", GroupNames::PointGroup);
+
+		for (int32 PointIndex = 0; PointIndex < PointPositions.Num(); ++PointIndex)
+		{
+			const TArray<TSharedPtr<FJsonValue>>& PointPosition = PointPositions[PointIndex]->AsArray();
+
+			const FVector3f Position = FVector3f(PointPosition[0]->AsNumber(), PointPosition[2]->AsNumber(), PointPosition[1]->AsNumber()) * 100.0f;
+			const float Scale = PointScales[PointIndex]->AsNumber() * 100.0f;
+
+			PositionArray[PointIndex] = Position;
+			ScaleArray[PointIndex] = Scale;
+		}
+
+		const TSharedPtr<FJsonObject>& Primitives = LoadedData->GetObjectField(TEXT("primitives"));
+		const TSharedPtr<FJsonObject>& PrimitiveAttributes = Primitives->GetObjectField(TEXT("attributes"));
+
+		const TSharedPtr<FJsonObject>& GlobalAttributes = LoadedData->GetObjectField(TEXT("globalAttributes"));
+
+		const TArray<TSharedPtr<FJsonValue>>& PrimitivesPoints = Primitives->GetArrayField(TEXT("points"));
+
+		Collection.AddGroup(GroupNames::BranchGroup);
+		Collection.AddElements(PrimitivesPoints.Num(), GroupNames::BranchGroup);
+
+		TManagedArray<TArray<int32>>& PointsArray = Collection.AddAttribute<TArray<int32>>(
+			PV::AttributeNames::BranchPoints, PV::GroupNames::BranchGroup);
+
+		for (int32 PrimitiveIndex = 0; PrimitiveIndex < PrimitivesPoints.Num(); ++PrimitiveIndex)
+		{
+			const TArray<TSharedPtr<FJsonValue>>& PrimitivePoints = PrimitivesPoints[PrimitiveIndex]->AsArray();
+			for (const TSharedPtr<FJsonValue>& PrimitivePoint : PrimitivePoints)
+			{
+				PointsArray[PrimitiveIndex].Add(PrimitivePoint->AsNumber());
+			}
+		}
+
+		Collection.AddGroup(GroupNames::DetailsGroup);
+		Collection.AddElements(1, GroupNames::DetailsGroup);
+
+		// Please note that LFR (length from root) for both skeleton points and foliage points are not being scaled by
+		// 100 at the moment, which is OK for now since they're only being compared relative to each other for now,
+		// however if they're used in any computation in the future in any meaningful way where other quantities are
+		// involved they will end up giving erroneous results.
+		FillAttributes(Collection, GroupNames::PointGroup, PointAttributes);
+		FillAttributes(Collection, GroupNames::BranchGroup, PrimitiveAttributes);
+		FillDetailsAttributes(Collection, GroupNames::DetailsGroup, GlobalAttributes);
+		FillPlantProfilesData(Collection, GroupNames::PlantProfilesGroup, GlobalAttributes);
+
+		FillFoliageData(Collection, PrimitiveAttributes, FilePath);
+
+		Facades::FMetaInfoFacade MetaInfoFacade(Collection);
+		MetaInfoFacade.CreateGuid(FilePath);
+
+		OutErrorMessage = "";
+		return true;
+	}
+
+	bool LoadGrowthDataJsonToCollection(FManagedArrayCollection& Collection, const FString& FilePath, FString& OutErrorMessage)
+	{
+		FString MegaPlantsData;
+		if (!FFileHelper::LoadFileToString(MegaPlantsData, *FilePath))
+		{
+			OutErrorMessage = TEXT("Failed to open skeleton file");
+			return false;
+		}
+
+		TSharedPtr<FJsonObject> LoadedData = MakeShareable(new FJsonObject);
+		if (!FJsonSerializer::Deserialize(TJsonReaderFactory<>::Create(MegaPlantsData), LoadedData))
+		{
+			OutErrorMessage = TEXT("Failed to deserialize JSON file for skeleton");
+			return false;
+		}
+
+		TArray<FString> RequiredJSONPaths = {
+			TEXT("points.positions"),
+			TEXT("points.attributes.budDirection.values"),
+			TEXT("primitives.points"),
+			TEXT("primitives.attributes.parents.values"),
+			TEXT("primitives.attributes.children.values"),
+			TEXT("primitives.attributes.branchNumber.values"),
+		};
+
+		for (const FString& Path : RequiredJSONPaths)
+		{
+			if (!HasJsonFieldPath(LoadedData, Path))
+			{
+				OutErrorMessage = FString::Printf(TEXT("Failed to find required path in JSON file: %s"), *Path);
+				return false;
+			}
+		}
+
+		const TSharedPtr<FJsonObject>& Points = LoadedData->GetObjectField(TEXT("points"));
+		const TSharedPtr<FJsonObject>& PointAttributes = Points->GetObjectField(TEXT("attributes"));
+
+		const TArray<TSharedPtr<FJsonValue>>& PointPositions = Points->GetArrayField(TEXT("positions"));
+
+		const TArray<TSharedPtr<FJsonValue>>* PointScales = Invoke([&]()
+		{
+			const TArray<TSharedPtr<FJsonValue>>* OutPointScales = nullptr;
+			const TSharedPtr<FJsonObject>* PScaleObject = nullptr;
+			if (PointAttributes->TryGetObjectField(TEXT("pscale"), PScaleObject))
+			{
+				(*PScaleObject)->TryGetArrayField(TEXT("values"), OutPointScales);
+			}
+			return OutPointScales;
+		});
+
+		Collection.AddGroup(GroupNames::PointGroup);
+		Collection.AddElements(PointPositions.Num(), GroupNames::PointGroup);
+
+		TManagedArray<FVector3f>& PositionArray = Collection.AddAttribute<FVector3f>(AttributeNames::PointPosition, GroupNames::PointGroup);
+		TManagedArray<float>& ScaleArray = Collection.AddAttribute<float>("Scale", GroupNames::PointGroup);
+
+		for (int32 PointIndex = 0; PointIndex < PointPositions.Num(); ++PointIndex)
+		{
+			const TArray<TSharedPtr<FJsonValue>>& PointPosition = PointPositions[PointIndex]->AsArray();
+
+			const FVector3f Position = FVector3f(PointPosition[0]->AsNumber(), PointPosition[2]->AsNumber(), PointPosition[1]->AsNumber()) * 100.0f;
+			const float Scale = PointScales ? (*PointScales)[PointIndex]->AsNumber() * 100.0f : 0.f;
+
+			PositionArray[PointIndex] = Position;
+			ScaleArray[PointIndex] = Scale;
+		}
+
+		const TSharedPtr<FJsonObject>& Primitives = LoadedData->GetObjectField(TEXT("primitives"));
+		const TSharedPtr<FJsonObject>& PrimitiveAttributes = Primitives->GetObjectField(TEXT("attributes"));
+		const TArray<TSharedPtr<FJsonValue>>& PrimitivesPoints = Primitives->GetArrayField(TEXT("points"));
+
+		Collection.AddGroup(GroupNames::BranchGroup);
+		Collection.AddElements(PrimitivesPoints.Num(), GroupNames::BranchGroup);
+
+		TManagedArray<TArray<int32>>& PointsArray = Collection.AddAttribute<TArray<int32>>(
+			PV::AttributeNames::BranchPoints, PV::GroupNames::BranchGroup);
+
+		for (int32 PrimitiveIndex = 0; PrimitiveIndex < PrimitivesPoints.Num(); ++PrimitiveIndex)
+		{
+			const TArray<TSharedPtr<FJsonValue>>& PrimitivePoints = PrimitivesPoints[PrimitiveIndex]->AsArray();
+			for (const TSharedPtr<FJsonValue>& PrimitivePoint : PrimitivePoints)
+			{
+				PointsArray[PrimitiveIndex].Add(PrimitivePoint->AsNumber());
+			}
+		}
+
+		Collection.AddGroup(GroupNames::DetailsGroup);
+		Collection.AddElements(1, GroupNames::DetailsGroup);
+
+		FillAttributes(Collection, GroupNames::PointGroup, PointAttributes);
+		FillAttributes(Collection, GroupNames::BranchGroup, PrimitiveAttributes);
+		FillDetailsAttributes(Collection, GroupNames::DetailsGroup, LoadedData->GetObjectField(TEXT("globalAttributes")));
+
+		Facades::FMetaInfoFacade MetaInfoFacade(Collection);
+		MetaInfoFacade.CreateGuid(FilePath);
+
+		OutErrorMessage = "";
+		return true;
+	}
+
+	TSharedPtr<FJsonObject> LoadMetaFileIntoJsonObject(const FString& FilePath, FString& OutErrorMessage)
+	{
+		FString MetaData;
+		if (!FFileHelper::LoadFileToString(MetaData, *FilePath))
+		{
+			OutErrorMessage = TEXT("Failed to open skeleton file");
+			return nullptr;
+		}
+
+		TSharedPtr<FJsonObject> LoadedData = MakeShareable(new FJsonObject);
+		if (!FJsonSerializer::Deserialize(TJsonReaderFactory<>::Create(MetaData), LoadedData))
+		{
+			OutErrorMessage = TEXT("Failed to deserialize JSON file for skeleton");
+			return nullptr;
+		}
+
+		TArray<FString> RequiredJSONPaths = {
+			TEXT("apix_uv"),
+			TEXT("apix_texture_channels")
+		};
+
+		for (const FString& Path : RequiredJSONPaths)
+		{
+			if (!HasJsonFieldPath(LoadedData, Path))
+			{
+				OutErrorMessage = FString::Printf(TEXT("Failed to find required path in JSON file: %s"), *Path);
+				return nullptr;
+			}
+		}
+
+		OutErrorMessage = "";
+		return LoadedData;
+	}
+
+	bool LoadMetaJsonToCollection(FManagedArrayCollection& Collection, TSharedPtr<FJsonObject> LoadedData)
+	{
+		
+		const TArray<TSharedPtr<FJsonValue>>& MaterialURanges = LoadedData->GetArrayField(TEXT("apix_uv"));
+
+		TArray<FVector2f> URanges;
+		URanges.Reserve(MaterialURanges.Num());
+		
+		for (int i=0; (i + 1) < MaterialURanges.Num(); i+= 2)
+		{
+			URanges.Add(FVector2f(MaterialURanges[i]->AsNumber() , MaterialURanges[i+1]->AsNumber()));
+		}
+
+		Facades::FBranchFacade Facade = Facades::FBranchFacade(Collection);
+		Facade.SetTrunkURange(URanges);
+		
+		return true;
+	}
+}
diff --git a/Engine/Plugins/Experimental/ProceduralVegetationEditor/Source/ProceduralVegetation/Private/Helpers/PVJSONHelper.h b/Engine/Plugins/Experimental/ProceduralVegetationEditor/Source/ProceduralVegetation/Private/Helpers/PVJSONHelper.h
--- a/Engine/Plugins/Experimental/ProceduralVegetationEditor/Source/ProceduralVegetation/Private/Helpers/PVJSONHelper.h
+++ b/Engine/Plugins/Experimental/ProceduralVegetationEditor/Source/ProceduralVegetation/Private/Helpers/PVJSONHelper.h
@@ -3,653 +3,29 @@
 #pragma once
 
 #include "CoreMinimal.h"
-#include "Chaos/Deformable/MuscleActivationConstraints.h"
-#include "GeometryCollection/ManagedArrayCollection.h"
-#include "Facades/PVAttributesNames.h"
-#include "Dom/JsonObject.h"
-#include "Dom/JsonValue.h"
-#include "Facades/PVBranchFacade.h"
-#include "Facades/PVMetaInfoFacade.h"
-#include "Facades/PVFoliageFacade.h"
-#include "Facades/PVProfileFacade.h"
-#include "Misc/FileHelper.h"
-#include "Serialization/JsonSerializer.h"
-#include "Misc/Paths.h"
-#include "Misc/PackageName.h"
 
-namespace PV
-{
-	static void FillAttributes(FManagedArrayCollection& Collection, const FName Group, const TSharedPtr<FJsonObject>& AttributesObject)
-	{
-		for (const auto& [Key, Attribute] : AttributesObject->Values)
-		{
-			const TSharedPtr<FJsonObject>& AttributeObject = Attribute->AsObject();
-
-			const FString Type = AttributeObject->GetStringField(TEXT("type"));
-			const int32 ElementCount = AttributeObject->GetNumberField(TEXT("size"));
-			const bool IsArray = AttributeObject->GetBoolField(TEXT("isArray"));
-			const TArray<TSharedPtr<FJsonValue>>& Values = AttributeObject->GetArrayField(TEXT("values"));
-
-			if (IsArray)
-			{
-				if (ElementCount == 1)
-				{
-					if (Type == "int")
-					{
-						TManagedArray<TArray<int32>>& Attr = Collection.AddAttribute<TArray<int32>>(*Key, Group);
-						for (int32 Index = 0; Index < Values.Num(); ++Index)
-						{
-							TArray<int32>& CurrentAttrib = Attr[Index];
-							const TArray<TSharedPtr<FJsonValue>>& AttribValues = Values[Index]->AsArray();
-							for (const TSharedPtr<FJsonValue>& AttribValue : AttribValues)
-							{
-								CurrentAttrib.Add(AttribValue->AsNumber());
-							}
-						}
-					}
-					else if (Type == "float")
-					{
-						TManagedArray<TArray<float>>& Attr = Collection.AddAttribute<TArray<float>>(*Key, Group);
-						for (int32 Index = 0; Index < Values.Num(); ++Index)
-						{
-							TArray<float>& CurrentAttrib = Attr[Index];
-							const TArray<TSharedPtr<FJsonValue>>& AttribValues = Values[Index]->AsArray();
-							for (const TSharedPtr<FJsonValue>& AttribValue : AttribValues)
-							{
-								CurrentAttrib.Add(AttribValue->AsNumber());
-							}
-						}
-					}
-					else if (Type == "string")
-					{
-						// TODO: Figure out how to manage string arrays
-					}
-				}
-				else if (ElementCount == 3)
-				{
-					if (Type == "int")
-					{
-						TManagedArray<TArray<FIntVector>>& Attr = Collection.AddAttribute<TArray<FIntVector>>(*Key, Group);
-						for (int32 Index = 0; Index < Values.Num(); ++Index)
-						{
-							TArray<FIntVector>& CurrentAttrib = Attr[Index];
-							const TArray<TSharedPtr<FJsonValue>>& AttribValues = Values[Index]->AsArray();
-							for (int ElemIndex = 0; ElemIndex < AttribValues.Num(); ElemIndex += ElementCount)
-							{
-								CurrentAttrib.Emplace(AttribValues[ElemIndex]->AsNumber(), AttribValues[ElemIndex + 2]->AsNumber(),
-									AttribValues[ElemIndex + 1]->AsNumber());
-							}
-						}
-					}
-					else if (Type == "float")
-					{
-						TManagedArray<TArray<FVector3f>>& Attr = Collection.AddAttribute<TArray<FVector3f>>(*Key, Group);
-						for (int32 Index = 0; Index < Values.Num(); ++Index)
-						{
-							TArray<FVector3f>& CurrentAttrib = Attr[Index];
-							const TArray<TSharedPtr<FJsonValue>>& AttribValues = Values[Index]->AsArray();
-							for (int ElemIndex = 0; ElemIndex < AttribValues.Num(); ElemIndex += ElementCount)
-							{
-								CurrentAttrib.Emplace(AttribValues[ElemIndex]->AsNumber(), AttribValues[ElemIndex + 2]->AsNumber(),
-									AttribValues[ElemIndex + 1]->AsNumber());
-							}
-						}
-					}
-				}
-			}
-			else
-			{
-				if (ElementCount == 1)
-				{
-					if (Type == "int")
-					{
-						TManagedArray<int32>& Attr = Collection.AddAttribute<int32>(*Key, Group);
-						for (int32 Index = 0; Index < Values.Num(); ++Index)
-						{
-							Attr[Index] = Values[Index]->AsNumber();
-						}
-					}
-					else if (Type == "float")
-					{
-						TManagedArray<float>& Attr = Collection.AddAttribute<float>(*Key, Group);
-						for (int32 Index = 0; Index < Values.Num(); ++Index)
-						{
-							Attr[Index] = Values[Index]->AsNumber();
-						}
-					}
-					else if (Type == "string")
-					{
-						TManagedArray<FString>& Attr = Collection.AddAttribute<FString>(*Key, Group);
-						for (int32 Index = 0; Index < Values.Num(); ++Index)
-						{
-							Attr[Index] = Values[Index]->AsString();
-						}
-					}
-				}
-				else if (ElementCount == 3)
-				{
-					if (Type == "int")
-					{
-						TManagedArray<FIntVector>& Attr = Collection.AddAttribute<FIntVector>(*Key, Group);
-						for (int32 Index = 0; Index < Values.Num(); ++Index)
-						{
-							const TArray<TSharedPtr<FJsonValue>>& VectorAttribute = Values[Index]->AsArray();
-							Attr[Index] = FIntVector(VectorAttribute[0]->AsNumber(), VectorAttribute[2]->AsNumber(), VectorAttribute[1]->AsNumber());
-						}
-					}
-					else if (Type == "float")
-					{
-						TManagedArray<FVector3f>& Attr = Collection.AddAttribute<FVector3f>(*Key, Group);
-						for (int32 Index = 0; Index < Values.Num(); ++Index)
-						{
-							const TArray<TSharedPtr<FJsonValue>>& VectorAttribute = Values[Index]->AsArray();
-							Attr[Index] = FVector3f(VectorAttribute[0]->AsNumber(), VectorAttribute[2]->AsNumber(), VectorAttribute[1]->AsNumber());
-						}
-					}
-				}
-			}
-		}
-	}
-
-	static void FillDetailsAttributes(FManagedArrayCollection& Collection, const FName Group, const TSharedPtr<FJsonObject>& AttributesObject)
-	{
-		for (const auto& [Key, Attribute] : AttributesObject->Values)
-		{
-			const TSharedPtr<FJsonObject>& AttributeObject = Attribute->AsObject();
-
-			const FString Type = AttributeObject->GetStringField(TEXT("type"));
-			const bool IsArray = AttributeObject->GetBoolField(TEXT("isArray"));
-
-			if (IsArray)
-			{
-				const TArray<TSharedPtr<FJsonValue>>& Values = AttributeObject->GetArrayField(TEXT("value"));
-				if (Type == "int")
-				{
-					TManagedArray<TArray<int32>>& Attr = Collection.AddAttribute<TArray<int32>>(*Key, Group);
-					TArray<int32>& CurrentAttrib = Attr[0];
-					for (const TSharedPtr<FJsonValue>& AttribValue : Values)
-					{
-						CurrentAttrib.Add(AttribValue->AsNumber());
-					}
-				}
-				else if (Type == "float")
-				{
-					TManagedArray<TArray<float>>& Attr = Collection.AddAttribute<TArray<float>>(*Key, Group);
-					TArray<float>& CurrentAttrib = Attr[0];
-					for (const TSharedPtr<FJsonValue>& AttribValue : Values)
-					{
-						CurrentAttrib.Add(AttribValue->AsNumber());
-					}
-				}
-				else if (Type == "string")
-				{
-					// TODO: Figure out how to manage string arrays
-				}
-			}
-			else
-			{
-				if (Type == "int")
-				{
-					TManagedArray<int32>& Attr = Collection.AddAttribute<int32>(*Key, Group);
-					Attr[0] = AttributeObject->GetNumberField(TEXT("value"));
-				}
-				else if (Type == "float")
-				{
-					TManagedArray<float>& Attr = Collection.AddAttribute<float>(*Key, Group);
-					Attr[0] = AttributeObject->GetNumberField(TEXT("value"));
-				}
-				else if (Type == "string")
-				{
-					TManagedArray<FString>& Attr = Collection.AddAttribute<FString>(*Key, Group);
-					Attr[0] = AttributeObject->GetStringField(TEXT("value"));
-				}
-			}
-		}
-	}
-
-	static void FillPlantProfilesData(FManagedArrayCollection& Collection, const FName Group, const TSharedPtr<FJsonObject>& AttributesObject)
-	{
-		Facades::FPlantProfileFacade Facade(Collection);
-
-		for (const auto& [Key, Attribute] : AttributesObject->Values)
-		{
-			if (const TSharedPtr<FJsonObject>& AttributeObject = Attribute->AsObject();
-				AttributeObject
-				&& AttributeObject->HasTypedField<EJson::Array>(TEXT("value"))
-				&& Key.StartsWith(TEXT("plantProfile_")))
-			{
-				const TArray<TSharedPtr<FJsonValue>>& Values = AttributeObject->GetArrayField(TEXT("value"));
-				TArray<float> Points;
-				Points.Reserve(Values.Num());
-
-				for (const TSharedPtr<FJsonValue>& AttribValue : Values)
-				{
-					if (double Point; AttribValue->TryGetNumber(Point))
-					{
-						Points.Add(static_cast<float>(AttribValue->AsNumber()));
-					}
-				}
-
-				Facade.AddProfileEntry(Points);
-			}
-		}
-	}
-
-	static void FillFoliageData(FManagedArrayCollection& Collection, const TSharedPtr<FJsonObject>& PrimitiveAttributesObject, const FString& InPath)
-	{
-		const TSharedPtr<FJsonObject>& InstancerNameObject = PrimitiveAttributesObject->GetObjectField(TEXT("instancer_name"));
-		const TArray<TSharedPtr<FJsonValue>>& InstancerNameValues = InstancerNameObject->GetArrayField(TEXT("values"));
-
-		const TSharedPtr<FJsonObject>& InstancerPivotObject = PrimitiveAttributesObject->GetObjectField(TEXT("instancer_pivot"));
-		const TArray<TSharedPtr<FJsonValue>>& InstancerPivotValues = InstancerPivotObject->GetArrayField(TEXT("values"));
-
-		const TSharedPtr<FJsonObject>& InstancerUpVectorObject = PrimitiveAttributesObject->GetObjectField(TEXT("instancer_UP"));
-		const TArray<TSharedPtr<FJsonValue>>& InstancerUpVectorValues = InstancerUpVectorObject->GetArrayField(TEXT("values"));
-
-		const TSharedPtr<FJsonObject>& InstancerNormalVectorObject = PrimitiveAttributesObject->GetObjectField(TEXT("instancer_N"));
-		const TArray<TSharedPtr<FJsonValue>>& InstancerNormalVectorValues = InstancerNormalVectorObject->GetArrayField(TEXT("values"));
-
-		const TSharedPtr<FJsonObject>& InstancerScaleObject = PrimitiveAttributesObject->GetObjectField(TEXT("instancer_scale"));
-		const TArray<TSharedPtr<FJsonValue>>& InstancerScaleValues = InstancerScaleObject->GetArrayField(TEXT("values"));
-
-		const TSharedPtr<FJsonObject>& InstancerLFRObject = PrimitiveAttributesObject->GetObjectField(TEXT("instancer_LFR"));
-		const TArray<TSharedPtr<FJsonValue>>& InstancerLFRValues = InstancerLFRObject->GetArrayField(TEXT("values"));
-
-		// Compute unique names and number of elements for group
-		int32 TotalValues = 0;
-		int32 NextFoliageIndex = 0;
-		TMap<FString, int32> FoliageNamesToIndex;
-		TArray<FString> FoliageNames;
-		for (int32 i = 0; i < InstancerNameValues.Num(); ++i)
-		{
-			const TArray<TSharedPtr<FJsonValue>>& InstancerNames = InstancerNameValues[i]->AsArray();
-			for (int32 Index = 0; Index < InstancerNames.Num(); ++Index)
-			{
-				const FString& Name = InstancerNames[Index]->AsString();
-				if (!FoliageNamesToIndex.Contains(Name))
-				{
-					FoliageNamesToIndex.Emplace(Name, NextFoliageIndex++);
-					FoliageNames.Emplace(Name);
-				}
-			}
-			TotalValues += InstancerNames.Num();
-		}
-
-		Facades::FFoliageFacade Facade(Collection, TotalValues);
-
-		for (FString& FoliageName : FoliageNames)
-		{
-			FoliageName = FoliageName + '.' + FoliageName;
-		}
-		
-		Facade.SetFoliageNames(FoliageNames);
-
-		int32 FoliageEntryIndex = 0;
-		for (int32 Id = 0; Id < InstancerNameValues.Num(); ++Id)
-		{
-			const TArray<TSharedPtr<FJsonValue>>& InstancerNames = InstancerNameValues[Id]->AsArray();
-			const TArray<TSharedPtr<FJsonValue>>& PivotPoints = InstancerPivotValues[Id]->AsArray();
-			const TArray<TSharedPtr<FJsonValue>>& UpVectors = InstancerUpVectorValues[Id]->AsArray();
-			const TArray<TSharedPtr<FJsonValue>>& NormalVectors = InstancerNormalVectorValues[Id]->AsArray();
-			const TArray<TSharedPtr<FJsonValue>>& Scales = InstancerScaleValues[Id]->AsArray();
-			const TArray<TSharedPtr<FJsonValue>>& LFRs = InstancerLFRValues[Id]->AsArray();
-
-			TArray<int32> FoliageEntryIds;
-			for (int32 j = 0; j < InstancerNames.Num(); ++j)
-			{
-				const FString Name = InstancerNames[j]->AsString();
-				const FVector3f PivotPoint = FVector3f(
-					PivotPoints[j * 3]->AsNumber(),
-					PivotPoints[(j * 3) + 2]->AsNumber(),
-					PivotPoints[(j * 3) + 1]->AsNumber()) * 100.0f;
-				const FVector3f UpVector = FVector3f(
-					UpVectors[j * 3]->AsNumber(),
-					UpVectors[(j * 3) + 2]->AsNumber(),
-					UpVectors[(j * 3) + 1]->AsNumber());
-				const FVector3f NormalVector = FVector3f(
-					NormalVectors[j * 3]->AsNumber(),
-					NormalVectors[(j * 3) + 2]->AsNumber(),
-					NormalVectors[(j * 3) + 1]->AsNumber());
-				const float Scale = Scales[j]->AsNumber();
-				const float LengthFromRoot = LFRs[j]->AsNumber();
-
-				Facade.SetFoliageEntry(FoliageEntryIndex, {
-					.NameId = FoliageNamesToIndex[Name],
-					.BranchId = Id,
-					.PivotPoint = PivotPoint,
-					.UpVector = UpVector,
-					.NormalVector = NormalVector,
-					.Scale = Scale,
-					.LengthFromRoot = LengthFromRoot,
-				});
-
-				FoliageEntryIds.Add(FoliageEntryIndex);
-				FoliageEntryIndex++;
-			}
-
-			Facade.SetFoliageIdsArray(Id, FoliageEntryIds);
-		}
-	}
-
-	static void SetFoliagePaths(FManagedArrayCollection& Collection,const FString& FilePath)
-	{
-		Facades::FFoliageFacade Facade(Collection);
-
-		FString AbsolutePath = FPaths::ConvertRelativePathToFull(FilePath);
-		FString PackagePath;
-
-		if (FPackageName::TryConvertFilenameToLongPackageName(AbsolutePath, PackagePath))
-		{
-			TArray<FString> FoliageNames = Facade.GetFoliageNames();
-			FString PackageFolder = FPackageName::GetLongPackagePath(PackagePath);
-			Facade.SetFoliagePath(PackageFolder);
-			for (FString& FoliageName : FoliageNames)
-			{
-				FoliageName = PackageFolder / FoliageName;
-			}
-			Facade.SetFoliageNames(FoliageNames);
-		}
-		
-	}
-
-	static bool HasJsonFieldPath(const TSharedPtr<FJsonObject>& JsonObject, const FString& Path)
-	{
-		if (!JsonObject.IsValid())
-		{
-			return false;
-		}
-
-		TArray<FString> Keys;
-		Path.ParseIntoArray(Keys, TEXT("."), true);
-
-		TSharedPtr<FJsonObject> CurrentObject = JsonObject;
-
-		for (int32 i = 0; i < Keys.Num(); ++i)
-		{
-			const FString& CurrentKey = Keys[i];
-			if (i == Keys.Num() - 1)
-			{
-				if (!CurrentObject->HasField(CurrentKey))
-				{
-					return false;
-				}
-			}
-			else
-			{
-				const TSharedPtr<FJsonObject>* NextObject;
-				if (CurrentObject->TryGetObjectField(CurrentKey, NextObject))
-				{
-					CurrentObject = *NextObject;
-				}
-				else
-				{
-					return false;
-				}
-			}
-		}
-
-		return true;
-	}
-
-	static bool LoadMegaPlantsJsonToCollection(FManagedArrayCollection& Collection, const FString& FilePath, FString& OutErrorMessage)
-	{
-		FString MegaPlantsData;
-		if (!FFileHelper::LoadFileToString(MegaPlantsData, *FilePath))
-		{
-			OutErrorMessage = TEXT("Failed to open skeleton file");
-			return false;
-		}
-
-		TSharedPtr<FJsonObject> LoadedData = MakeShareable(new FJsonObject);
-		if (!FJsonSerializer::Deserialize(TJsonReaderFactory<>::Create(MegaPlantsData), LoadedData))
-		{
-			OutErrorMessage = TEXT("Failed to deserialize JSON file for skeleton");
-			return false;
-		}
-
-		TArray<FString> RequiredJSONPaths = {
-			TEXT("points.attributes.pscale"),
-			TEXT("points.positions"),
-			TEXT("points.attributes.lengthFromRoot.values"),
-			TEXT("points.attributes.LOD_totalPscaleGradient.values"),
-			TEXT("points.attributes.budDirection.values"),
-			TEXT("primitives.points"),
-			TEXT("primitives.attributes.instancer_name.values"),
-			TEXT("primitives.attributes.instancer_pivot.values"),
-			TEXT("primitives.attributes.instancer_UP.values"),
-			TEXT("primitives.attributes.instancer_scale.values"),
-			TEXT("primitives.attributes.instancer_LFR.values"),
-			TEXT("primitives.attributes.parents.values"),
-			TEXT("primitives.attributes.children.values"),
-			TEXT("primitives.attributes.branchNumber.values"),
-			TEXT("globalAttributes.phyllotaxyLeaf.value")
-		};
-
-		for (const FString& Path : RequiredJSONPaths)
-		{
-			if (!HasJsonFieldPath(LoadedData, Path))
-			{
-				OutErrorMessage = FString::Printf(TEXT("Failed to find required path in JSON file: %s"), *Path);
-				return false;
-			}
-		}
-
-		const TSharedPtr<FJsonObject>& Points = LoadedData->GetObjectField(TEXT("points"));
-		const TSharedPtr<FJsonObject>& PointAttributes = Points->GetObjectField(TEXT("attributes"));
+struct FManagedArrayCollection;
+class FJsonObject;
 
-		const TArray<TSharedPtr<FJsonValue>>& PointPositions = Points->GetArrayField(TEXT("positions"));
-		const TArray<TSharedPtr<FJsonValue>>& PointScales = PointAttributes->GetObjectField(TEXT("pscale"))->GetArrayField(TEXT("values"));
-
-		Collection.AddGroup(GroupNames::PointGroup);
-		Collection.AddElements(PointPositions.Num(), GroupNames::PointGroup);
-
-		TManagedArray<FVector3f>& PositionArray = Collection.AddAttribute<FVector3f>(AttributeNames::PointPosition, GroupNames::PointGroup);
-		TManagedArray<float>& ScaleArray = Collection.AddAttribute<float>("Scale", GroupNames::PointGroup);
-
-		for (int32 PointIndex = 0; PointIndex < PointPositions.Num(); ++PointIndex)
-		{
-			const TArray<TSharedPtr<FJsonValue>>& PointPosition = PointPositions[PointIndex]->AsArray();
-
-			const FVector3f Position = FVector3f(PointPosition[0]->AsNumber(), PointPosition[2]->AsNumber(), PointPosition[1]->AsNumber()) * 100.0f;
-			const float Scale = PointScales[PointIndex]->AsNumber() * 100.0f;
-
-			PositionArray[PointIndex] = Position;
-			ScaleArray[PointIndex] = Scale;
-		}
-
-		const TSharedPtr<FJsonObject>& Primitives = LoadedData->GetObjectField(TEXT("primitives"));
-		const TSharedPtr<FJsonObject>& PrimitiveAttributes = Primitives->GetObjectField(TEXT("attributes"));
-
-		const TSharedPtr<FJsonObject>& GlobalAttributes = LoadedData->GetObjectField(TEXT("globalAttributes"));
-
-		const TArray<TSharedPtr<FJsonValue>>& PrimitivesPoints = Primitives->GetArrayField(TEXT("points"));
-
-		Collection.AddGroup(GroupNames::BranchGroup);
-		Collection.AddElements(PrimitivesPoints.Num(), GroupNames::BranchGroup);
-
-		TManagedArray<TArray<int32>>& PointsArray = Collection.AddAttribute<TArray<int32>>(
-			PV::AttributeNames::BranchPoints, PV::GroupNames::BranchGroup);
-
-		for (int32 PrimitiveIndex = 0; PrimitiveIndex < PrimitivesPoints.Num(); ++PrimitiveIndex)
-		{
-			const TArray<TSharedPtr<FJsonValue>>& PrimitivePoints = PrimitivesPoints[PrimitiveIndex]->AsArray();
-			for (const TSharedPtr<FJsonValue>& PrimitivePoint : PrimitivePoints)
-			{
-				PointsArray[PrimitiveIndex].Add(PrimitivePoint->AsNumber());
-			}
-		}
-
-		Collection.AddGroup(GroupNames::DetailsGroup);
-		Collection.AddElements(1, GroupNames::DetailsGroup);
-
-		// Please note that LFR (length from root) for both skeleton points and foliage points are not being scaled by
-		// 100 at the moment, which is OK for now since they're only being compared relative to each other for now,
-		// however if they're used in any computation in the future in any meaningful way where other quantities are
-		// involved they will end up giving erroneous results.
-		FillAttributes(Collection, GroupNames::PointGroup, PointAttributes);
-		FillAttributes(Collection, GroupNames::BranchGroup, PrimitiveAttributes);
-		FillDetailsAttributes(Collection, GroupNames::DetailsGroup, GlobalAttributes);
-		FillPlantProfilesData(Collection, GroupNames::PlantProfilesGroup, GlobalAttributes);
-
-		FillFoliageData(Collection, PrimitiveAttributes, FilePath);
-
-		Facades::FMetaInfoFacade MetaInfoFacade(Collection);
-		MetaInfoFacade.CreateGuid(FilePath);
-
-		OutErrorMessage = "";
-		return true;
-	}
-
-	static bool LoadGrowthDataJsonToCollection(FManagedArrayCollection& Collection, const FString& FilePath, FString& OutErrorMessage)
-	{
-		FString MegaPlantsData;
-		if (!FFileHelper::LoadFileToString(MegaPlantsData, *FilePath))
-		{
-			OutErrorMessage = TEXT("Failed to open skeleton file");
-			return false;
-		}
-
-		TSharedPtr<FJsonObject> LoadedData = MakeShareable(new FJsonObject);
-		if (!FJsonSerializer::Deserialize(TJsonReaderFactory<>::Create(MegaPlantsData), LoadedData))
-		{
-			OutErrorMessage = TEXT("Failed to deserialize JSON file for skeleton");
-			return false;
-		}
-
-		TArray<FString> RequiredJSONPaths = {
-			TEXT("points.positions"),
-			TEXT("points.attributes.budDirection.values"),
-			TEXT("primitives.points"),
-			TEXT("primitives.attributes.parents.values"),
-			TEXT("primitives.attributes.children.values"),
-			TEXT("primitives.attributes.branchNumber.values"),
-		};
-
-		for (const FString& Path : RequiredJSONPaths)
-		{
-			if (!HasJsonFieldPath(LoadedData, Path))
-			{
-				OutErrorMessage = FString::Printf(TEXT("Failed to find required path in JSON file: %s"), *Path);
-				return false;
-			}
-		}
-
-		const TSharedPtr<FJsonObject>& Points = LoadedData->GetObjectField(TEXT("points"));
-		const TSharedPtr<FJsonObject>& PointAttributes = Points->GetObjectField(TEXT("attributes"));
-
-		const TArray<TSharedPtr<FJsonValue>>& PointPositions = Points->GetArrayField(TEXT("positions"));
-
-		const TArray<TSharedPtr<FJsonValue>>* PointScales = Invoke([&]()
-		{
-			const TArray<TSharedPtr<FJsonValue>>* OutPointScales = nullptr;
-			const TSharedPtr<FJsonObject>* PScaleObject = nullptr;
-			if (PointAttributes->TryGetObjectField(TEXT("pscale"), PScaleObject))
-			{
-				(*PScaleObject)->TryGetArrayField(TEXT("values"), OutPointScales);
-			}
-			return OutPointScales;
-		});
-
-		Collection.AddGroup(GroupNames::PointGroup);
-		Collection.AddElements(PointPositions.Num(), GroupNames::PointGroup);
-
-		TManagedArray<FVector3f>& PositionArray = Collection.AddAttribute<FVector3f>(AttributeNames::PointPosition, GroupNames::PointGroup);
-		TManagedArray<float>& ScaleArray = Collection.AddAttribute<float>("Scale", GroupNames::PointGroup);
-
-		for (int32 PointIndex = 0; PointIndex < PointPositions.Num(); ++PointIndex)
-		{
-			const TArray<TSharedPtr<FJsonValue>>& PointPosition = PointPositions[PointIndex]->AsArray();
-
-			const FVector3f Position = FVector3f(PointPosition[0]->AsNumber(), PointPosition[2]->AsNumber(), PointPosition[1]->AsNumber()) * 100.0f;
-			const float Scale = PointScales ? (*PointScales)[PointIndex]->AsNumber() * 100.0f : 0.f;
-
-			PositionArray[PointIndex] = Position;
-			ScaleArray[PointIndex] = Scale;
-		}
-
-		const TSharedPtr<FJsonObject>& Primitives = LoadedData->GetObjectField(TEXT("primitives"));
-		const TSharedPtr<FJsonObject>& PrimitiveAttributes = Primitives->GetObjectField(TEXT("attributes"));
-		const TArray<TSharedPtr<FJsonValue>>& PrimitivesPoints = Primitives->GetArrayField(TEXT("points"));
-
-		Collection.AddGroup(GroupNames::BranchGroup);
-		Collection.AddElements(PrimitivesPoints.Num(), GroupNames::BranchGroup);
-
-		TManagedArray<TArray<int32>>& PointsArray = Collection.AddAttribute<TArray<int32>>(
-			PV::AttributeNames::BranchPoints, PV::GroupNames::BranchGroup);
-
-		for (int32 PrimitiveIndex = 0; PrimitiveIndex < PrimitivesPoints.Num(); ++PrimitiveIndex)
-		{
-			const TArray<TSharedPtr<FJsonValue>>& PrimitivePoints = PrimitivesPoints[PrimitiveIndex]->AsArray();
-			for (const TSharedPtr<FJsonValue>& PrimitivePoint : PrimitivePoints)
-			{
-				PointsArray[PrimitiveIndex].Add(PrimitivePoint->AsNumber());
-			}
-		}
-
-		Collection.AddGroup(GroupNames::DetailsGroup);
-		Collection.AddElements(1, GroupNames::DetailsGroup);
-
-		FillAttributes(Collection, GroupNames::PointGroup, PointAttributes);
-		FillAttributes(Collection, GroupNames::BranchGroup, PrimitiveAttributes);
-		FillDetailsAttributes(Collection, GroupNames::DetailsGroup, LoadedData->GetObjectField(TEXT("globalAttributes")));
-
-		Facades::FMetaInfoFacade MetaInfoFacade(Collection);
-		MetaInfoFacade.CreateGuid(FilePath);
-
-		OutErrorMessage = "";
-		return true;
-	}
+namespace PV::JSON
+{
+	void FillAttributes(FManagedArrayCollection& Collection, const FName Group, const TSharedPtr<FJsonObject>& AttributesObject);
 
-	static TSharedPtr<FJsonObject> LoadMetaFileIntoJsonObject(const FString& FilePath, FString& OutErrorMessage)
-	{
-		FString MetaData;
-		if (!FFileHelper::LoadFileToString(MetaData, *FilePath))
-		{
-			OutErrorMessage = TEXT("Failed to open skeleton file");
-			return nullptr;
-		}
+	void FillDetailsAttributes(FManagedArrayCollection& Collection, const FName Group, const TSharedPtr<FJsonObject>& AttributesObject);
 
-		TSharedPtr<FJsonObject> LoadedData = MakeShareable(new FJsonObject);
-		if (!FJsonSerializer::Deserialize(TJsonReaderFactory<>::Create(MetaData), LoadedData))
-		{
-			OutErrorMessage = TEXT("Failed to deserialize JSON file for skeleton");
-			return nullptr;
-		}
+	void FillPlantProfilesData(FManagedArrayCollection& Collection, const FName Group, const TSharedPtr<FJsonObject>& AttributesObject);
 
-		TArray<FString> RequiredJSONPaths = {
-			TEXT("apix_uv"),
-			TEXT("apix_texture_channels")
-		};
+	void FillFoliageData(FManagedArrayCollection& Collection, const TSharedPtr<FJsonObject>& PrimitiveAttributesObject, const FString& InPath);
 
-		for (const FString& Path : RequiredJSONPaths)
-		{
-			if (!HasJsonFieldPath(LoadedData, Path))
-			{
-				OutErrorMessage = FString::Printf(TEXT("Failed to find required path in JSON file: %s"), *Path);
-				return nullptr;
-			}
-		}
+	void SetFoliagePaths(FManagedArrayCollection& Collection,const FString& FilePath);
 
-		OutErrorMessage = "";
-		return LoadedData;
-	}
+	bool HasJsonFieldPath(const TSharedPtr<FJsonObject>& JsonObject, const FString& Path);
+	
+	bool LoadMegaPlantsJsonToCollection(FManagedArrayCollection& Collection, const FString& FilePath, FString& OutErrorMessage);
 
-	static bool LoadMetaJsonToCollection(FManagedArrayCollection& Collection, TSharedPtr<FJsonObject> LoadedData)
-	{
-		
-		const TArray<TSharedPtr<FJsonValue>>& MaterialURanges = LoadedData->GetArrayField(TEXT("apix_uv"));
+	bool LoadGrowthDataJsonToCollection(FManagedArrayCollection& Collection, const FString& FilePath, FString& OutErrorMessage);
 
-		TArray<FVector2f> URanges;
-		URanges.Reserve(MaterialURanges.Num());
-		
-		for (int i=0; (i + 1) < MaterialURanges.Num(); i+= 2)
-		{
-			URanges.Add(FVector2f(MaterialURanges[i]->AsNumber() , MaterialURanges[i+1]->AsNumber()));
-		}
+	TSharedPtr<FJsonObject> LoadMetaFileIntoJsonObject(const FString& FilePath, FString& OutErrorMessage);
 
-		Facades::FBranchFacade Facade = Facades::FBranchFacade(Collection);
-		Facade.SetTrunkURange(URanges);
-		
-		return true;
-	}
+	bool LoadMetaJsonToCollection(FManagedArrayCollection& Collection, TSharedPtr<FJsonObject> LoadedData);
 }
diff --git a/Engine/Plugins/Experimental/ProceduralVegetationEditor/Source/ProceduralVegetation/Private/Nodes/Development/PVGrowthDataJsonImporterSettings.cpp b/Engine/Plugins/Experimental/ProceduralVegetationEditor/Source/ProceduralVegetation/Private/Nodes/Development/PVGrowthDataJsonImporterSettings.cpp
--- a/Engine/Plugins/Experimental/ProceduralVegetationEditor/Source/ProceduralVegetation/Private/Nodes/Development/PVGrowthDataJsonImporterSettings.cpp
+++ b/Engine/Plugins/Experimental/ProceduralVegetationEditor/Source/ProceduralVegetation/Private/Nodes/Development/PVGrowthDataJsonImporterSettings.cpp
@@ -65,7 +65,7 @@ bool FPVGrowthDataJsonImporterElement::ExecuteInternal(FPCGContext* InContext) c
 	
 	FString ErrorMessage;
 	FManagedArrayCollection Collection;
-	bool bSuccess = PV::LoadGrowthDataJsonToCollection(Collection, Settings->GrowthDataFile.FilePath, ErrorMessage);
+	bool bSuccess = PV::JSON::LoadGrowthDataJsonToCollection(Collection, Settings->GrowthDataFile.FilePath, ErrorMessage);
 	if (!bSuccess)
 	{
 		UE_LOG(LogProceduralVegetation, Error, TEXT("%s"), *ErrorMessage);
diff --git a/Engine/Plugins/Experimental/ProceduralVegetationEditor/Source/ProceduralVegetation/Private/Nodes/PVImporterSettings.cpp b/Engine/Plugins/Experimental/ProceduralVegetationEditor/Source/ProceduralVegetation/Private/Nodes/PVImporterSettings.cpp
--- a/Engine/Plugins/Experimental/ProceduralVegetationEditor/Source/ProceduralVegetation/Private/Nodes/PVImporterSettings.cpp
+++ b/Engine/Plugins/Experimental/ProceduralVegetationEditor/Source/ProceduralVegetation/Private/Nodes/PVImporterSettings.cpp
@@ -56,7 +56,7 @@ bool FPVImporterElement::ExecuteInternal(FPCGContext* InContext) const
 	FString ErrorMessage;
  
 	FManagedArrayCollection Collection;
-	bool bSuccess = PV::LoadMegaPlantsJsonToCollection(Collection, SkeletonFilePath, ErrorMessage);
+	bool bSuccess = PV::JSON::LoadMegaPlantsJsonToCollection(Collection, SkeletonFilePath, ErrorMessage);
 	if (!bSuccess)
 	{
 		PCGLog::LogErrorOnGraph(FText::FromString(ErrorMessage), InContext);
diff --git a/Engine/Plugins/Experimental/ProceduralVegetationEditor/Source/ProceduralVegetation/Private/ProceduralVegetationPreset.cpp b/Engine/Plugins/Experimental/ProceduralVegetationEditor/Source/ProceduralVegetation/Private/ProceduralVegetationPreset.cpp
--- a/Engine/Plugins/Experimental/ProceduralVegetationEditor/Source/ProceduralVegetation/Private/ProceduralVegetationPreset.cpp
+++ b/Engine/Plugins/Experimental/ProceduralVegetationEditor/Source/ProceduralVegetation/Private/ProceduralVegetationPreset.cpp
@@ -7,11 +7,13 @@
 #include "AssetRegistry/IAssetRegistry.h"
 #include "Facades/PVBranchFacade.h"
 #include "Facades/PVProfileFacade.h"
+#include "HAL/FileManager.h"
 #include "Helpers/PVFoliageJSONHelper.h"
 #include "Helpers/PVJSONHelper.h"
 #include "Helpers/PVUtilities.h"
 #include "Implementations/PVFoliage.h"
- 
+#include "Misc/Paths.h"
+
 void FPVPresetVariationInfo::Fill(FName InName, FManagedArrayCollection Collection)
 {
 	Name = InName;
@@ -58,7 +60,7 @@ void UProceduralVegetationPreset::LoadFromVariantFiles(const TMap<FString, FStri
 	{
 		UE_LOG(LogProceduralVegetation, Log, TEXT("Loaded variant : %s"), *VariantName);
 		FManagedArrayCollection Collection;
-		if (FString OutError; PV::LoadMegaPlantsJsonToCollection(Collection, VariantFile, OutError))
+		if (FString OutError; PV::JSON::LoadMegaPlantsJsonToCollection(Collection, VariantFile, OutError))
 		{
 			Variants.Add(VariantName, Collection);
 		}
@@ -109,12 +111,12 @@ void UProceduralVegetationPreset::UpdateDataAsset()
 		FString FullPath = JsonDirectoryPath.Path / FileName;
 		FManagedArrayCollection Collection;
 		FString MetaJsonError;
-		if ( PV::LoadMegaPlantsJsonToCollection(Collection, FullPath, OutError))
+		if ( PV::JSON::LoadMegaPlantsJsonToCollection(Collection, FullPath, OutError))
 		{
 			UE_LOG(LogProceduralVegetation, Log, TEXT("Variant %s loaded from file"), *FPaths::GetBaseFilename(FullPath));
 			Variants.Add(FPaths::GetBaseFilename(FullPath), Collection);
 		}
-		else if (TSharedPtr<FJsonObject> MetaJson = PV::LoadMetaFileIntoJsonObject(FullPath, MetaJsonError))
+		else if (TSharedPtr<FJsonObject> MetaJson = PV::JSON::LoadMetaFileIntoJsonObject(FullPath, MetaJsonError))
 		{
 			OutError.Empty();
 			MetaFiles.Add(MetaJson);
@@ -135,11 +137,11 @@ void UProceduralVegetationPreset::UpdateDataAsset()
 	{
 		if (MetaFiles.IsValidIndex(VariantIndex))
 		{
-			PV::LoadMetaJsonToCollection(Pair.Value, MetaFiles[VariantIndex]);
+			PV::JSON::LoadMetaJsonToCollection(Pair.Value, MetaFiles[VariantIndex]);
 		}
 		else if (MetaFiles.Num() > 0 && VariantIndex > MetaFiles.Num() - 1)
 		{
-			PV::LoadMetaJsonToCollection(Pair.Value, MetaFiles[0]);
+			PV::JSON::LoadMetaJsonToCollection(Pair.Value, MetaFiles[0]);
 		}
  
 		// load Foliage Data again
