[PVE] Fixed non unity issue

#Synced-CL 50934585

diff --git a/Engine/Plugins/Experimental/ProceduralVegetationEditor/Source/ProceduralVegetation/Private/Nodes/PVImporterSettings.cpp b/Engine/Plugins/Experimental/ProceduralVegetationEditor/Source/ProceduralVegetation/Private/Nodes/PVImporterSettings.cpp
--- a/Engine/Plugins/Experimental/ProceduralVegetationEditor/Source/ProceduralVegetation/Private/Nodes/PVImporterSettings.cpp
+++ b/Engine/Plugins/Experimental/ProceduralVegetationEditor/Source/ProceduralVegetation/Private/Nodes/PVImporterSettings.cpp
@@ -13,6 +13,7 @@
 #include "Helpers/PCGSettingsHelpers.h"
 #include "Helpers/PVFoliageJSONHelper.h"
 #include "Helpers/PVJSONHelper.h"
+#include "Misc/Paths.h"
  
 UPVImporterSettings::UPVImporterSettings()
 {
