[PVE] Exposed mesh builder, foliage distributor, Grafter and bone reduction parameters as overridable pins.  This would enable the artists to test the complete workflow with subgraphs.

#Synced-CL 51086708

diff --git a/Engine/Plugins/Experimental/ProceduralVegetationEditor/Source/ProceduralVegetation/Private/Implementations/PVMeshBuilder.h b/Engine/Plugins/Experimental/ProceduralVegetationEditor/Source/ProceduralVegetation/Private/Implementations/PVMeshBuilder.h
--- a/Engine/Plugins/Experimental/ProceduralVegetationEditor/Source/ProceduralVegetation/Private/Implementations/PVMeshBuilder.h
+++ b/Engine/Plugins/Experimental/ProceduralVegetationEditor/Source/ProceduralVegetation/Private/Implementations/PVMeshBuilder.h
@@ -32,12 +32,12 @@ struct FPVMeshBuilderParams
 	GENERATED_BODY()
 
 	UPROPERTY(EditAnywhere, Category="Materials",
-		meta=(PinHiddenByDefault, InlineEditConditionToggle, Tooltip=
+		meta=(PCG_NotOverridable, PinHiddenByDefault, InlineEditConditionToggle, Tooltip=
 			"Toggles between default and custom material settings.\n\nDefault material setting is defined in Procedural Vegetation Preset.\n\nCustom material settings allow for editing."
 		))
 	bool bOverrideMaterial = false;
 
-	UPROPERTY(EditAnywhere, Category="Materials", meta = (EditCondition="bOverrideMaterial"))
+	UPROPERTY(EditAnywhere, Category="Materials", meta = (PCG_NotOverridable, EditCondition="bOverrideMaterial"))
 	FPVMaterialSettings MaterialSettings;
 
 	UPROPERTY(EditAnywhere, Category="Mesh Settings", DisplayName="Point Removal",
@@ -143,7 +143,7 @@ struct FPVMeshBuilderParams
 	FPVFloatRamp ScaleRetentionGradient;
 
 	UPROPERTY(EditAnywhere, Category="Plant Profile Settings", DisplayName="Profile",
-		meta=(EditCondition="bIsPlantProfileDropdownEnabled", EditConditionHides, HideEditConditionToggle, GetOptions="GetPlantProfileOptions",
+		meta=(PCG_NotOverridable, EditCondition="bIsPlantProfileDropdownEnabled", EditConditionHides, HideEditConditionToggle, GetOptions="GetPlantProfileOptions",
 			Tooltip="Select a profile (cross sectional shape) from predefined options\n\nNone applies a simple circular cross sectional shape"))
 	FString SelectedPlantProfile = "None";
 
diff --git a/Engine/Plugins/Experimental/ProceduralVegetationEditor/Source/ProceduralVegetation/Private/Nodes/PVBoneReductionSettings.h b/Engine/Plugins/Experimental/ProceduralVegetationEditor/Source/ProceduralVegetation/Private/Nodes/PVBoneReductionSettings.h
--- a/Engine/Plugins/Experimental/ProceduralVegetationEditor/Source/ProceduralVegetation/Private/Nodes/PVBoneReductionSettings.h
+++ b/Engine/Plugins/Experimental/ProceduralVegetationEditor/Source/ProceduralVegetation/Private/Nodes/PVBoneReductionSettings.h
@@ -30,8 +30,8 @@ protected:
 	virtual FPCGElementPtr CreateElement() const override;
 
 private:
-	UPROPERTY(EditAnywhere, Category = "Bone Reduction", meta = (ShowOnlyInnerProperties, ClampMin ="0", ClampMax = "1", UIMin = "0", UIMax = "1", Tooltip="Controls how aggressively bones are reduced.\n\nControls how aggressively bones are reduced. Higher values result in less bones. Reduced bone complexity helps in optimizing performance but can effect the accuracy of wind simulation."))
-	float Strength;
+	UPROPERTY(EditAnywhere, Category = "Bone Reduction", meta = (PCG_Overridable, ShowOnlyInnerProperties, ClampMin ="0", ClampMax = "1", UIMin = "0", UIMax = "1", Tooltip="Controls how aggressively bones are reduced.\n\nControls how aggressively bones are reduced. Higher values result in less bones. Reduced bone complexity helps in optimizing performance but can effect the accuracy of wind simulation."))
+	float Strength = 0;
 };
 
 class FPVBoneReductionElement : public FPVBaseElement
diff --git a/Engine/Plugins/Experimental/ProceduralVegetationEditor/Source/ProceduralVegetation/Private/Nodes/PVFoliageDistributorSettings.h b/Engine/Plugins/Experimental/ProceduralVegetationEditor/Source/ProceduralVegetation/Private/Nodes/PVFoliageDistributorSettings.h
--- a/Engine/Plugins/Experimental/ProceduralVegetationEditor/Source/ProceduralVegetation/Private/Nodes/PVFoliageDistributorSettings.h
+++ b/Engine/Plugins/Experimental/ProceduralVegetationEditor/Source/ProceduralVegetation/Private/Nodes/PVFoliageDistributorSettings.h
@@ -42,82 +42,82 @@ protected:
 
 private:
 
-	UPROPERTY(EditAnywhere, Category="Distribution Settings", meta=(ClampMin=0.0f, ClampMax=1.0f, UIMin=0.0f, UIMax=1.0f, Tooltip="Threshold controlling foliage release based on an ethylene-like signal.\n\nHigher Threshold → more buds/branches retained, making the plant denser and bushier.\nLower Threshold → fewer buds/branches retained, resulting in a sparser, pruned structure."))
+	UPROPERTY(EditAnywhere, Category="Distribution Settings", meta=(PCG_Overridable, ClampMin=0.0f, ClampMax=1.0f, UIMin=0.0f, UIMax=1.0f, Tooltip="Threshold controlling foliage release based on an ethylene-like signal.\n\nHigher Threshold → more buds/branches retained, making the plant denser and bushier.\nLower Threshold → fewer buds/branches retained, resulting in a sparser, pruned structure."))
 	float EthyleneThreshold = 0.85f;
 
 	UPROPERTY(EditAnywhere, Category="Distribution Settings", meta=(Tooltip="Override default placement with custom distribution rules.\n\nOverides the distribution settings loaded from the procedural vegetation preset."))
 	bool OverrideDistribution = true;
 
-	UPROPERTY(EditAnywhere, Category="Distribution Settings", meta=(EditCondition="OverrideDistribution", ClampMin=0.0f, ClampMax=1.0f, UIMin=0.0f, UIMax=1.0f, Tooltip="Minimum distance between placed instances.\n\nHigher spacing yields a more open distribution; lower spacing allows denser placement."))
+	UPROPERTY(EditAnywhere, Category="Distribution Settings", meta=(PCG_Overridable, EditCondition="OverrideDistribution", ClampMin=0.0f, ClampMax=1.0f, UIMin=0.0f, UIMax=1.0f, Tooltip="Minimum distance between placed instances.\n\nHigher spacing yields a more open distribution; lower spacing allows denser placement."))
 	float InstanceSpacing = 0.3f;
 
-	UPROPERTY(EditAnywhere, Category="Distribution Settings", meta=(EditCondition="OverrideDistribution", XAxisMin=0.0f, XAxisMax=1.0f, YAxisMin=0.0f, YAxisMax=1.0f, Tooltip="Ramp that varies instance spacing across the height of the plant\n\nUseful to reduce spacing in some areas based on the height of the plant."))
+	UPROPERTY(EditAnywhere, Category="Distribution Settings", meta=(PCG_Overridable, EditCondition="OverrideDistribution", XAxisMin=0.0f, XAxisMax=1.0f, YAxisMin=0.0f, YAxisMax=1.0f, Tooltip="Ramp that varies instance spacing across the height of the plant\n\nUseful to reduce spacing in some areas based on the height of the plant."))
 	FPVFloatRamp InstanceSpacingRamp;
 
-	UPROPERTY(EditAnywhere, Category="Distribution Settings", meta=(EditCondition="OverrideDistribution", ClampMin=0.0f, ClampMax=1.0f, UIMin=0.0f, UIMax=1.0f, Tooltip="Strength of the spacing ramp’s influence.\n\nBlends the ramp with the base spacing. Lower values favor the base spacing (subtle effect); higher values apply the ramp more strongly and can amplify spacing differences."))
+	UPROPERTY(EditAnywhere, Category="Distribution Settings", meta=(PCG_Overridable, EditCondition="OverrideDistribution", ClampMin=0.0f, ClampMax=1.0f, UIMin=0.0f, UIMax=1.0f, Tooltip="Strength of the spacing ramp’s influence.\n\nBlends the ramp with the base spacing. Lower values favor the base spacing (subtle effect); higher values apply the ramp more strongly and can amplify spacing differences."))
 	float InstanceSpacingRampEffect = 0.0f;
 
-	UPROPERTY(EditAnywhere, Category="Distribution Settings", meta=(EditCondition="OverrideDistribution", ClampMin=-1, ClampMax=1000, Tooltip="Caps the number of instances per branch\n\nlimit on how many foliage instances can be placed per branch to avoid overcrowding."))
+	UPROPERTY(EditAnywhere, Category="Distribution Settings", meta=(PCG_Overridable, EditCondition="OverrideDistribution", ClampMin=-1, ClampMax=1000, Tooltip="Caps the number of instances per branch\n\nlimit on how many foliage instances can be placed per branch to avoid overcrowding."))
 	int32 MaxPerBranch = -1;
 
-	UPROPERTY(EditAnywhere, Category="Scale Settings", meta=(EditCondition="OverrideDistribution", ClampMin=0.0f, UIMin=0.0f, UIMax=5.0f, Tooltip="Default size multiplier for placed foliage\n\nSets the base uniform scale applied to each instance before randomness and ramps are considered. Increase to enlarge all instances proportionally; decrease for smaller foliage. Final size still respects Min/Max limits."))
+	UPROPERTY(EditAnywhere, Category="Scale Settings", meta=(PCG_Overridable, EditCondition="OverrideDistribution", ClampMin=0.0f, UIMin=0.0f, UIMax=5.0f, Tooltip="Default size multiplier for placed foliage\n\nSets the base uniform scale applied to each instance before randomness and ramps are considered. Increase to enlarge all instances proportionally; decrease for smaller foliage. Final size still respects Min/Max limits."))
 	float BaseScale = 1.0f;
 
-	UPROPERTY(EditAnywhere, Category="Scale Settings", meta=(EditCondition="OverrideDistribution", ClampMin=0.0f, ClampMax=1.0f, UIMin=0.0f, UIMax=1.0f, Tooltip="Impact of branch size on foliage size\n\nThe size of the foliage being scattered will be chosen in relation to the physical size of the branch it is distributed on. The largest branches will be distributed the largest foliage pieces."))
+	UPROPERTY(EditAnywhere, Category="Scale Settings", meta=(PCG_Overridable, EditCondition="OverrideDistribution", ClampMin=0.0f, ClampMax=1.0f, UIMin=0.0f, UIMax=1.0f, Tooltip="Impact of branch size on foliage size\n\nThe size of the foliage being scattered will be chosen in relation to the physical size of the branch it is distributed on. The largest branches will be distributed the largest foliage pieces."))
 	float BranchScaleImpact = 0.0f;
 
-	UPROPERTY(EditAnywhere, Category="Scale Settings", meta=(EditCondition="OverrideDistribution", ClampMin=0.0f, UIMin=0.0f, UIMax=5.0f, Tooltip="Minimum allowed instance size.\n\nClamps the final computed scale so it never goes below this value. Use to prevent tiny, hard-to-see instances."))
+	UPROPERTY(EditAnywhere, Category="Scale Settings", meta=(PCG_Overridable, EditCondition="OverrideDistribution", ClampMin=0.0f, UIMin=0.0f, UIMax=5.0f, Tooltip="Minimum allowed instance size.\n\nClamps the final computed scale so it never goes below this value. Use to prevent tiny, hard-to-see instances."))
 	float MinScale = 0.5f;
 
-	UPROPERTY(EditAnywhere, Category="Scale Settings", meta=(EditCondition="OverrideDistribution", ClampMin=0.0f, UIMin=0.0f, UIMax=5.0f, Tooltip="Maximum allowed instance size.\n\nClamps the final computed scale so it never exceeds this value. Use to cap unusually large instances for visual consistency."))
+	UPROPERTY(EditAnywhere, Category="Scale Settings", meta=(PCG_Overridable, EditCondition="OverrideDistribution", ClampMin=0.0f, UIMin=0.0f, UIMax=5.0f, Tooltip="Maximum allowed instance size.\n\nClamps the final computed scale so it never exceeds this value. Use to cap unusually large instances for visual consistency."))
 	float MaxScale = 1.0f;
 
-	UPROPERTY(EditAnywhere, Category="Scale Settings", meta=(EditCondition="OverrideDistribution", ClampMin=0.0f, UIMin=0.0f, UIMax=5.0f, Tooltip="Lower bound for per-instance size variation.\n\nSets the minimum of the random scale range used to vary instance sizes. Each instance samples within Random Scale Min–Max, then combines with Base Scale and any ramps (respecting Min/Max Scale)."))
+	UPROPERTY(EditAnywhere, Category="Scale Settings", meta=(PCG_Overridable, EditCondition="OverrideDistribution", ClampMin=0.0f, UIMin=0.0f, UIMax=5.0f, Tooltip="Lower bound for per-instance size variation.\n\nSets the minimum of the random scale range used to vary instance sizes. Each instance samples within Random Scale Min–Max, then combines with Base Scale and any ramps (respecting Min/Max Scale)."))
 	float RandomScaleMin = 1.0f;
 
-	UPROPERTY(EditAnywhere, Category="Scale Settings", meta=(EditCondition="OverrideDistribution", ClampMin=1.0f, UIMin=1.0f, UIMax=5.0f, Tooltip="Upper bound for per-instance size variation.\n\nSets the maximum of the random scale range used to vary instance sizes. Each instance samples within Random Scale Min–Max, then combines with Base Scale and any ramps (respecting Min/Max Scale)."))
+	UPROPERTY(EditAnywhere, Category="Scale Settings", meta=(PCG_Overridable, EditCondition="OverrideDistribution", ClampMin=1.0f, UIMin=1.0f, UIMax=5.0f, Tooltip="Upper bound for per-instance size variation.\n\nSets the maximum of the random scale range used to vary instance sizes. Each instance samples within Random Scale Min–Max, then combines with Base Scale and any ramps (respecting Min/Max Scale)."))
 	float RandomScaleMax = 1.0f;
 
 	UPROPERTY(EditAnywhere, Category="Scale Settings", meta=(EditCondition="OverrideDistribution", XAxisMin=0.0f, XAxisMax=1.0f, YAxisMin=0.0f, YAxisMax=1.0f, Tooltip="Curve that varies scale across the height of the plant\n\nUse to taper foliage toward tips or emphasize size in targeted regions."))
 	FPVFloatRamp ScaleRamp;
 
-	UPROPERTY(EditAnywhere, Category="Vector Settings", meta=(EditCondition="OverrideDistribution", Tooltip="Use a custom axil angle instead of the default.\n\nEnables a user-defined orientation. When on, the system uses AxilAngle and the AxilAngleRamp settings (and Effect) to set the tilt of foliage relative to its parent branch direction, overriding any automatic/species defaults."))
+	UPROPERTY(EditAnywhere, Category="Vector Settings", meta=(PCG_Overridable, EditCondition="OverrideDistribution", Tooltip="Use a custom axil angle instead of the default.\n\nEnables a user-defined orientation. When on, the system uses AxilAngle and the AxilAngleRamp settings (and Effect) to set the tilt of foliage relative to its parent branch direction, overriding any automatic/species defaults."))
 	bool OverrideAxilAngle = true;
 
-	UPROPERTY(EditAnywhere, Category="Vector Settings", meta=(EditCondition="OverrideDistribution", ClampMin=0.0f, ClampMax=90.0f, UIMin=0.0f, UIMax=90.0f, Tooltip="Base tilt from the parent axis\n\nSets the base axil angle in degrees—the tilt of the instance relative to the parent branch’s direction. Acts as the lower bound when used with a ramp; higher values tilt foliage further away from the branch axis."))
+	UPROPERTY(EditAnywhere, Category="Vector Settings", meta=(PCG_Overridable, EditCondition="OverrideDistribution", ClampMin=0.0f, ClampMax=90.0f, UIMin=0.0f, UIMax=90.0f, Tooltip="Base tilt from the parent axis\n\nSets the base axil angle in degrees—the tilt of the instance relative to the parent branch’s direction. Acts as the lower bound when used with a ramp; higher values tilt foliage further away from the branch axis."))
 	float AxilAngle = 45.0f;
 
 	UPROPERTY(EditAnywhere, Category="Vector Settings", meta=(EditCondition="OverrideDistribution", XAxisMin=0.0f, XAxisMax=1.0f, YAxisMin=0.0f, YAxisMax=1.0f, Tooltip="Ramp that varies the axil angle across the height of the plant\n\nIts output is remapped to an angle range defined by AxilAngle and AxilAngleRampUpperValue, modulated by AxilAngleRampEffect."))
 	FPVFloatRamp AxilAngleRamp;
 
-	UPROPERTY(EditAnywhere, Category="Vector Settings", meta=(EditCondition="OverrideDistribution", ClampMin=0.0f, ClampMax=90.0f, UIMin=0.0f, UIMax=90.0f, Tooltip="Max angle used when the ramp\n\nProvides detailed control over this parameter to balance visual quality and performance. Valid values: 0.0–90.0."))
+	UPROPERTY(EditAnywhere, Category="Vector Settings", meta=(PCG_Overridable, EditCondition="OverrideDistribution", ClampMin=0.0f, ClampMax=90.0f, UIMin=0.0f, UIMax=90.0f, Tooltip="Max angle used when the ramp\n\nProvides detailed control over this parameter to balance visual quality and performance. Valid values: 0.0–90.0."))
 	float AxilAngleRampUpperValue = 45.0f;
 
-	UPROPERTY(EditAnywhere, Category="Vector Settings", meta=(EditCondition="OverrideDistribution", ClampMin=0.0f, ClampMax=1.0f, UIMin=0.0f, UIMax=1.0f, Tooltip="Adjusts this parameter for finer control.\n\nDefines the upper axil angle reached when AxilAngleRamp evaluates to 1. Together with AxilAngle, it forms the min–max range for ramp-driven orientation."))
+	UPROPERTY(EditAnywhere, Category="Vector Settings", meta=(PCG_Overridable, EditCondition="OverrideDistribution", ClampMin=0.0f, ClampMax=1.0f, UIMin=0.0f, UIMax=1.0f, Tooltip="Adjusts this parameter for finer control.\n\nDefines the upper axil angle reached when AxilAngleRamp evaluates to 1. Together with AxilAngle, it forms the min–max range for ramp-driven orientation."))
 	float AxilAngleRampEffect = 0.0f;
 
 	UPROPERTY(EditAnywhere, Category="Phyllotaxy Settings", meta=(EditCondition="OverrideDistribution", Tooltip="Use custom phyllotaxy settings instead of defaults.\n\nOveride the phyllotaxy settings loaded from the procedural vegetation preset."))
 	bool OverridePhyllotaxy = true;
 
-	UPROPERTY(EditAnywhere, Category="Phyllotaxy Settings", meta=(EditCondition="OverrideDistribution", Tooltip="Pattern for arranging buds/leaves around the stem.\n\nSelects the phyllotactic pattern used for placement (e.g., Alternate, Opposite, Whorled, Spiral). This defines the sequence and angular spacing of buds/leaves around the stem and along its length."))
+	UPROPERTY(EditAnywhere, Category="Phyllotaxy Settings", meta=(PCG_Overridable, EditCondition="OverrideDistribution", Tooltip="Pattern for arranging buds/leaves around the stem.\n\nSelects the phyllotactic pattern used for placement (e.g., Alternate, Opposite, Whorled, Spiral). This defines the sequence and angular spacing of buds/leaves around the stem and along its length."))
 	EPhyllotaxyType PhyllotaxyType = EPhyllotaxyType::Alternate;
 
-	UPROPERTY(EditAnywhere, Category="Phyllotaxy Settings", meta=(EditCondition="OverrideDistribution && PhyllotaxyType == EPhyllotaxyType::Spiral", EditConditionHides))
+	UPROPERTY(EditAnywhere, Category="Phyllotaxy Settings", meta=(PCG_Overridable, EditCondition="OverrideDistribution && PhyllotaxyType == EPhyllotaxyType::Spiral", EditConditionHides))
 	EPhyllotaxyFormation PhyllotaxyFormation = EPhyllotaxyFormation::Distichous;
 
-	UPROPERTY(EditAnywhere, Category="Phyllotaxy Settings", meta=(EditCondition="OverrideDistribution", ClampMin=1, ClampMax=10, UIMin=1, UIMax=10, Tooltip="Minimum buds/leaves per node.\n\nSets the lower bound on the number of buds/leaves generated at each node. Guarantees at least this many instances (subject to spacing and other constraints). Should be ≤ MaximumNodeBuds. Applied when Whorled phyllotaxy is selected."))
+	UPROPERTY(EditAnywhere, Category="Phyllotaxy Settings", meta=(PCG_Overridable, EditCondition="OverrideDistribution", ClampMin=1, ClampMax=10, UIMin=1, UIMax=10, Tooltip="Minimum buds/leaves per node.\n\nSets the lower bound on the number of buds/leaves generated at each node. Guarantees at least this many instances (subject to spacing and other constraints). Should be ≤ MaximumNodeBuds. Applied when Whorled phyllotaxy is selected."))
 	int32 MinimumNodeBuds = 2;
 
-	UPROPERTY(EditAnywhere, Category="Phyllotaxy Settings", meta=(EditCondition="OverrideDistribution", ClampMin=1, ClampMax=10, UIMin=1, UIMax=10, Tooltip="Maximum buds/leaves per node.\n\nCaps the number of buds/leaves spawned. Use with MinimumNodeBuds to bound density and keep arrangements predictable. Applied when Whorled phyllotaxy is selected."))
+	UPROPERTY(EditAnywhere, Category="Phyllotaxy Settings", meta=(PCG_Overridable, EditCondition="OverrideDistribution", ClampMin=1, ClampMax=10, UIMin=1, UIMax=10, Tooltip="Maximum buds/leaves per node.\n\nCaps the number of buds/leaves spawned. Use with MinimumNodeBuds to bound density and keep arrangements predictable. Applied when Whorled phyllotaxy is selected."))
 	int32 MaximumNodeBuds = 3;
 
-	UPROPERTY(EditAnywhere, Category="Phyllotaxy Settings", meta=(EditCondition="OverrideDistribution", ClampMin=0.0f, ClampMax=360.0f, UIMin=0.0f, UIMax=360.0f, Tooltip="Extra divergence angle (°) added to the phyllotaxy pattern\n\nAdds an angular offset to the base phyllotactic divergence."))
+	UPROPERTY(EditAnywhere, Category="Phyllotaxy Settings", meta=(PCG_Overridable, EditCondition="OverrideDistribution", ClampMin=0.0f, ClampMax=360.0f, UIMin=0.0f, UIMax=360.0f, Tooltip="Extra divergence angle (°) added to the phyllotaxy pattern\n\nAdds an angular offset to the base phyllotactic divergence."))
 	float PhyllotaxyAdditionalAngle = 0.0f;
 
-	UPROPERTY(EditAnywhere, Category="Misc Settings", meta=(EditCondition="OverrideDistribution"))
+	UPROPERTY(EditAnywhere, Category="Misc Settings", meta=(PCG_Overridable, EditCondition="OverrideDistribution"))
 	int32 RandomSeed = 123456;
 
-	UPROPERTY(EditAnywhere, Category="ConditionSettings", meta = (ShowOnlyInnerProperties, EditCondition="OverrideDistribution"))
+	UPROPERTY(EditAnywhere, Category="ConditionSettings", meta = (PCG_Overridable, ShowOnlyInnerProperties, EditCondition="OverrideDistribution"))
 	FPVFoliageDistributionConditions Conditions;
 	
 	UPROPERTY()
diff --git a/Engine/Plugins/Experimental/ProceduralVegetationEditor/Source/ProceduralVegetation/Private/Nodes/PVGrafterSettings.h b/Engine/Plugins/Experimental/ProceduralVegetationEditor/Source/ProceduralVegetation/Private/Nodes/PVGrafterSettings.h
--- a/Engine/Plugins/Experimental/ProceduralVegetationEditor/Source/ProceduralVegetation/Private/Nodes/PVGrafterSettings.h
+++ b/Engine/Plugins/Experimental/ProceduralVegetationEditor/Source/ProceduralVegetation/Private/Nodes/PVGrafterSettings.h
@@ -44,65 +44,65 @@ protected:
 
 private:
 	UPROPERTY(EditAnywhere, Category="Distribution Settings",
-		meta=(ClampMin=0.0f, ClampMax=1.0f, UIMin=0.0f, UIMax=1.0f, Tooltip=
+		meta=(PCG_Overridable, ClampMin=0.0f, ClampMax=1.0f, UIMin=0.0f, UIMax=1.0f, Tooltip=
 			"Threshold controlling foliage release based on an ethylene-like signal.\n\nHigher Threshold → more buds/branches retained, making the plant denser and bushier.\nLower Threshold → fewer buds/branches retained, resulting in a sparser, pruned structure."
 		))
 	float EthyleneThreshold = 0.85f;
 
 	UPROPERTY(EditAnywhere, Category="Distribution Settings",
-		meta=(ClampMin=0.0f, ClampMax=1.0f, UIMin=0.0f, UIMax=1.0f, Tooltip=
+		meta=(PCG_Overridable, ClampMin=0.0f, ClampMax=1.0f, UIMin=0.0f, UIMax=1.0f, Tooltip=
 			"Minimum distance between placed instances.\n\nHigher spacing yields a more open distribution; lower spacing allows denser placement."))
 	float InstanceSpacing = 0.3f;
 
 	UPROPERTY(EditAnywhere, Category="Distribution Settings",
-		meta=(XAxisMin=0.0f, XAxisMax=1.0f, YAxisMin=0.0f, YAxisMax=1.0f, Tooltip=
+		meta=(PCG_Overridable, XAxisMin=0.0f, XAxisMax=1.0f, YAxisMin=0.0f, YAxisMax=1.0f, Tooltip=
 			"Ramp that varies instance spacing across the height of the plant\n\nUseful to reduce spacing in some areas based on the height of the plant."
 		))
 	FPVFloatRamp InstanceSpacingRamp;
 
 	UPROPERTY(EditAnywhere, Category="Distribution Settings",
-		meta=(ClampMin=0.0f, ClampMax=1.0f, UIMin=0.0f, UIMax=1.0f, Tooltip=
+		meta=(PCG_Overridable, ClampMin=0.0f, ClampMax=1.0f, UIMin=0.0f, UIMax=1.0f, Tooltip=
 			"Strength of the spacing ramp’s influence.\n\nBlends the ramp with the base spacing. Lower values favor the base spacing (subtle effect); higher values apply the ramp more strongly and can amplify spacing differences."
 		))
 	float InstanceSpacingRampEffect = 0.0f;
 
 	UPROPERTY(EditAnywhere, Category="Distribution Settings",
-		meta=(ClampMin=-1, ClampMax=1000, Tooltip=
+		meta=(PCG_Overridable, ClampMin=-1, ClampMax=1000, Tooltip=
 			"Caps the number of instances per branch\n\nlimit on how many foliage instances can be placed per branch to avoid overcrowding."))
 	int32 MaxPerBranch = -1;
 
 	UPROPERTY(EditAnywhere, Category="Scale Settings",
-		meta=(ClampMin=0.0f, UIMin=0.0f, UIMax=5.0f, Tooltip=
+		meta=(PCG_Overridable, ClampMin=0.0f, UIMin=0.0f, UIMax=5.0f, Tooltip=
 			"Default size multiplier for placed foliage\n\nSets the base uniform scale applied to each instance before randomness and ramps are considered. Increase to enlarge all instances proportionally; decrease for smaller foliage. Final size still respects Min/Max limits."
 		))
 	float BaseScale = 1.0f;
 
 	UPROPERTY(EditAnywhere, Category="Scale Settings",
-		meta=(ClampMin=0.0f, ClampMax=1.0f, UIMin=0.0f, UIMax=1.0f, Tooltip=
+		meta=(PCG_Overridable, ClampMin=0.0f, ClampMax=1.0f, UIMin=0.0f, UIMax=1.0f, Tooltip=
 			"Impact of branch size on foliage size\n\nThe size of the foliage being scattered will be chosen in relation to the physical size of the branch it is distributed on. The largest branches will be distributed the largest foliage pieces."
 		))
 	float BranchScaleImpact = 0.0f;
 
 	UPROPERTY(EditAnywhere, Category="Scale Settings",
-		meta=(ClampMin=0.0f, UIMin=0.0f, UIMax=5.0f, Tooltip=
+		meta=(PCG_Overridable, ClampMin=0.0f, UIMin=0.0f, UIMax=5.0f, Tooltip=
 			"Minimum allowed instance size.\n\nClamps the final computed scale so it never goes below this value. Use to prevent tiny, hard-to-see instances."
 		))
 	float MinScale = 0.5f;
 
 	UPROPERTY(EditAnywhere, Category="Scale Settings",
-		meta=(ClampMin=0.0f, UIMin=0.0f, UIMax=5.0f, Tooltip=
+		meta=(PCG_Overridable, ClampMin=0.0f, UIMin=0.0f, UIMax=5.0f, Tooltip=
 			"Maximum allowed instance size.\n\nClamps the final computed scale so it never exceeds this value. Use to cap unusually large instances for visual consistency."
 		))
 	float MaxScale = 1.0f;
 
 	UPROPERTY(EditAnywhere, Category="Scale Settings",
-		meta=(ClampMin=0.0f, UIMin=0.0f, UIMax=5.0f, Tooltip=
+		meta=(PCG_Overridable, ClampMin=0.0f, UIMin=0.0f, UIMax=5.0f, Tooltip=
 			"Lower bound for per-instance size variation.\n\nSets the minimum of the random scale range used to vary instance sizes. Each instance samples within Random Scale Min–Max, then combines with Base Scale and any ramps (respecting Min/Max Scale)."
 		))
 	float RandomScaleMin = 1.0f;
 
 	UPROPERTY(EditAnywhere, Category="Scale Settings",
-		meta=(ClampMin=1.0f, UIMin=1.0f, UIMax=5.0f, Tooltip=
+		meta=(PCG_Overridable, ClampMin=1.0f, UIMin=1.0f, UIMax=5.0f, Tooltip=
 			"Upper bound for per-instance size variation.\n\nSets the maximum of the random scale range used to vary instance sizes. Each instance samples within Random Scale Min–Max, then combines with Base Scale and any ramps (respecting Min/Max Scale)."
 		))
 	float RandomScaleMax = 1.0f;
@@ -113,13 +113,13 @@ private:
 	FPVFloatRamp ScaleRamp;
 
 	UPROPERTY(EditAnywhere, Category="Vector Settings",
-		meta=(Tooltip=
+		meta=(PCG_Overridable, Tooltip=
 			"Use a custom axil angle instead of the default.\n\nEnables a user-defined orientation. When on, the system uses AxilAngle and the AxilAngleRamp settings (and Effect) to set the tilt of foliage relative to its parent branch direction, overriding any automatic/species defaults."
 		))
 	bool OverrideAxilAngle = true;
 
 	UPROPERTY(EditAnywhere, Category="Vector Settings",
-		meta=(ClampMin=0.0f, ClampMax=90.0f, UIMin=0.0f, UIMax=90.0f, Tooltip=
+		meta=(PCG_Overridable, ClampMin=0.0f, ClampMax=90.0f, UIMin=0.0f, UIMax=90.0f, Tooltip=
 			"Base tilt from the parent axis\n\nSets the base axil angle in degrees—the tilt of the instance relative to the parent branch’s direction. Acts as the lower bound when used with a ramp; higher values tilt foliage further away from the branch axis."
 		))
 	float AxilAngle = 45.0f;
@@ -131,18 +131,18 @@ private:
 	FPVFloatRamp AxilAngleRamp;
 
 	UPROPERTY(EditAnywhere, Category="Vector Settings",
-		meta=(ClampMin=0.0f, ClampMax=90.0f, UIMin=0.0f, UIMax=90.0f, Tooltip=
+		meta=(PCG_Overridable, ClampMin=0.0f, ClampMax=90.0f, UIMin=0.0f, UIMax=90.0f, Tooltip=
 			"Max angle used when the ramp\n\nProvides detailed control over this parameter to balance visual quality and performance. Valid values: 0.0–90.0."
 		))
 	float AxilAngleRampUpperValue = 45.0f;
 
 	UPROPERTY(EditAnywhere, Category="Vector Settings",
-		meta=(ClampMin=0.0f, ClampMax=1.0f, UIMin=0.0f, UIMax=1.0f, Tooltip=
+		meta=(PCG_Overridable, ClampMin=0.0f, ClampMax=1.0f, UIMin=0.0f, UIMax=1.0f, Tooltip=
 			"Adjusts this parameter for finer control.\n\nDefines the upper axil angle reached when AxilAngleRamp evaluates to 1. Together with AxilAngle, it forms the min–max range for ramp-driven orientation."
 		))
 	float AxilAngleRampEffect = 0.0f;
 
-	UPROPERTY(EditAnywhere, Category="Misc Settings")
+	UPROPERTY(EditAnywhere, Category="Misc Settings", meta = (PCG_Overridable))
 	int32 RandomSeed = 123456;
 };
 
diff --git a/Engine/Plugins/Experimental/ProceduralVegetationEditor/Source/ProceduralVegetation/Private/Nodes/PVMeshBuilderSettings.h b/Engine/Plugins/Experimental/ProceduralVegetationEditor/Source/ProceduralVegetation/Private/Nodes/PVMeshBuilderSettings.h
--- a/Engine/Plugins/Experimental/ProceduralVegetationEditor/Source/ProceduralVegetation/Private/Nodes/PVMeshBuilderSettings.h
+++ b/Engine/Plugins/Experimental/ProceduralVegetationEditor/Source/ProceduralVegetation/Private/Nodes/PVMeshBuilderSettings.h
@@ -36,7 +36,7 @@ protected:
 	virtual FPCGDataTypeIdentifier GetOutputPinTypeIdentifier() const override;
 
 public:	
-	UPROPERTY(EditAnywhere, Category = "Mesh", meta = (ShowOnlyInnerProperties))
+	UPROPERTY(EditAnywhere, Category = "Mesh", meta = (ShowOnlyInnerProperties, PCG_Overridable))
 	FPVMeshBuilderParams MesherSettings;
 
 	UPROPERTY()
diff --git a/Engine/Plugins/Experimental/ProceduralVegetationEditor/Source/ProceduralVegetation/Public/Nodes/PVOutputSettings.h b/Engine/Plugins/Experimental/ProceduralVegetationEditor/Source/ProceduralVegetation/Public/Nodes/PVOutputSettings.h
--- a/Engine/Plugins/Experimental/ProceduralVegetationEditor/Source/ProceduralVegetation/Public/Nodes/PVOutputSettings.h
+++ b/Engine/Plugins/Experimental/ProceduralVegetationEditor/Source/ProceduralVegetation/Public/Nodes/PVOutputSettings.h
@@ -35,7 +35,7 @@ protected:
 #endif
 	
 public:
-	UPROPERTY(EditAnywhere, Category = ExportSettings, meta = (ShowOnlyInnerProperties))
+	UPROPERTY(EditAnywhere, Category = ExportSettings, meta = (PCG_Overridable, ShowOnlyInnerProperties))
 	FPVExportParams ExportSettings;
 };
 
