[PVE] Disabled PVE debug flag and added a null check for Collection widget

#Synced-CL 51086708

diff --git a/Engine/Plugins/Experimental/ProceduralVegetationEditor/Source/ProceduralVegetation/Private/Helpers/PVUtilities.cpp b/Engine/Plugins/Experimental/ProceduralVegetationEditor/Source/ProceduralVegetation/Private/Helpers/PVUtilities.cpp
--- a/Engine/Plugins/Experimental/ProceduralVegetationEditor/Source/ProceduralVegetation/Private/Helpers/PVUtilities.cpp
+++ b/Engine/Plugins/Experimental/ProceduralVegetationEditor/Source/ProceduralVegetation/Private/Helpers/PVUtilities.cpp
@@ -24,7 +24,7 @@ namespace PV::Utilities
 #if WITH_EDITOR
 	static TAutoConsoleVariable<bool> CVarEnablePVDebugMode(
 		TEXT("PV.DebugMode.Enabled"),
-		true,
+		false,
 		TEXT("Enables debug mode for the Procedural Vegetation Editor"), 
 		FConsoleVariableDelegate::CreateLambda([](IConsoleVariable* ConsoleVariable)
 			{
diff --git a/Engine/Plugins/Experimental/ProceduralVegetationEditor/Source/ProceduralVegetationEditor/Private/PVEditor.cpp b/Engine/Plugins/Experimental/ProceduralVegetationEditor/Source/ProceduralVegetationEditor/Private/PVEditor.cpp
--- a/Engine/Plugins/Experimental/ProceduralVegetationEditor/Source/ProceduralVegetationEditor/Private/PVEditor.cpp
+++ b/Engine/Plugins/Experimental/ProceduralVegetationEditor/Source/ProceduralVegetationEditor/Private/PVEditor.cpp
@@ -511,6 +511,9 @@ void FPVEditor::UpdateInspectedCollection()
 
 void FPVEditor::UpdateCollectionSpreadSheet()
 {
+	if (!CollectionSpreadSheetWidget)
+		return;
+	
 	if (!NodeBeingInspected)
 	{
 		CollectionSpreadSheetWidget->SetData(FString());
