Fix memory corruption crashes related to accessing UPVData::DebugSettings on multiple threads

#Synced-CL 51206265
#Synced-CL 51206265

diff --git a/Engine/Plugins/Experimental/ProceduralVegetationEditor/Source/ProceduralVegetation/Private/DataTypes/PVData.cpp b/Engine/Plugins/Experimental/ProceduralVegetationEditor/Source/ProceduralVegetation/Private/DataTypes/PVData.cpp
--- a/Engine/Plugins/Experimental/ProceduralVegetationEditor/Source/ProceduralVegetation/Private/DataTypes/PVData.cpp
+++ b/Engine/Plugins/Experimental/ProceduralVegetationEditor/Source/ProceduralVegetation/Private/DataTypes/PVData.cpp
@@ -80,6 +80,24 @@ void UPVData::AddVisualizationCollection(FManagedArrayCollection InCollection, T
 {
 	VisualizationCollections.Emplace(MoveTemp(InCollection), MoveTemp(InVisualizationSettings));
 }
+
+void UPVData::SetDebugSettings(const FPVDebugSettings& InSettings) const
+{
+	const FScopeLock Lock(&DebugSettingsMutex);
+	DebugSettings = InSettings;
+}
+
+void UPVData::SetDebugSettings(FPVDebugSettings&& InSettings) const
+{ 
+	const FScopeLock Lock(&DebugSettingsMutex);
+	DebugSettings = MoveTemp(InSettings);
+}
+
+FPVDebugSettings UPVData::GetDebugSettings() const 
+{ 
+	const FScopeLock Lock(&DebugSettingsMutex);
+	return DebugSettings;
+}
 #endif
 
 const UPCGPointArrayData* UPVData::ToPointArrayData(FPCGContext* Context, const FBox& InBounds) const
diff --git a/Engine/Plugins/Experimental/ProceduralVegetationEditor/Source/ProceduralVegetation/Public/DataTypes/PVData.h b/Engine/Plugins/Experimental/ProceduralVegetationEditor/Source/ProceduralVegetation/Public/DataTypes/PVData.h
--- a/Engine/Plugins/Experimental/ProceduralVegetationEditor/Source/ProceduralVegetation/Public/DataTypes/PVData.h
+++ b/Engine/Plugins/Experimental/ProceduralVegetationEditor/Source/ProceduralVegetation/Public/DataTypes/PVData.h
@@ -211,9 +211,9 @@ public:
 
 #if WITH_EDITORONLY_DATA
 	void AddVisualizationCollection(FManagedArrayCollection InCollection, TArray<FPVVisualizationSettings> InVisualizationSettings = {});
-	void SetDebugSettings(const FPVDebugSettings& InSettings) const { DebugSettings = InSettings; }
-	void SetDebugSettings(FPVDebugSettings&& InSettings) const { DebugSettings = MoveTemp(InSettings); }
-	const FPVDebugSettings& GetDebugSettings() const { return DebugSettings; }
+	void SetDebugSettings(const FPVDebugSettings& InSettings) const;
+	void SetDebugSettings(FPVDebugSettings&& InSettings) const;
+	FPVDebugSettings GetDebugSettings() const;
 	const TArray<FPVVisualizationCollection>& GetVisualizationCollections() const { return VisualizationCollections; }
 #endif
 	
@@ -242,6 +242,7 @@ protected:
 	TArray<TObjectPtr<const UObject>> ManagedResources;
 
 #if WITH_EDITOR
+	mutable FCriticalSection DebugSettingsMutex;
 	mutable FPVDebugSettings DebugSettings;
 
 	TArray<FPVVisualizationCollection> VisualizationCollections;
