[PVE] Removed Preset node from public facing nodes.

#Synced-CL 51086708

diff --git a/Engine/Plugins/Experimental/ProceduralVegetationEditor/Source/ProceduralVegetation/Private/Helpers/PVUtilities.cpp b/Engine/Plugins/Experimental/ProceduralVegetationEditor/Source/ProceduralVegetation/Private/Helpers/PVUtilities.cpp
--- a/Engine/Plugins/Experimental/ProceduralVegetationEditor/Source/ProceduralVegetation/Private/Helpers/PVUtilities.cpp
+++ b/Engine/Plugins/Experimental/ProceduralVegetationEditor/Source/ProceduralVegetation/Private/Helpers/PVUtilities.cpp
@@ -17,6 +17,7 @@
 #include "Facades/PVPointFacade.h"
 #include "Implementations/PVFoliage.h"
 #include "Misc/FileHelper.h"
+#include "Nodes/PVPresetLoaderSettings.h"
 #include "Rendering/SkeletalMeshRenderData.h"
  
 namespace PV::Utilities
@@ -29,6 +30,7 @@ namespace PV::Utilities
 		FConsoleVariableDelegate::CreateLambda([](IConsoleVariable* ConsoleVariable)
 			{
 				UPVImporterSettings::StaticClass()->GetDefaultObject<UPVImporterSettings>()->bExposeToLibrary = ConsoleVariable->GetBool();
+				UPVPresetLoaderSettings::StaticClass()->GetDefaultObject<UPVPresetLoaderSettings>()->bExposeToLibrary = ConsoleVariable->GetBool();
 				UProceduralVegetationPreset::ShowHideInternalProperties(ConsoleVariable->GetBool());
 			})
 	);
diff --git a/Engine/Plugins/Experimental/ProceduralVegetationEditor/Source/ProceduralVegetation/Private/Nodes/PVPresetLoaderSettings.cpp b/Engine/Plugins/Experimental/ProceduralVegetationEditor/Source/ProceduralVegetation/Private/Nodes/PVPresetLoaderSettings.cpp
--- a/Engine/Plugins/Experimental/ProceduralVegetationEditor/Source/ProceduralVegetation/Private/Nodes/PVPresetLoaderSettings.cpp
+++ b/Engine/Plugins/Experimental/ProceduralVegetationEditor/Source/ProceduralVegetation/Private/Nodes/PVPresetLoaderSettings.cpp
@@ -38,6 +38,16 @@ void UPVPresetLoaderSettings::PostEditChangeProperty(FPropertyChangedEvent& Prop
 }
 #endif
 
+UPVPresetLoaderSettings::UPVPresetLoaderSettings()
+{
+	#if WITH_EDITOR
+    	if (HasAnyFlags(RF_ClassDefaultObject))
+    	{
+    		bExposeToLibrary = PV::Utilities::DebugModeEnabled();
+    	}
+    #endif
+}
+
 void UPVPresetLoaderSettings::PostLoad()
 {
 	Super::PostLoad();
diff --git a/Engine/Plugins/Experimental/ProceduralVegetationEditor/Source/ProceduralVegetation/Private/Nodes/PVPresetLoaderSettings.h b/Engine/Plugins/Experimental/ProceduralVegetationEditor/Source/ProceduralVegetation/Private/Nodes/PVPresetLoaderSettings.h
--- a/Engine/Plugins/Experimental/ProceduralVegetationEditor/Source/ProceduralVegetation/Private/Nodes/PVPresetLoaderSettings.h
+++ b/Engine/Plugins/Experimental/ProceduralVegetationEditor/Source/ProceduralVegetation/Private/Nodes/PVPresetLoaderSettings.h
@@ -14,6 +14,7 @@ class UPVPresetLoaderSettings : public UPVBaseSettings
 	GENERATED_BODY()
 
 public:
+	UPVPresetLoaderSettings();
 	virtual void PostLoad() override;
 	
 #if WITH_EDITOR
