[PVE] FoliageDistributor refactor
- Removed mutable usage
- Remove getting settings from incoming pin.
- Removed override distribution. So now It will always distribute the foliage
- Keeping OverridePhyllotaxy and OverrideAxilAngles for now.

#Synced-CL 51086708

diff --git a/Engine/Plugins/Experimental/ProceduralVegetationEditor/Source/ProceduralVegetation/Private/Nodes/PVFoliageDistributorSettings.cpp b/Engine/Plugins/Experimental/ProceduralVegetationEditor/Source/ProceduralVegetation/Private/Nodes/PVFoliageDistributorSettings.cpp
--- a/Engine/Plugins/Experimental/ProceduralVegetationEditor/Source/ProceduralVegetation/Private/Nodes/PVFoliageDistributorSettings.cpp
+++ b/Engine/Plugins/Experimental/ProceduralVegetationEditor/Source/ProceduralVegetation/Private/Nodes/PVFoliageDistributorSettings.cpp
@@ -31,24 +31,6 @@ FText UPVFoliageDistributorSettings::GetNodeTooltipText() const
 		"\n\nPress Ctrl + L to lock/unlock node output"
 	); 
 }
- 
-void UPVFoliageDistributorSettings::PostEditChangeProperty(FPropertyChangedEvent& PropertyChangedEvent)
-{
-	Super::PostEditChangeProperty(PropertyChangedEvent);
- 
-	if (PropertyChangedEvent.Property
-		&& PropertyChangedEvent.Property->GetFName() == GET_MEMBER_NAME_CHECKED(UPVFoliageDistributorSettings, OverrideDistribution))
-	{
-		if (OverrideDistribution)
-		{
-			Conditions = CustomConditions;
-		}
-		else
-		{
-			Conditions = DefaultConditions;
-		}	
-	}
-}
 #endif
  
 UPVFoliageDistributorSettings::UPVFoliageDistributorSettings()
@@ -77,7 +59,7 @@ UPVFoliageDistributorSettings::UPVFoliageDistributorSettings()
 		RichCurve->SetKeyInterpMode(Handle5, RCIM_Linear);
 	}
 }
-
+ 
 TArray<FPCGPinProperties> UPVFoliageDistributorSettings::InputPinProperties() const
 {
 	TArray<FPCGPinProperties> Properties = Super::InputPinProperties();
@@ -88,7 +70,7 @@ TArray<FPCGPinProperties> UPVFoliageDistributorSettings::InputPinProperties() co
 	
 	return Properties;
 }
-
+ 
 FPCGDataTypeIdentifier UPVFoliageDistributorSettings::GetInputPinTypeIdentifier() const
 {
 	return FPCGDataTypeIdentifier{ FPVDataTypeInfoMesh::AsId() };
@@ -113,9 +95,6 @@ bool FPVFoliageDistributorElement::ExecuteInternal(FPCGContext* InContext) const
 	const UPVFoliageDistributorSettings* Settings = InContext->GetInputSettings<UPVFoliageDistributorSettings>();
 	check(Settings);
  
-	UPVFoliageDistributorSettings* MutableSettings = const_cast<UPVFoliageDistributorSettings*>(Settings);
-	check(MutableSettings);
-
 	FManagedArrayCollection FoliageCollection;
 	
 	for (const FPCGTaggedData& Input : InContext->InputData.GetInputsByPin(PV::Pins::FoliageDistributorFoliageInputLabel))
@@ -136,7 +115,7 @@ bool FPVFoliageDistributorElement::ExecuteInternal(FPCGContext* InContext) const
  
 			const FDistributionSettings DistributionSettings = {
 				.EthyleneThreshold = Settings->EthyleneThreshold,
-				.OverrideDistribution = Settings->OverrideDistribution,
+				.OverrideDistribution = true,
 				.InstanceSpacing = Settings->InstanceSpacing,
 				.InstanceSpacingRamp = &Settings->InstanceSpacingRamp,
 				.InstanceSpacingRampEffect = Settings->InstanceSpacingRampEffect,
@@ -174,45 +153,6 @@ bool FPVFoliageDistributorElement::ExecuteInternal(FPCGContext* InContext) const
 				.RandomSeed = Settings->RandomSeed
 			};
  
-			if (!Settings->OverrideDistribution)
-			{
-				PV::Facades::FFoliageConditionFacade FoliageConditionFacade(Collection);
- 
-				auto DefaultConditionsData = FoliageConditionFacade.GetData();
- 
-				for (auto ConditonInfo : DefaultConditionsData)
-				{
-					auto ConditionType = ConditonInfo.GetType();
- 
-					if (ConditionType != EFoliageDistributionCondition::None)
-					{
-						
-						if (auto ConditionSettings = MutableSettings->DefaultConditions.GetSettings(ConditionType))
-						{
-							ConditionSettings->Weight = ConditonInfo.Weight;
-							ConditionSettings->Offset = ConditonInfo.Offset;
- 
-							MutableSettings->DefaultConditions.SetConditionState(ConditionType, true);	
-						}
-					}
-				}
- 
-				if (!MutableSettings->bCustomConditionSet)
-				{
-					MutableSettings->CustomConditions = MutableSettings->DefaultConditions;
-					MutableSettings->bCustomConditionSet = true;
-				}
-			}
- 
-			if (!MutableSettings->OverrideDistribution)
-			{
-				MutableSettings->Conditions = MutableSettings->DefaultConditions;
-			}
-			else
-			{
-				MutableSettings->CustomConditions = MutableSettings->Conditions;
-			}	
- 
 			FPVFoliage::DistributeFoliage(Collection, Collection, FoliageCollection, DistributionSettings, ScaleSettings, VectorSettings,
 												 PhyllotaxySettings, MiscSettings, Settings->Conditions);
  
diff --git a/Engine/Plugins/Experimental/ProceduralVegetationEditor/Source/ProceduralVegetation/Private/Nodes/PVFoliageDistributorSettings.h b/Engine/Plugins/Experimental/ProceduralVegetationEditor/Source/ProceduralVegetation/Private/Nodes/PVFoliageDistributorSettings.h
--- a/Engine/Plugins/Experimental/ProceduralVegetationEditor/Source/ProceduralVegetation/Private/Nodes/PVFoliageDistributorSettings.h
+++ b/Engine/Plugins/Experimental/ProceduralVegetationEditor/Source/ProceduralVegetation/Private/Nodes/PVFoliageDistributorSettings.h
@@ -1,137 +1,124 @@
 ﻿// Copyright Epic Games, Inc. All Rights Reserved.
 #pragma once
-
+ 
 #include "CoreMinimal.h"
 #include "Utils/PVFloatRamp.h"
-
+ 
 #include "DataTypes/PVData.h"
-
+ 
 #include "Implementations/PVFoliage.h"
-
+ 
 #include "Nodes/PVBaseSettings.h"
-
+ 
 #include "PVFoliageDistributorSettings.generated.h"
-
+ 
 UCLASS(BlueprintType, ClassGroup = (Procedural))
 class UPVFoliageDistributorSettings : public UPVBaseSettings
 {
 	GENERATED_BODY()
-
+ 
 	friend class FPVFoliageDistributorElement;
 	
 public:
 	UPVFoliageDistributorSettings();
-
+ 
 	#if WITH_EDITOR
 	virtual FName GetDefaultNodeName() const override { return FName(TEXT("ProceduralVegetationFoliageDistributor")); }
 	virtual FText GetDefaultNodeTitle() const override;
 	virtual FText GetNodeTooltipText() const override;
 	virtual FLinearColor GetNodeTitleColor() const override { return FLinearColor::Blue; }
-
+ 
 	virtual TArray<EPVRenderType> GetDefaultRenderType() const override { return TArray{ EPVRenderType::Mesh, EPVRenderType::Foliage }; }
-	virtual void PostEditChangeProperty(FPropertyChangedEvent& PropertyChangedEvent) override;
 	virtual TArray<EPVRenderType> GetSupportedRenderTypes() const override { return TArray{ EPVRenderType::Mesh, EPVRenderType::Foliage, EPVRenderType::FoliageAttachments }; }
 	#endif
-
+ 
 protected:
 	virtual TArray<FPCGPinProperties> InputPinProperties() const override;
 	virtual FPCGDataTypeIdentifier GetInputPinTypeIdentifier() const override;
 	virtual FPCGDataTypeIdentifier GetOutputPinTypeIdentifier() const override;
-
+ 
 	virtual FPCGElementPtr CreateElement() const override;
-
+ 
 private:
-
+ 
 	UPROPERTY(EditAnywhere, Category="Distribution Settings", meta=(PCG_Overridable, ClampMin=0.0f, ClampMax=1.0f, UIMin=0.0f, UIMax=1.0f, Tooltip="Threshold controlling foliage release based on an ethylene-like signal.\n\nHigher Threshold → more buds/branches retained, making the plant denser and bushier.\nLower Threshold → fewer buds/branches retained, resulting in a sparser, pruned structure."))
 	float EthyleneThreshold = 0.85f;
-
-	UPROPERTY(EditAnywhere, Category="Distribution Settings", meta=(Tooltip="Override default placement with custom distribution rules.\n\nOverides the distribution settings loaded from the procedural vegetation preset."))
-	bool OverrideDistribution = true;
-
-	UPROPERTY(EditAnywhere, Category="Distribution Settings", meta=(PCG_Overridable, EditCondition="OverrideDistribution", ClampMin=0.0f, ClampMax=1.0f, UIMin=0.0f, UIMax=1.0f, Tooltip="Minimum distance between placed instances.\n\nHigher spacing yields a more open distribution; lower spacing allows denser placement."))
+ 
+	UPROPERTY(EditAnywhere, Category="Distribution Settings", meta=(PCG_Overridable, ClampMin=0.0f, ClampMax=1.0f, UIMin=0.0f, UIMax=1.0f, Tooltip="Minimum distance between placed instances.\n\nHigher spacing yields a more open distribution; lower spacing allows denser placement."))
 	float InstanceSpacing = 0.3f;
-
-	UPROPERTY(EditAnywhere, Category="Distribution Settings", meta=(PCG_Overridable, EditCondition="OverrideDistribution", XAxisMin=0.0f, XAxisMax=1.0f, YAxisMin=0.0f, YAxisMax=1.0f, Tooltip="Ramp that varies instance spacing across the height of the plant\n\nUseful to reduce spacing in some areas based on the height of the plant."))
+ 
+	UPROPERTY(EditAnywhere, Category="Distribution Settings", meta=(PCG_Overridable, XAxisMin=0.0f, XAxisMax=1.0f, YAxisMin=0.0f, YAxisMax=1.0f, Tooltip="Ramp that varies instance spacing across the height of the plant\n\nUseful to reduce spacing in some areas based on the height of the plant."))
 	FPVFloatRamp InstanceSpacingRamp;
-
-	UPROPERTY(EditAnywhere, Category="Distribution Settings", meta=(PCG_Overridable, EditCondition="OverrideDistribution", ClampMin=0.0f, ClampMax=1.0f, UIMin=0.0f, UIMax=1.0f, Tooltip="Strength of the spacing ramp’s influence.\n\nBlends the ramp with the base spacing. Lower values favor the base spacing (subtle effect); higher values apply the ramp more strongly and can amplify spacing differences."))
+ 
+	UPROPERTY(EditAnywhere, Category="Distribution Settings", meta=(PCG_Overridable, ClampMin=0.0f, ClampMax=1.0f, UIMin=0.0f, UIMax=1.0f, Tooltip="Strength of the spacing ramp’s influence.\n\nBlends the ramp with the base spacing. Lower values favor the base spacing (subtle effect); higher values apply the ramp more strongly and can amplify spacing differences."))
 	float InstanceSpacingRampEffect = 0.0f;
-
-	UPROPERTY(EditAnywhere, Category="Distribution Settings", meta=(PCG_Overridable, EditCondition="OverrideDistribution", ClampMin=-1, ClampMax=1000, Tooltip="Caps the number of instances per branch\n\nlimit on how many foliage instances can be placed per branch to avoid overcrowding."))
+ 
+	UPROPERTY(EditAnywhere, Category="Distribution Settings", meta=(PCG_Overridable, ClampMin=-1, ClampMax=1000, Tooltip="Caps the number of instances per branch\n\nlimit on how many foliage instances can be placed per branch to avoid overcrowding."))
 	int32 MaxPerBranch = -1;
-
-	UPROPERTY(EditAnywhere, Category="Scale Settings", meta=(PCG_Overridable, EditCondition="OverrideDistribution", ClampMin=0.0f, UIMin=0.0f, UIMax=5.0f, Tooltip="Default size multiplier for placed foliage\n\nSets the base uniform scale applied to each instance before randomness and ramps are considered. Increase to enlarge all instances proportionally; decrease for smaller foliage. Final size still respects Min/Max limits."))
+ 
+	UPROPERTY(EditAnywhere, Category="Scale Settings", meta=(PCG_Overridable, ClampMin=0.0f, UIMin=0.0f, UIMax=5.0f, Tooltip="Default size multiplier for placed foliage\n\nSets the base uniform scale applied to each instance before randomness and ramps are considered. Increase to enlarge all instances proportionally; decrease for smaller foliage. Final size still respects Min/Max limits."))
 	float BaseScale = 1.0f;
-
-	UPROPERTY(EditAnywhere, Category="Scale Settings", meta=(PCG_Overridable, EditCondition="OverrideDistribution", ClampMin=0.0f, ClampMax=1.0f, UIMin=0.0f, UIMax=1.0f, Tooltip="Impact of branch size on foliage size\n\nThe size of the foliage being scattered will be chosen in relation to the physical size of the branch it is distributed on. The largest branches will be distributed the largest foliage pieces."))
+ 
+	UPROPERTY(EditAnywhere, Category="Scale Settings", meta=(PCG_Overridable, ClampMin=0.0f, ClampMax=1.0f, UIMin=0.0f, UIMax=1.0f, Tooltip="Impact of branch size on foliage size\n\nThe size of the foliage being scattered will be chosen in relation to the physical size of the branch it is distributed on. The largest branches will be distributed the largest foliage pieces."))
 	float BranchScaleImpact = 0.0f;
-
-	UPROPERTY(EditAnywhere, Category="Scale Settings", meta=(PCG_Overridable, EditCondition="OverrideDistribution", ClampMin=0.0f, UIMin=0.0f, UIMax=5.0f, Tooltip="Minimum allowed instance size.\n\nClamps the final computed scale so it never goes below this value. Use to prevent tiny, hard-to-see instances."))
+ 
+	UPROPERTY(EditAnywhere, Category="Scale Settings", meta=(PCG_Overridable, ClampMin=0.0f, UIMin=0.0f, UIMax=5.0f, Tooltip="Minimum allowed instance size.\n\nClamps the final computed scale so it never goes below this value. Use to prevent tiny, hard-to-see instances."))
 	float MinScale = 0.5f;
-
-	UPROPERTY(EditAnywhere, Category="Scale Settings", meta=(PCG_Overridable, EditCondition="OverrideDistribution", ClampMin=0.0f, UIMin=0.0f, UIMax=5.0f, Tooltip="Maximum allowed instance size.\n\nClamps the final computed scale so it never exceeds this value. Use to cap unusually large instances for visual consistency."))
+ 
+	UPROPERTY(EditAnywhere, Category="Scale Settings", meta=(PCG_Overridable, ClampMin=0.0f, UIMin=0.0f, UIMax=5.0f, Tooltip="Maximum allowed instance size.\n\nClamps the final computed scale so it never exceeds this value. Use to cap unusually large instances for visual consistency."))
 	float MaxScale = 1.0f;
-
-	UPROPERTY(EditAnywhere, Category="Scale Settings", meta=(PCG_Overridable, EditCondition="OverrideDistribution", ClampMin=0.0f, UIMin=0.0f, UIMax=5.0f, Tooltip="Lower bound for per-instance size variation.\n\nSets the minimum of the random scale range used to vary instance sizes. Each instance samples within Random Scale Min–Max, then combines with Base Scale and any ramps (respecting Min/Max Scale)."))
+ 
+	UPROPERTY(EditAnywhere, Category="Scale Settings", meta=(PCG_Overridable, ClampMin=0.0f, UIMin=0.0f, UIMax=5.0f, Tooltip="Lower bound for per-instance size variation.\n\nSets the minimum of the random scale range used to vary instance sizes. Each instance samples within Random Scale Min–Max, then combines with Base Scale and any ramps (respecting Min/Max Scale)."))
 	float RandomScaleMin = 1.0f;
-
-	UPROPERTY(EditAnywhere, Category="Scale Settings", meta=(PCG_Overridable, EditCondition="OverrideDistribution", ClampMin=1.0f, UIMin=1.0f, UIMax=5.0f, Tooltip="Upper bound for per-instance size variation.\n\nSets the maximum of the random scale range used to vary instance sizes. Each instance samples within Random Scale Min–Max, then combines with Base Scale and any ramps (respecting Min/Max Scale)."))
+ 
+	UPROPERTY(EditAnywhere, Category="Scale Settings", meta=(PCG_Overridable, ClampMin=1.0f, UIMin=1.0f, UIMax=5.0f, Tooltip="Upper bound for per-instance size variation.\n\nSets the maximum of the random scale range used to vary instance sizes. Each instance samples within Random Scale Min–Max, then combines with Base Scale and any ramps (respecting Min/Max Scale)."))
 	float RandomScaleMax = 1.0f;
-
-	UPROPERTY(EditAnywhere, Category="Scale Settings", meta=(EditCondition="OverrideDistribution", XAxisMin=0.0f, XAxisMax=1.0f, YAxisMin=0.0f, YAxisMax=1.0f, Tooltip="Curve that varies scale across the height of the plant\n\nUse to taper foliage toward tips or emphasize size in targeted regions."))
+ 
+	UPROPERTY(EditAnywhere, Category="Scale Settings", meta=(XAxisMin=0.0f, XAxisMax=1.0f, YAxisMin=0.0f, YAxisMax=1.0f, Tooltip="Curve that varies scale across the height of the plant\n\nUse to taper foliage toward tips or emphasize size in targeted regions."))
 	FPVFloatRamp ScaleRamp;
-
-	UPROPERTY(EditAnywhere, Category="Vector Settings", meta=(PCG_Overridable, EditCondition="OverrideDistribution", Tooltip="Use a custom axil angle instead of the default.\n\nEnables a user-defined orientation. When on, the system uses AxilAngle and the AxilAngleRamp settings (and Effect) to set the tilt of foliage relative to its parent branch direction, overriding any automatic/species defaults."))
+ 
+	UPROPERTY(EditAnywhere, Category="Vector Settings", meta=(PCG_Overridable, Tooltip="Use a custom axil angle instead of the default.\n\nEnables a user-defined orientation. When on, the system uses AxilAngle and the AxilAngleRamp settings (and Effect) to set the tilt of foliage relative to its parent branch direction, overriding any automatic/species defaults."))
 	bool OverrideAxilAngle = true;
-
-	UPROPERTY(EditAnywhere, Category="Vector Settings", meta=(PCG_Overridable, EditCondition="OverrideDistribution", ClampMin=0.0f, ClampMax=90.0f, UIMin=0.0f, UIMax=90.0f, Tooltip="Base tilt from the parent axis\n\nSets the base axil angle in degrees—the tilt of the instance relative to the parent branch’s direction. Acts as the lower bound when used with a ramp; higher values tilt foliage further away from the branch axis."))
+ 
+	UPROPERTY(EditAnywhere, Category="Vector Settings", meta=(PCG_Overridable, EditCondition="OverrideAxilAngle", ClampMin=0.0f, ClampMax=90.0f, UIMin=0.0f, UIMax=90.0f, Tooltip="Base tilt from the parent axis\n\nSets the base axil angle in degrees—the tilt of the instance relative to the parent branch’s direction. Acts as the lower bound when used with a ramp; higher values tilt foliage further away from the branch axis."))
 	float AxilAngle = 45.0f;
-
-	UPROPERTY(EditAnywhere, Category="Vector Settings", meta=(EditCondition="OverrideDistribution", XAxisMin=0.0f, XAxisMax=1.0f, YAxisMin=0.0f, YAxisMax=1.0f, Tooltip="Ramp that varies the axil angle across the height of the plant\n\nIts output is remapped to an angle range defined by AxilAngle and AxilAngleRampUpperValue, modulated by AxilAngleRampEffect."))
+ 
+	UPROPERTY(EditAnywhere, Category="Vector Settings", meta=(EditCondition="OverrideAxilAngle", XAxisMin=0.0f, XAxisMax=1.0f, YAxisMin=0.0f, YAxisMax=1.0f, Tooltip="Ramp that varies the axil angle across the height of the plant\n\nIts output is remapped to an angle range defined by AxilAngle and AxilAngleRampUpperValue, modulated by AxilAngleRampEffect."))
 	FPVFloatRamp AxilAngleRamp;
-
-	UPROPERTY(EditAnywhere, Category="Vector Settings", meta=(PCG_Overridable, EditCondition="OverrideDistribution", ClampMin=0.0f, ClampMax=90.0f, UIMin=0.0f, UIMax=90.0f, Tooltip="Max angle used when the ramp\n\nProvides detailed control over this parameter to balance visual quality and performance. Valid values: 0.0–90.0."))
+ 
+	UPROPERTY(EditAnywhere, Category="Vector Settings", meta=(PCG_Overridable, EditCondition="OverrideAxilAngle", ClampMin=0.0f, ClampMax=90.0f, UIMin=0.0f, UIMax=90.0f, Tooltip="Max angle used when the ramp\n\nProvides detailed control over this parameter to balance visual quality and performance. Valid values: 0.0–90.0."))
 	float AxilAngleRampUpperValue = 45.0f;
-
-	UPROPERTY(EditAnywhere, Category="Vector Settings", meta=(PCG_Overridable, EditCondition="OverrideDistribution", ClampMin=0.0f, ClampMax=1.0f, UIMin=0.0f, UIMax=1.0f, Tooltip="Adjusts this parameter for finer control.\n\nDefines the upper axil angle reached when AxilAngleRamp evaluates to 1. Together with AxilAngle, it forms the min–max range for ramp-driven orientation."))
+ 
+	UPROPERTY(EditAnywhere, Category="Vector Settings", meta=(PCG_Overridable, EditCondition="OverrideAxilAngle", ClampMin=0.0f, ClampMax=1.0f, UIMin=0.0f, UIMax=1.0f, Tooltip="Adjusts this parameter for finer control.\n\nDefines the upper axil angle reached when AxilAngleRamp evaluates to 1. Together with AxilAngle, it forms the min–max range for ramp-driven orientation."))
 	float AxilAngleRampEffect = 0.0f;
-
-	UPROPERTY(EditAnywhere, Category="Phyllotaxy Settings", meta=(EditCondition="OverrideDistribution", Tooltip="Use custom phyllotaxy settings instead of defaults.\n\nOveride the phyllotaxy settings loaded from the procedural vegetation preset."))
+ 
+	UPROPERTY(VisibleAnywhere, Category="Phyllotaxy Settings", meta=(Tooltip="Use custom phyllotaxy settings instead of defaults.\n\nOveride the phyllotaxy settings loaded from the procedural vegetation preset."))
 	bool OverridePhyllotaxy = true;
-
-	UPROPERTY(EditAnywhere, Category="Phyllotaxy Settings", meta=(PCG_Overridable, EditCondition="OverrideDistribution", Tooltip="Pattern for arranging buds/leaves around the stem.\n\nSelects the phyllotactic pattern used for placement (e.g., Alternate, Opposite, Whorled, Spiral). This defines the sequence and angular spacing of buds/leaves around the stem and along its length."))
+ 
+	UPROPERTY(EditAnywhere, Category="Phyllotaxy Settings", meta=(PCG_Overridable, EditCondition="OverridePhyllotaxy", Tooltip="Pattern for arranging buds/leaves around the stem.\n\nSelects the phyllotactic pattern used for placement (e.g., Alternate, Opposite, Whorled, Spiral). This defines the sequence and angular spacing of buds/leaves around the stem and along its length."))
 	EPhyllotaxyType PhyllotaxyType = EPhyllotaxyType::Alternate;
-
-	UPROPERTY(EditAnywhere, Category="Phyllotaxy Settings", meta=(PCG_Overridable, EditCondition="OverrideDistribution && PhyllotaxyType == EPhyllotaxyType::Spiral", EditConditionHides))
+ 
+	UPROPERTY(EditAnywhere, Category="Phyllotaxy Settings", meta=(PCG_Overridable, EditCondition="OverridePhyllotaxy && PhyllotaxyType == EPhyllotaxyType::Spiral", EditConditionHides))
 	EPhyllotaxyFormation PhyllotaxyFormation = EPhyllotaxyFormation::Distichous;
-
-	UPROPERTY(EditAnywhere, Category="Phyllotaxy Settings", meta=(PCG_Overridable, EditCondition="OverrideDistribution", ClampMin=1, ClampMax=10, UIMin=1, UIMax=10, Tooltip="Minimum buds/leaves per node.\n\nSets the lower bound on the number of buds/leaves generated at each node. Guarantees at least this many instances (subject to spacing and other constraints). Should be ≤ MaximumNodeBuds. Applied when Whorled phyllotaxy is selected."))
+ 
+	UPROPERTY(EditAnywhere, Category="Phyllotaxy Settings", meta=(PCG_Overridable, EditCondition="OverridePhyllotaxy", ClampMin=1, ClampMax=10, UIMin=1, UIMax=10, Tooltip="Minimum buds/leaves per node.\n\nSets the lower bound on the number of buds/leaves generated at each node. Guarantees at least this many instances (subject to spacing and other constraints). Should be ≤ MaximumNodeBuds. Applied when Whorled phyllotaxy is selected."))
 	int32 MinimumNodeBuds = 2;
-
-	UPROPERTY(EditAnywhere, Category="Phyllotaxy Settings", meta=(PCG_Overridable, EditCondition="OverrideDistribution", ClampMin=1, ClampMax=10, UIMin=1, UIMax=10, Tooltip="Maximum buds/leaves per node.\n\nCaps the number of buds/leaves spawned. Use with MinimumNodeBuds to bound density and keep arrangements predictable. Applied when Whorled phyllotaxy is selected."))
+ 
+	UPROPERTY(EditAnywhere, Category="Phyllotaxy Settings", meta=(PCG_Overridable, EditCondition="OverridePhyllotaxy", ClampMin=1, ClampMax=10, UIMin=1, UIMax=10, Tooltip="Maximum buds/leaves per node.\n\nCaps the number of buds/leaves spawned. Use with MinimumNodeBuds to bound density and keep arrangements predictable. Applied when Whorled phyllotaxy is selected."))
 	int32 MaximumNodeBuds = 3;
-
-	UPROPERTY(EditAnywhere, Category="Phyllotaxy Settings", meta=(PCG_Overridable, EditCondition="OverrideDistribution", ClampMin=0.0f, ClampMax=360.0f, UIMin=0.0f, UIMax=360.0f, Tooltip="Extra divergence angle (°) added to the phyllotaxy pattern\n\nAdds an angular offset to the base phyllotactic divergence."))
+ 
+	UPROPERTY(EditAnywhere, Category="Phyllotaxy Settings", meta=(PCG_Overridable, EditCondition="OverridePhyllotaxy", ClampMin=0.0f, ClampMax=360.0f, UIMin=0.0f, UIMax=360.0f, Tooltip="Extra divergence angle (°) added to the phyllotaxy pattern\n\nAdds an angular offset to the base phyllotactic divergence."))
 	float PhyllotaxyAdditionalAngle = 0.0f;
-
-	UPROPERTY(EditAnywhere, Category="Misc Settings", meta=(PCG_Overridable, EditCondition="OverrideDistribution"))
+ 
+	UPROPERTY(EditAnywhere, Category="Misc Settings", meta=(PCG_Overridable))
 	int32 RandomSeed = 123456;
-
-	UPROPERTY(EditAnywhere, Category="ConditionSettings", meta = (PCG_Overridable, ShowOnlyInnerProperties, EditCondition="OverrideDistribution"))
+ 
+	UPROPERTY(EditAnywhere, Category="ConditionSettings", meta = (PCG_Overridable, ShowOnlyInnerProperties))
 	FPVFoliageDistributionConditions Conditions;
-	
-	UPROPERTY()
-	FPVFoliageDistributionConditions DefaultConditions;
-	
-	UPROPERTY()
-	FPVFoliageDistributionConditions CustomConditions;
-
-	UPROPERTY()
-	bool bCustomConditionSet = false;
 };
-
+ 
 class FPVFoliageDistributorElement : public FPVBaseElement
 {
 protected:
 	virtual bool ExecuteInternal(FPCGContext* InContext) const override;
-};
+};
\ No newline at end of file
