[PVE] Added Grower data asset, which will only contain growth info from Apix Houdini
- Created new node, GrowerDataLoader, which will output growth data similar to the grower node.
- JSON To ManagedArrayCollection was already written in PVGrowthDataJSONImporter, so I reused it.
- This is the first implementation, and its data doesnt fully match with our Grower version.

#Synced-CL 51086708

diff --git a/Engine/Plugins/Experimental/ProceduralVegetationEditor/Source/ProceduralVegetation/Private/DataAssets/ProceduralVegetationGrowthDataAsset.cpp b/Engine/Plugins/Experimental/ProceduralVegetationEditor/Source/ProceduralVegetation/Private/DataAssets/ProceduralVegetationGrowthDataAsset.cpp
new file mode 100644
--- /dev/null
+++ b/Engine/Plugins/Experimental/ProceduralVegetationEditor/Source/ProceduralVegetation/Private/DataAssets/ProceduralVegetationGrowthDataAsset.cpp
@@ -0,0 +1,57 @@
+// Copyright Epic Games, Inc. All Rights Reserved.
+
+#include "ProceduralVegetationGrowthDataAsset.h"
+#include "ProceduralVegetationModule.h"
+#include "GeometryCollection/ManagedArrayCollection.h"
+#include "HAL/FileManager.h"
+#include "Helpers/PVJSONHelper.h"
+#include "Misc/Paths.h"
+#include "UObject/Package.h"
+
+
+void UProceduralVegetationGrowthDataAsset::UpdateDataAsset()
+{
+	UE_LOGF(LogProceduralVegetation, Log, "UpdateDataAsset Clicked");
+	
+	if (JsonDirectoryPath.Path.IsEmpty())
+	{
+		UE_LOG(LogProceduralVegetation, Warning, TEXT("Please specify a JSON path"));
+		return;
+	}
+	
+	FString OutError;
+	IFileManager& FileManager = IFileManager::Get();
+	
+	TArray<FString> FileNames;
+	FileManager.FindFiles(FileNames, *JsonDirectoryPath.Path, TEXT("*.json"));
+ 
+	Variants.Empty();
+	GrowthVariations.Empty();
+	
+	FString JSONExtension = ".json";
+
+	for (const auto& FileName : FileNames)
+	{
+		FString FullPath = JsonDirectoryPath.Path / FileName;
+
+		FManagedArrayCollection Collection;
+		FString MetaJsonError;
+
+		if (PV::JSON::LoadGrowthDataJsonToCollection(Collection, FullPath, OutError))
+		{
+			FString Name = FPaths::GetBaseFilename(FullPath);
+			UE_LOGF(LogProceduralVegetation, Log, "Variant %ls loaded from file", *Name);
+
+			Variants.Add(Name, Collection);
+
+			GrowthVariations.Add({FName(Name)});
+		}
+		
+		if (!OutError.IsEmpty())
+		{
+			UE_LOGF(LogProceduralVegetation, Warning, "Invalid growth data json File %ls %ls",*FullPath, *OutError);
+		}
+	}
+	
+	GetPackage()->SetDirtyFlag(true);
+}
diff --git a/Engine/Plugins/Experimental/ProceduralVegetationEditor/Source/ProceduralVegetation/Private/DataAssets/ProceduralVegetationGrowthDataAsset.h b/Engine/Plugins/Experimental/ProceduralVegetationEditor/Source/ProceduralVegetation/Private/DataAssets/ProceduralVegetationGrowthDataAsset.h
new file mode 100644
--- /dev/null
+++ b/Engine/Plugins/Experimental/ProceduralVegetationEditor/Source/ProceduralVegetation/Private/DataAssets/ProceduralVegetationGrowthDataAsset.h
@@ -0,0 +1,36 @@
+// Copyright Epic Games, Inc. All Rights Reserved.
+
+#pragma once
+
+#include "Engine/DataAsset.h"
+#include "GeometryCollection/ManagedArrayCollection.h"
+#include "ProceduralVegetationGrowthDataAsset.generated.h"
+
+USTRUCT()
+struct FPVGrowthVariationInfo
+{
+	GENERATED_BODY()
+
+	UPROPERTY(VisibleAnywhere, Category = "Variations")
+	FName Name;
+};
+
+UCLASS(BlueprintType)
+class PROCEDURALVEGETATION_API UProceduralVegetationGrowthDataAsset : public UDataAsset
+{
+	GENERATED_BODY()
+	
+public:
+	UPROPERTY()
+	TMap<FString, FManagedArrayCollection> Variants;
+
+	UPROPERTY(VisibleAnywhere, Category = "Growth Data")
+	TArray<FPVGrowthVariationInfo> GrowthVariations;
+
+private:
+	UPROPERTY(EditAnywhere, Category="Internal", meta = (DevelopmentOnly))
+	FDirectoryPath JsonDirectoryPath;
+
+	UFUNCTION(CallInEditor, Category="Internal", meta = (DevelopmentOnly))
+	void UpdateDataAsset();
+};
diff --git a/Engine/Plugins/Experimental/ProceduralVegetationEditor/Source/ProceduralVegetation/Private/Helpers/PVJSONHelper.cpp b/Engine/Plugins/Experimental/ProceduralVegetationEditor/Source/ProceduralVegetation/Private/Helpers/PVJSONHelper.cpp
--- a/Engine/Plugins/Experimental/ProceduralVegetationEditor/Source/ProceduralVegetation/Private/Helpers/PVJSONHelper.cpp
+++ b/Engine/Plugins/Experimental/ProceduralVegetationEditor/Source/ProceduralVegetation/Private/Helpers/PVJSONHelper.cpp
@@ -1,6 +1,9 @@
 ﻿// Copyright Epic Games, Inc. All Rights Reserved.
  
 #include "PVJSONHelper.h"
+
+#include "ProceduralVegetationModule.h"
+#include "PVAttributesHelper.h"
 #include "Dom/JsonObject.h"
 #include "Dom/JsonValue.h"
 #include "Facades/PVAttributesNames.h"
@@ -498,6 +501,59 @@ namespace PV::JSON
 		return true;
 	}
 
+	void AddGrowthMissingData(FManagedArrayCollection& Collection)
+	{
+		if (!PV::FPointLengthFromRootAttribute::HasAttribute(Collection)
+			|| !PV::FPointLengthFromSeedAttribute::HasAttribute(Collection))
+		{
+			PV::FPointLengthFromRootAttribute::AddAttribute(Collection);
+			PV::FPointLengthFromSeedAttribute::AddAttribute(Collection);
+			PV::AttributesHelper::ComputeLengthFromRoot(Collection);
+		}
+
+		if (!PV::FPointGroundGradientAttribute::HasAttribute(Collection))
+		{
+			PV::FPointGroundGradientAttribute::AddAttribute(Collection);
+			PV::AttributesHelper::ComputePointGroundGradient(Collection);
+		}
+
+		if (!PV::FPointScaleGradientAttribute::HasAttribute(Collection))
+		{
+			PV::FPointScaleGradientAttribute::AddAttribute(Collection);
+			PV::AttributesHelper::ComputePointScaleGradient(Collection);
+		}
+
+		if (!PV::FPointPlantGradientAttribute::HasAttribute(Collection) 
+			|| !PV::FPointBranchGradientAttribute::HasAttribute(Collection))
+		{
+			PV::FPointPlantGradientAttribute::AddAttribute(Collection);
+			PV::FPointBranchGradientAttribute::AddAttribute(Collection);
+			PV::AttributesHelper::ComputePointPlantAndBranchGradient(Collection);
+		}
+
+		if (!PV::FPointHullGradientAttribute::HasAttribute(Collection))
+		{
+			PV::FPointHullGradientAttribute::AddAttribute(Collection, 0);
+		}
+
+		if (!PV::FPointMainTrunkGradientAttribute::HasAttribute(Collection))
+		{
+			PV::FPointMainTrunkGradientAttribute::AddAttribute(Collection, 0);
+		}
+
+		if (!PV::FPointNjordPixelIndexAttribute::HasAttribute(Collection))
+		{
+			PV::FPointNjordPixelIndexAttribute::AddAttribute(Collection);
+			PV::AttributesHelper::ComputeNjordPixelIndex(Collection);
+		}
+
+		if (!PV::FBranchHierarchyNumberAttribute::HasAttribute(Collection))
+		{
+			PV::FBranchHierarchyNumberAttribute::AddAttribute(Collection);
+			PV::AttributesHelper::ComputeBranchHierarchyNumbers(Collection);
+		}
+	}
+
 	bool LoadGrowthDataJsonToCollection(FManagedArrayCollection& Collection, const FString& FilePath, FString& OutErrorMessage)
 	{
 		FString MegaPlantsData;
@@ -534,35 +590,18 @@ namespace PV::JSON
 
 		const TSharedPtr<FJsonObject>& Points = LoadedData->GetObjectField(TEXT("points"));
 		const TSharedPtr<FJsonObject>& PointAttributes = Points->GetObjectField(TEXT("attributes"));
-
 		const TArray<TSharedPtr<FJsonValue>>& PointPositions = Points->GetArrayField(TEXT("positions"));
 
-		const TArray<TSharedPtr<FJsonValue>>* PointScales = Invoke([&]()
-		{
-			const TArray<TSharedPtr<FJsonValue>>* OutPointScales = nullptr;
-			const TSharedPtr<FJsonObject>* PScaleObject = nullptr;
-			if (PointAttributes->TryGetObjectField(TEXT("pscale"), PScaleObject))
-			{
-				(*PScaleObject)->TryGetArrayField(TEXT("values"), OutPointScales);
-			}
-			return OutPointScales;
-		});
-
 		Collection.AddGroup(GroupNames::PointGroup);
 		Collection.AddElements(PointPositions.Num(), GroupNames::PointGroup);
 
 		TManagedArray<FVector3f>& PositionArray = Collection.AddAttribute<FVector3f>(AttributeNames::PointPosition, GroupNames::PointGroup);
-		TManagedArray<float>& ScaleArray = Collection.AddAttribute<float>("Scale", GroupNames::PointGroup);
 
 		for (int32 PointIndex = 0; PointIndex < PointPositions.Num(); ++PointIndex)
 		{
 			const TArray<TSharedPtr<FJsonValue>>& PointPosition = PointPositions[PointIndex]->AsArray();
-
 			const FVector3f Position = FVector3f(PointPosition[0]->AsNumber(), PointPosition[2]->AsNumber(), PointPosition[1]->AsNumber()) * 100.0f;
-			const float Scale = PointScales ? (*PointScales)[PointIndex]->AsNumber() * 100.0f : 0.f;
-
 			PositionArray[PointIndex] = Position;
-			ScaleArray[PointIndex] = Scale;
 		}
 
 		const TSharedPtr<FJsonObject>& Primitives = LoadedData->GetObjectField(TEXT("primitives"));
@@ -589,12 +628,26 @@ namespace PV::JSON
 
 		FillAttributes(Collection, GroupNames::PointGroup, PointAttributes);
 		FillAttributes(Collection, GroupNames::BranchGroup, PrimitiveAttributes);
-		FillDetailsAttributes(Collection, GroupNames::DetailsGroup, LoadedData->GetObjectField(TEXT("globalAttributes")));
 
+		TManagedArray<float>& ScaleArray = Collection.AddAttribute<float>("Scale", GroupNames::PointGroup);
+		FBudLateralMeristemAttributeConstView BudLateralMeristem = FBudLateralMeristemAttribute::GetAttribute(Collection);
+
+		for (int32 PointIndex = 0; PointIndex < PointPositions.Num(); ++PointIndex)
+		{
+			float Scale = BudLateralMeristem[PointIndex].LateralMeristem;
+			
+			Scale *= 100.0f;
+
+			ScaleArray[PointIndex] = Scale;
+		}
+		
+		AddGrowthMissingData(Collection);
+		
 		Facades::FMetaInfoFacade MetaInfoFacade(Collection);
 		MetaInfoFacade.CreateGuid(FilePath);
 
 		OutErrorMessage = "";
+		
 		return true;
 	}
 
diff --git a/Engine/Plugins/Experimental/ProceduralVegetationEditor/Source/ProceduralVegetation/Private/Nodes/Development/PVGrowthDataJsonImporterSettings.cpp b/Engine/Plugins/Experimental/ProceduralVegetationEditor/Source/ProceduralVegetation/Private/Nodes/Development/PVGrowthDataJsonImporterSettings.cpp
--- a/Engine/Plugins/Experimental/ProceduralVegetationEditor/Source/ProceduralVegetation/Private/Nodes/Development/PVGrowthDataJsonImporterSettings.cpp
+++ b/Engine/Plugins/Experimental/ProceduralVegetationEditor/Source/ProceduralVegetation/Private/Nodes/Development/PVGrowthDataJsonImporterSettings.cpp
@@ -73,61 +73,6 @@ bool FPVGrowthDataJsonImporterElement::ExecuteInternal(FPCGContext* InContext) c
 		return true;
 	}
 
-	if (!PV::FPointScaleAttribute::HasAttribute(Collection))
-	{
-		PV::FPointScaleAttribute::AddAttribute(Collection, 0);
-	}
-	
-	if (!PV::FPointLengthFromRootAttribute::HasAttribute(Collection)
-		|| !PV::FPointLengthFromSeedAttribute::HasAttribute(Collection))
-	{
-		PV::FPointLengthFromRootAttribute::AddAttribute(Collection);
-		PV::FPointLengthFromSeedAttribute::AddAttribute(Collection);
-		PV::AttributesHelper::ComputeLengthFromRoot(Collection);
-	}
-
-	if (!PV::FPointGroundGradientAttribute::HasAttribute(Collection))
-	{
-		PV::FPointGroundGradientAttribute::AddAttribute(Collection);
-		PV::AttributesHelper::ComputePointGroundGradient(Collection);
-	}
-
-	if (!PV::FPointScaleGradientAttribute::HasAttribute(Collection))
-	{
-		PV::FPointScaleGradientAttribute::AddAttribute(Collection);
-		PV::AttributesHelper::ComputePointScaleGradient(Collection);
-	}
-
-	if (!PV::FPointPlantGradientAttribute::HasAttribute(Collection) 
-		|| !PV::FPointBranchGradientAttribute::HasAttribute(Collection))
-	{
-		PV::FPointPlantGradientAttribute::AddAttribute(Collection);
-		PV::FPointBranchGradientAttribute::AddAttribute(Collection);
-		PV::AttributesHelper::ComputePointPlantAndBranchGradient(Collection);
-	}
-
-	if (!PV::FPointHullGradientAttribute::HasAttribute(Collection))
-	{
-		PV::FPointHullGradientAttribute::AddAttribute(Collection, 0);
-	}
-
-	if (!PV::FPointMainTrunkGradientAttribute::HasAttribute(Collection))
-	{
-		PV::FPointMainTrunkGradientAttribute::AddAttribute(Collection, 0);
-	}
-
-	if (!PV::FPointNjordPixelIndexAttribute::HasAttribute(Collection))
-	{
-		PV::FPointNjordPixelIndexAttribute::AddAttribute(Collection);
-		PV::AttributesHelper::ComputeNjordPixelIndex(Collection);
-	}
-
-	if (!PV::FBranchHierarchyNumberAttribute::HasAttribute(Collection))
-	{
-		PV::FBranchHierarchyNumberAttribute::AddAttribute(Collection);
-		PV::AttributesHelper::ComputeBranchHierarchyNumbers(Collection);
-	}
-
 	UPVGrowthData* OutGrowthData = FPCGContext::NewObject_AnyThread<UPVGrowthData>(InContext);
 	OutGrowthData->Initialize(MoveTemp(Collection));
 	InContext->OutputData.TaggedData.Emplace(OutGrowthData);
diff --git a/Engine/Plugins/Experimental/ProceduralVegetationEditor/Source/ProceduralVegetation/Private/Nodes/PVGrowthDataLoader.cpp b/Engine/Plugins/Experimental/ProceduralVegetationEditor/Source/ProceduralVegetation/Private/Nodes/PVGrowthDataLoader.cpp
new file mode 100644
--- /dev/null
+++ b/Engine/Plugins/Experimental/ProceduralVegetationEditor/Source/ProceduralVegetation/Private/Nodes/PVGrowthDataLoader.cpp
@@ -0,0 +1,145 @@
+// Copyright Epic Games, Inc. All Rights Reserved.
+
+#include "PVGrowthDataLoader.h"
+#include "PCGContext.h"
+#include "ProceduralVegetationModule.h"
+#include "DataAssets/ProceduralVegetationGrowthDataAsset.h"
+#include "DataTypes/PVGrowthData.h"
+#include "Facades/PVPointFacade.h"
+
+#define LOCTEXT_NAMESPACE "PVPresetLoaderSettings"
+
+#if WITH_EDITOR
+FText UPVGrowthDataLoaderSettings::GetDefaultNodeTitle() const
+{ 
+	return LOCTEXT("NodeTitle", "Growth Data Loader"); 
+}
+
+FText UPVGrowthDataLoaderSettings::GetNodeTooltipText() const
+{	
+	return LOCTEXT("NodeTooltip", 
+		"This allows the user to load a procedural vegetation growth data. "
+		"A Procedural Vegetation Growth Asset contains data for a specie of vegetation. "
+		"It contains hormonal information, the main structure of the vegetation, data about the branches"
+		"\n\nPress Ctrl + L to lock/unlock node output"
+	);
+}
+
+void UPVGrowthDataLoaderSettings::PostEditChangeProperty(FPropertyChangedEvent& PropertyChangedEvent)
+{
+	Super::PostEditChangeProperty(PropertyChangedEvent);
+	
+	if (PropertyChangedEvent.Property
+		&& PropertyChangedEvent.Property->GetFName() == GET_MEMBER_NAME_CHECKED(UPVGrowthDataLoaderSettings, GrowthAsset))
+	{
+		FillGrowthVariationsInfo();
+	}
+}
+#endif
+
+void UPVGrowthDataLoaderSettings::PostLoad()
+{
+	Super::PostLoad();
+	
+	if (GrowthVariations.IsEmpty())
+	{
+		FillGrowthVariationsInfo();
+	}
+}
+
+TArray<FPCGPinProperties> UPVGrowthDataLoaderSettings::InputPinProperties() const
+{
+	return TArray<FPCGPinProperties>();
+}
+
+TArray<FPCGPinProperties> UPVGrowthDataLoaderSettings::OutputPinProperties() const
+{
+	TArray<FPCGPinProperties> Properties;
+
+	if (GrowthAsset)
+	{
+		for (const auto& [VariantName, VariantData] : GrowthAsset->Variants)
+		{
+			if (PV::Utilities::IsValidGrowthData(VariantData))
+			{
+				FPCGPinProperties& Pin = Properties.Emplace_GetRef(*VariantName, GetOutputPinTypeIdentifier());
+				Pin.bAllowMultipleData = false;
+			}
+		}
+	}
+
+	return Properties;
+}
+
+FPCGDataTypeIdentifier UPVGrowthDataLoaderSettings::GetOutputPinTypeIdentifier() const
+{
+	return FPCGDataTypeIdentifier{ FPVDataTypeInfoGrowth::AsId() };
+}
+
+FPCGElementPtr UPVGrowthDataLoaderSettings::CreateElement() const
+{
+	return MakeShared<FPVGrowthDataLoaderElement>();
+}
+
+void UPVGrowthDataLoaderSettings::FillGrowthVariationsInfo()
+{
+	GrowthVariations.Empty();
+
+	if (!GrowthAsset)
+	{
+		UE_LOGF(LogProceduralVegetation, Log, "Growth asset is null, no preset info can be filled.");
+		return;
+	}
+
+	if (GrowthAsset->GrowthVariations.IsEmpty())
+	{
+		// GrowthAsset->Fill();
+	}
+
+	GrowthVariations = GrowthAsset->GrowthVariations;
+}
+
+bool FPVGrowthDataLoaderElement::ExecuteInternal(FPCGContext* InContext) const
+{
+	TRACE_CPUPROFILER_EVENT_SCOPE(FPVGrowthDataLoaderElement::ExecuteInternal);
+
+	check(InContext);
+
+	const UPVGrowthDataLoaderSettings* Settings = InContext->GetInputSettings<UPVGrowthDataLoaderSettings>();
+	check(Settings);
+
+	if (!Settings->GrowthAsset)
+	{
+		PCGLog::LogWarningOnGraph(NSLOCTEXT("PVPresetLoaderSettings", "MissingDataAsset", "Growth asset is not set"), InContext);
+		return true;
+	}
+
+	bool bValidDataExist = false;
+	for (const auto& [VariantName, VariantData] : Settings->GrowthAsset->Variants)
+	{
+		if (PV::Utilities::IsValidGrowthData(VariantData))
+		{
+			UPVGrowthData* OutVariantData = FPCGContext::NewObject_AnyThread<UPVGrowthData>(InContext);
+			OutVariantData->Initialize(CopyTemp(VariantData));
+
+			FPCGTaggedData& CollectionOutput = InContext->OutputData.TaggedData.Emplace_GetRef();
+			CollectionOutput.Data = OutVariantData;
+			CollectionOutput.Pin = *VariantName;
+			bValidDataExist = true;
+		}
+		else
+		{
+			PCGLog::LogWarningOnGraph(FText::FromString(FString::Format(TEXT("Skipping variation {0} due to invalid data."), {*VariantName})), InContext);
+		}
+	}
+
+	if (!bValidDataExist)
+	{
+		PCGLog::LogErrorOnGraph(FText::FromString(FString::Format(TEXT("No valid variation found in data asset {0}"), { Settings->GrowthAsset.GetFName().ToString()})), InContext);
+		return true;
+	}
+	
+	return true;
+}
+
+#undef LOCTEXT_NAMESPACE
\ No newline at end of file
diff --git a/Engine/Plugins/Experimental/ProceduralVegetationEditor/Source/ProceduralVegetation/Private/Nodes/PVGrowthDataLoader.h b/Engine/Plugins/Experimental/ProceduralVegetationEditor/Source/ProceduralVegetation/Private/Nodes/PVGrowthDataLoader.h
new file mode 100644
--- /dev/null
+++ b/Engine/Plugins/Experimental/ProceduralVegetationEditor/Source/ProceduralVegetation/Private/Nodes/PVGrowthDataLoader.h
@@ -0,0 +1,49 @@
+// Copyright Epic Games, Inc. All Rights Reserved.
+#pragma once
+
+#include "CoreMinimal.h"
+#include "Nodes/PVBaseSettings.h"
+#include "PVGrowthDataLoader.generated.h"
+
+struct FPVGrowthVariationInfo;
+
+UCLASS(BlueprintType, ClassGroup = (Procedural))
+class UPVGrowthDataLoaderSettings : public UPVBaseSettings
+{
+	GENERATED_BODY()
+
+public:
+	virtual void PostLoad() override;
+	
+#if WITH_EDITOR
+	virtual FName GetDefaultNodeName() const override { return FName(TEXT("GrowthDataLoader")); }
+	virtual FText GetDefaultNodeTitle() const override;
+	virtual FText GetNodeTooltipText() const override;
+	virtual EPCGSettingsType GetType() const override { return EPCGSettingsType::InputOutput; }
+	virtual TArray<EPVRenderType> GetDefaultRenderType() const override { return TArray{ EPVRenderType::PointData }; }
+
+	virtual void PostEditChangeProperty(FPropertyChangedEvent& PropertyChangedEvent) override;
+#endif
+
+protected:
+	virtual TArray<FPCGPinProperties> InputPinProperties() const override;
+	virtual TArray<FPCGPinProperties> OutputPinProperties() const override;
+	virtual FPCGDataTypeIdentifier GetOutputPinTypeIdentifier() const override;
+	virtual FPCGElementPtr CreateElement() const override;
+
+public:
+	UPROPERTY(EditAnywhere, Category=Settings, meta=(Tooltip="Loads the selected procedural vegetation preset\n\nThis allows the user to load a procedural vegetation preset. A procedural vegatation preset contains data for a specie of vegetaion. It contains hormonal information, the main structure of the vegetation, data about the branches, where foliage will/can be attached, references to foliage meshes, materials to be applied"))
+	TObjectPtr<class UProceduralVegetationGrowthDataAsset> GrowthAsset = nullptr;
+
+	UPROPERTY(VisibleAnywhere, Category = Settings)
+	TArray<FPVGrowthVariationInfo> GrowthVariations;
+
+private:
+	void FillGrowthVariationsInfo();
+};
+
+class FPVGrowthDataLoaderElement : public FPVBaseElement
+{
+protected:
+	virtual bool ExecuteInternal(FPCGContext* InContext) const override;
+};
