Minor TraversalHelper refactor

- Change so GetBranchPointChildBranchIndices takes a BranchPointIndex instead of a PointIndex to avoid performing a Find in the BranchPoints array
- Add ParentPointIndex to FRecursiveWalkBranchPointsParams so we don't have to keep a duplicated data structure in FRotateBranchPointParams
- Add RecursiveRotatePlantPoints to PVTransformHelpers.h

#Synced-CL 51206265

diff --git a/Engine/Plugins/Experimental/ProceduralVegetationEditor/Source/ProceduralVegetation/Private/Helpers/PVAttributesHelper.cpp b/Engine/Plugins/Experimental/ProceduralVegetationEditor/Source/ProceduralVegetation/Private/Helpers/PVAttributesHelper.cpp
--- a/Engine/Plugins/Experimental/ProceduralVegetationEditor/Source/ProceduralVegetation/Private/Helpers/PVAttributesHelper.cpp
+++ b/Engine/Plugins/Experimental/ProceduralVegetationEditor/Source/ProceduralVegetation/Private/Helpers/PVAttributesHelper.cpp
@@ -432,7 +432,7 @@ void PV::AttributesHelper::ComputeBudStatus(const FComputeBudStatusAttributes& I
 					InAttributes.BranchPoints,
 					InAttributes.PointLengthFromRoot,
 					BranchIndex, 
-					PointIndex
+					BranchPointIndex
 				).Num();
 
 				if (NumChildren > 0) 
diff --git a/Engine/Plugins/Experimental/ProceduralVegetationEditor/Source/ProceduralVegetation/Private/Helpers/PVPlantTraversalHelper.h b/Engine/Plugins/Experimental/ProceduralVegetationEditor/Source/ProceduralVegetation/Private/Helpers/PVPlantTraversalHelper.h
--- a/Engine/Plugins/Experimental/ProceduralVegetationEditor/Source/ProceduralVegetation/Private/Helpers/PVPlantTraversalHelper.h
+++ b/Engine/Plugins/Experimental/ProceduralVegetationEditor/Source/ProceduralVegetation/Private/Helpers/PVPlantTraversalHelper.h
@@ -234,22 +234,27 @@ namespace PV::PlantTraversalHelper
 		PV::FBranchPointsAttributeConstView BranchPointsAttribute,
 		PV::FPointLengthFromRootAttributeConstView PointLengthFromRootAttribute,
 		int32 BranchIndex,
-		int32 PointIndex
+		int32 BranchPointIndex
 	)
 	{
-		const TArray<int32> BranchPoints = BranchPointsAttribute[BranchIndex];
-		const int32 PointBranchIndex = BranchPoints.Find(PointIndex);
-		if (PointBranchIndex == INDEX_NONE)
+		if (!BranchPointsAttribute.IsValidIndex(BranchIndex))
 		{
 			return TArray<int32>();
 		}
 
-		const int32 PreviousPointBranchIndex = PointBranchIndex > 0 ? PointBranchIndex - 1 : INDEX_NONE;
+		const TArray<int32>& BranchPoints = BranchPointsAttribute[BranchIndex];
+		if (!BranchPoints.IsValidIndex(BranchPointIndex))
+		{
+			return TArray<int32>();
+		}
+
+		const int32 PreviousPointBranchIndex = BranchPointIndex > 0 ? BranchPointIndex - 1 : INDEX_NONE;
 		if (PreviousPointBranchIndex == INDEX_NONE)
 		{
 			return TArray<int32>();
 		}
 
+		const int32 PointIndex = BranchPoints[BranchPointIndex];
 		const int32 PreviousPointIndex = BranchPoints[PreviousPointBranchIndex];
 
 		const TArray<int32>& BranchChildren = BranchChildrenAttribute[BranchIndex];
@@ -299,9 +304,21 @@ namespace PV::PlantTraversalHelper
 		PV::FBranchNumberAttributeConstView BranchNumberAttribute,
 		PV::FBranchPointsAttributeConstView BranchPointsAttribute,
 		int32 BranchIndex,
-		int32 PointIndex
+		int32 BranchPointIndex
 	)
 	{
+		if (!BranchPointsAttribute.IsValidIndex(BranchIndex))
+		{
+			return TArray<int32>();
+		}
+
+		const TArray<int32>& BranchPoints = BranchPointsAttribute[BranchIndex];
+		if (!BranchPoints.IsValidIndex(BranchPointIndex))
+		{
+			return TArray<int32>();
+		}
+
+		const int32 PointIndex = BranchPoints[BranchPointIndex];
 		const TArray<int32>& BranchChildren = BranchChildrenAttribute[BranchIndex];
 
 		TArray<int32> ChildrenIndices;
@@ -344,7 +361,7 @@ namespace PV::PlantTraversalHelper
 		PV::FBranchPointsAttributeConstView BranchPointsAttribute,
 		PV::FPointLengthFromRootAttributeConstView PointLengthFromRootAttribute,
 		int32 BranchIndex,
-		int32 PointIndex
+		int32 BranchPointIndex
 	)
 	{
 		if (HasFusedPoints(BranchPointsAttribute, BranchParentNumberAttribute, BranchNumberAttribute))
@@ -355,7 +372,7 @@ namespace PV::PlantTraversalHelper
 				BranchNumberAttribute,
 				BranchPointsAttribute,
 				BranchIndex,
-				PointIndex
+				BranchPointIndex
 			);
 		}
 		else
@@ -367,7 +384,7 @@ namespace PV::PlantTraversalHelper
 				BranchPointsAttribute,
 				PointLengthFromRootAttribute,
 				BranchIndex,
-				PointIndex
+				BranchPointIndex
 			);
 		}
 	}
@@ -375,8 +392,9 @@ namespace PV::PlantTraversalHelper
 	struct FRecursiveWalkBranchPointsParams
 	{
 		int32 BranchIndex;
-		int32 BranchPointIndex;
-		int32 PointIndex;
+		int32 BranchPointIndex; // The index of the point in the branch points attribute
+		int32 PointIndex;	    // The global point index of this point
+		int32 ParentPointIndex; // The global point index of this points parent point (INDEX_NONE for the first point on the trunk)
 		const TArray<int32>& ChildBranchIndices;
 		const TArray<int32>& ChildPointIndices;
 	};
@@ -393,7 +411,7 @@ namespace PV::PlantTraversalHelper
 		const TFunction<EForEachResult(const FRecursiveWalkBranchPointsParams& Params, StateType& PointState)>& Body
 	)
 	{
-		const TFunction<EForEachResult(int32, int32, const StateType&)> WalkPoints = [&](int32 BranchIndex, int32 BranchPointIndex, const StateType& ParentState) -> EForEachResult
+		const TFunction<EForEachResult(int32, int32, int32, const StateType&)> WalkPoints = [&](int32 BranchIndex, int32 BranchPointIndex, int32 ParentPointIndex, const StateType& ParentState) -> EForEachResult
 		{
 			StateType PointState = ParentState;
 			
@@ -407,7 +425,7 @@ namespace PV::PlantTraversalHelper
 				BranchPointsAttribute,
 				PointLengthFromRootAttribute,
 				BranchIndex,
-				PointIndex
+				BranchPointIndex
 			);
 
 			TArray<int32> ChildPointIndices;
@@ -418,7 +436,7 @@ namespace PV::PlantTraversalHelper
 				ChildPointIndices[i] = BranchPointsAttribute[ChildBranchIndex].Num() > 0 ? BranchPointsAttribute[ChildBranchIndex][0] : INDEX_NONE;
 			}
 
-			if (Body({ BranchIndex, BranchPointIndex, PointIndex, ChildBranchIndices, ChildPointIndices }, PointState) == EForEachResult::Break)
+			if (Body({ BranchIndex, BranchPointIndex, PointIndex, ParentPointIndex, ChildBranchIndices, ChildPointIndices }, PointState) == EForEachResult::Break)
 			{
 				return EForEachResult::Break;
 			}
@@ -458,7 +476,7 @@ namespace PV::PlantTraversalHelper
 
 			for (const auto& [ChildBranchIndex, ChildBranchPointIndex] : Children)
 			{
-				if (WalkPoints(ChildBranchIndex, ChildBranchPointIndex, PointState) == EForEachResult::Break)
+				if (WalkPoints(ChildBranchIndex, ChildBranchPointIndex, PointIndex, PointState) == EForEachResult::Break)
 				{
 					return EForEachResult::Break;
 				}
@@ -467,7 +485,7 @@ namespace PV::PlantTraversalHelper
 			return EForEachResult::Continue;
 		};
 
-		return WalkPoints(BranchIndex, 0, RootPointState);
+		return WalkPoints(BranchIndex, 0, INDEX_NONE, RootPointState);
 	}
 
 	template <typename StateType, bool bWalkAxillaryFirst = false>
@@ -481,7 +499,7 @@ namespace PV::PlantTraversalHelper
 		const TFunction<EForEachResult(const FRecursiveWalkBranchPointsParams& Params, StateType& PointState)>& Body
 	)
 	{
-		const TFunction<EForEachResult(int32, int32, const StateType&)> WalkPoints = [&](int32 BranchIndex, int32 BranchPointIndex, const StateType& ParentState) -> EForEachResult
+		const TFunction<EForEachResult(int32, int32, int32, const StateType&)> WalkPoints = [&](int32 BranchIndex, int32 BranchPointIndex, int32 ParentPointIndex, const StateType& ParentState) -> EForEachResult
 		{
 			StateType PointState = ParentState;
 			
@@ -494,7 +512,7 @@ namespace PV::PlantTraversalHelper
 				BranchNumberAttribute,
 				BranchPointsAttribute,
 				BranchIndex,
-				PointIndex
+				BranchPointIndex
 			);
 
 			TArray<int32> ChildPointIndices;
@@ -505,7 +523,7 @@ namespace PV::PlantTraversalHelper
 				ChildPointIndices[i] = BranchPointsAttribute[ChildBranchIndex].Num() > 1 ? BranchPointsAttribute[ChildBranchIndex][1] : INDEX_NONE;
 			}
 
-			if (Body({ BranchIndex, BranchPointIndex, PointIndex, ChildBranchIndices, ChildPointIndices }, PointState) == EForEachResult::Break)
+			if (Body({ BranchIndex, BranchPointIndex, PointIndex, ParentPointIndex, ChildBranchIndices, ChildPointIndices }, PointState) == EForEachResult::Break)
 			{
 				return EForEachResult::Break;
 			}
@@ -545,7 +563,7 @@ namespace PV::PlantTraversalHelper
 
 			for (const auto& [ChildBranchIndex, ChildBranchPointIndex] : Children)
 			{
-				if (WalkPoints(ChildBranchIndex, ChildBranchPointIndex, PointState) == EForEachResult::Break)
+				if (WalkPoints(ChildBranchIndex, ChildBranchPointIndex, PointIndex, PointState) == EForEachResult::Break)
 				{
 					return EForEachResult::Break;
 				}
@@ -554,7 +572,7 @@ namespace PV::PlantTraversalHelper
 			return EForEachResult::Continue;
 		};
 
-		return WalkPoints(BranchIndex, 0, RootPointState);
+		return WalkPoints(BranchIndex, 0, INDEX_NONE, RootPointState);
 	}
 
 	template<typename StateType, bool bWalkAxillaryFirst = false>
diff --git a/Engine/Plugins/Experimental/ProceduralVegetationEditor/Source/ProceduralVegetation/Private/Helpers/PVTransformHelper.h b/Engine/Plugins/Experimental/ProceduralVegetationEditor/Source/ProceduralVegetation/Private/Helpers/PVTransformHelper.h
--- a/Engine/Plugins/Experimental/ProceduralVegetationEditor/Source/ProceduralVegetation/Private/Helpers/PVTransformHelper.h
+++ b/Engine/Plugins/Experimental/ProceduralVegetationEditor/Source/ProceduralVegetation/Private/Helpers/PVTransformHelper.h
@@ -109,15 +109,7 @@ PROCEDURALVEGETATION_API extern void TransformPoint(
 	const FQuat4f& Rotation
 );
 
-struct FRotateBranchPointParams
-{
-	int32 BranchIndex;
-	int32 BranchPointIndex; // The index of the point in the branch points attribute
-	int32 PointIndex;		// The global point index of this point
-	int32 ParentPointIndex; // INDEX_NONE for the first point on the trunk
-	const TArray<int32>& ChildBranchIndices;
-	const TArray<int32>& ChildPointIndices;
-};
+using FRotateBranchPointParams = PV::PlantTraversalHelper::FRecursiveWalkBranchPointsParams;
 
 struct FRotateBranchPointResult
 {
@@ -168,7 +160,6 @@ inline bool RecursiveRotateBranchPoints(
 
 	struct FPointState
 	{
-		int32      ParentPointIndex;
 		FVector3f  ParentPivotPoint;
 		FVector3f  ParentIdentityPosition;
 		FMatrix44f TransformMatrix;
@@ -176,7 +167,6 @@ inline bool RecursiveRotateBranchPoints(
 	};
 
 	const FPointState DefaultPointState = {
-		INDEX_NONE,
 		FVector3f::ZeroVector,
 		FVector3f::ZeroVector,
 		FMatrix44f::Identity,
@@ -193,22 +183,14 @@ inline bool RecursiveRotateBranchPoints(
 		DefaultPointState,
 		[&](const FRecursiveWalkBranchPointsParams& InParams, FPointState& PointState)
 		{
-			const FRotateBranchPointParams Params = FRotateBranchPointParams(
-				InParams.BranchIndex,
-				InParams.BranchPointIndex,
-				InParams.PointIndex,
-				PointState.ParentPointIndex,
-				InParams.ChildBranchIndices,
-				InParams.ChildPointIndices
-			);
 			FRotateBranchPointResult Result = {
 				.Rotation = FQuat4f::Identity,
 				.bAbsoluteRotation = false,
 			};
-			RotatePointsBody(Params, Result, PointState.UserState);
+			RotatePointsBody(InParams, Result, PointState.UserState);
 
 			const FVector3f& PointPosition = Attributes.PointPosition[InParams.PointIndex];
-			const FVector3f RelativePosition = PointState.ParentPointIndex == INDEX_NONE
+			const FVector3f RelativePosition = InParams.ParentPointIndex == INDEX_NONE
 				? TrunkPivot
 				: PointPosition - PointState.ParentIdentityPosition;
 
@@ -232,8 +214,6 @@ inline bool RecursiveRotateBranchPoints(
 			const FQuat4f PointRotation = PointState.TransformMatrix.ToQuat();
 			TransformPoint(Attributes, InParams.BranchIndex, InParams.PointIndex, PointTranslation, PointRotation);
 
-			PointState.ParentPointIndex = InParams.PointIndex;
-
 			return EForEachResult::Continue;
 		}
 	);
@@ -257,4 +237,50 @@ inline bool RecursiveRotateBranchPoints(
 	);
 }
 
+template<typename StateType>
+inline bool RecursiveRotatePlantPoints(
+	const FRotateBranchPointsAttributeCollection& Attributes,
+	const StateType& RootPointState,
+	const TFunction<void(const FRotateBranchPointParams& Params, FRotateBranchPointResult& Result, StateType& State)>& RotatePointsBody,
+	const FVector3f& TrunkPivot = FVector3f::ZeroVector
+)
+{
+	if (!Attributes.BranchParentNumber.IsValid())
+	{
+		return false;
+	}
+
+	for (const int32 TrunkIndex : PV::PlantTraversalHelper::GetTrunkIndices(Attributes.BranchParentNumber))
+	{
+		const auto Result = RecursiveRotateBranchPoints<StateType>(
+			Attributes,
+			TrunkIndex,
+			RootPointState,
+			RotatePointsBody,
+			TrunkPivot
+		);
+
+		if (!Result)
+		{
+			return false;
+		}
+	}
+
+	return true;
+}
+
+inline bool RecursiveRotatePlantPoints(
+	const FRotateBranchPointsAttributeCollection& Attributes,
+	const TFunction<void(const FRotateBranchPointParams& Params, FRotateBranchPointResult& Result)>& RotatePointsBody,
+	const FVector3f& TrunkPivot = FVector3f::ZeroVector
+)
+{
+	return RecursiveRotatePlantPoints<int32>(
+		Attributes,
+		0,
+		[&](const FRotateBranchPointParams& Params, FRotateBranchPointResult& Result, int32&) { RotatePointsBody(Params, Result); },
+		TrunkPivot
+	);
+}
+
 };
\ No newline at end of file
