[PVE] Plant profile refactor
- Added Plant profile data asset
- Added ability to load profiles json separately in the data asset.
- Created a new node, PlantProfileLoader, which load the profiles info.
- Modified MeshBuilder to accept Profile pin
#Synced-CL 51086708

diff --git a/Engine/Plugins/Experimental/ProceduralVegetationEditor/Source/ProceduralVegetation/Private/DataAssets/ProceduralVegetationPlantProfileDataAsset.cpp b/Engine/Plugins/Experimental/ProceduralVegetationEditor/Source/ProceduralVegetation/Private/DataAssets/ProceduralVegetationPlantProfileDataAsset.cpp
new file mode 100644
--- /dev/null
+++ b/Engine/Plugins/Experimental/ProceduralVegetationEditor/Source/ProceduralVegetation/Private/DataAssets/ProceduralVegetationPlantProfileDataAsset.cpp
@@ -0,0 +1,23 @@
+﻿// Copyright Epic Games, Inc. All Rights Reserved.
+ 
+#include "ProceduralVegetationPlantProfileDataAsset.h"
+#include "ProceduralVegetationModule.h"
+#include "Helpers/PVJSONHelper.h"
+
+void UProceduralVegetationPlantProfileDataAsset::Load()
+{
+	FString OutErrorMessage;
+	TArray<FPVPlantProfile> OutProfileData;
+	
+	if (PV::JSON::LoadProfileData(ProfileFilePath.FilePath, OutProfileData, OutErrorMessage))
+	{
+		if (OutProfileData.Num() > 0)
+		{
+			Profiles = OutProfileData;
+		}
+	}
+	else if (!OutErrorMessage.IsEmpty())
+	{
+		UE_LOGF(LogProceduralVegetation, Error, ("PlantProfileDataAsset : Error loading profile data : %ls"), *OutErrorMessage);
+	}
+}
diff --git a/Engine/Plugins/Experimental/ProceduralVegetationEditor/Source/ProceduralVegetation/Private/DataAssets/ProceduralVegetationPlantProfileDataAsset.h b/Engine/Plugins/Experimental/ProceduralVegetationEditor/Source/ProceduralVegetation/Private/DataAssets/ProceduralVegetationPlantProfileDataAsset.h
new file mode 100644
--- /dev/null
+++ b/Engine/Plugins/Experimental/ProceduralVegetationEditor/Source/ProceduralVegetation/Private/DataAssets/ProceduralVegetationPlantProfileDataAsset.h
@@ -0,0 +1,33 @@
+// Copyright Epic Games, Inc. All Rights Reserved.
+
+#pragma once
+
+#include "Engine/DataAsset.h"
+#include "ProceduralVegetationPlantProfileDataAsset.generated.h"
+
+USTRUCT(BlueprintType)
+struct PROCEDURALVEGETATION_API FPVPlantProfile
+{
+	GENERATED_BODY()
+
+	UPROPERTY(VisibleAnywhere, Category="PlantProfile")
+	FString Name;
+	UPROPERTY(VisibleAnywhere, Category="PlantProfile")
+	TArray<float> Points;
+};
+
+
+UCLASS(BlueprintType)
+class PROCEDURALVEGETATION_API UProceduralVegetationPlantProfileDataAsset : public UDataAsset
+{
+	GENERATED_BODY()
+public:
+	UPROPERTY(VisibleAnywhere, Category="PlantProfiles")
+	TArray<FPVPlantProfile> Profiles;
+
+	UPROPERTY(EditAnywhere, Category="Internal", meta = (DevelopmentOnly))
+	FFilePath ProfileFilePath;
+
+	UFUNCTION(CallInEditor, Category="Internal", meta = (DevelopmentOnly))
+	void Load();
+};
\ No newline at end of file
diff --git a/Engine/Plugins/Experimental/ProceduralVegetationEditor/Source/ProceduralVegetation/Private/DataTypes/PVPlantProfileData.cpp b/Engine/Plugins/Experimental/ProceduralVegetationEditor/Source/ProceduralVegetation/Private/DataTypes/PVPlantProfileData.cpp
new file mode 100644
--- /dev/null
+++ b/Engine/Plugins/Experimental/ProceduralVegetationEditor/Source/ProceduralVegetation/Private/DataTypes/PVPlantProfileData.cpp
@@ -0,0 +1,5 @@
+// Copyright Epic Games, Inc. All Rights Reserved.
+
+#include "DataTypes/PVPlantProfileData.h"
+
+PCG_DEFINE_TYPE_INFO(FPVDataTypeInfoPlantProfile, UPVPlantProfileData);
\ No newline at end of file
diff --git a/Engine/Plugins/Experimental/ProceduralVegetationEditor/Source/ProceduralVegetation/Private/Helpers/PVJSONHelper.cpp b/Engine/Plugins/Experimental/ProceduralVegetationEditor/Source/ProceduralVegetation/Private/Helpers/PVJSONHelper.cpp
--- a/Engine/Plugins/Experimental/ProceduralVegetationEditor/Source/ProceduralVegetation/Private/Helpers/PVJSONHelper.cpp
+++ b/Engine/Plugins/Experimental/ProceduralVegetationEditor/Source/ProceduralVegetation/Private/Helpers/PVJSONHelper.cpp
@@ -4,6 +4,7 @@
 
 #include "ProceduralVegetationModule.h"
 #include "PVAttributesHelper.h"
+#include "DataAssets/ProceduralVegetationPlantProfileDataAsset.h"
 #include "Dom/JsonObject.h"
 #include "Dom/JsonValue.h"
 #include "Facades/PVAttributesNames.h"
@@ -20,6 +21,25 @@
 
 namespace PV::JSON
 {
+	bool LoadJSON(const FString& InPath, TSharedPtr<FJsonObject>& OutData, FString& OutErrorMessage)
+	{
+		FString JsonString;
+		if (!FFileHelper::LoadFileToString(JsonString, *InPath))
+		{
+			OutErrorMessage = TEXT("Failed to open file");
+			return false;
+		}
+
+		OutData = MakeShareable(new FJsonObject);
+		if (!FJsonSerializer::Deserialize(TJsonReaderFactory<>::Create(JsonString), OutData))
+		{
+			OutErrorMessage = TEXT("Failed to deserialize JSON file");
+			return false;
+		}
+
+		return true;
+	}
+	
 	void FillAttributes(FManagedArrayCollection& Collection, const FName Group, const TSharedPtr<FJsonObject>& AttributesObject)
 	{
 		for (const auto& [Key, Attribute] : AttributesObject->Values)
@@ -703,4 +723,80 @@ namespace PV::JSON
 		
 		return true;
 	}
+
+	bool LoadProfileData(FString InPath, TArray<FPVPlantProfile>& OutProfileData, FString& OutErrorMessage)
+	{
+		TSharedPtr<FJsonObject> OutData = MakeShareable(new FJsonObject);
+		
+		OutProfileData.Empty();
+
+		if (LoadJSON(InPath, OutData, OutErrorMessage))
+		{
+			if (OutData->HasTypedField(TEXT("PlantProfiles"), EJson::Array))
+			{
+				TArray<FPVPlantProfile> Profiles;
+				
+				const TArray<TSharedPtr<FJsonValue>>& ProfilesArray = OutData->GetArrayField(TEXT("PlantProfiles"));
+
+				int Index = 0;
+
+				for (TSharedPtr<FJsonValue> ProfileJSONValue : ProfilesArray)
+				{
+					if (ProfileJSONValue->Type == EJson::Object)
+					{
+						FPVPlantProfile Profile;
+						
+						const TSharedPtr<FJsonObject> ProfilesObject = ProfileJSONValue->AsObject();
+
+						if (!ProfilesObject->TryGetStringField(TEXT("Name"), Profile.Name))
+						{
+							OutErrorMessage = FString::Printf(TEXT("ProfileInfo Doesnt contain valid name for Profile[%i]"), Index);
+							return false;
+						}
+
+						if (ProfilesObject->HasTypedField<EJson::Array>(TEXT("Points")))
+						{
+							const TArray<TSharedPtr<FJsonValue>>& Values = ProfilesObject->GetArrayField(TEXT("Points"));
+						
+							Profile.Points.Reserve(Values.Num());
+
+							for (const TSharedPtr<FJsonValue>& AttribValue : Values)
+							{
+								if (double Point; AttribValue->TryGetNumber(Point))
+								{
+									Profile.Points.Add(static_cast<float>(AttribValue->AsNumber()));
+								}
+							}	
+						}
+						else
+						{
+							OutErrorMessage = FString::Printf(TEXT("ProfileInfo doesnt contain valid points for Profile[%i]"), Index);
+							return false;
+						}
+
+						Profiles.Add(Profile);
+					}
+
+					Index++;
+				}
+
+				if (Profiles.Num())
+				{
+					OutProfileData = Profiles;
+				}
+			}
+			else
+			{
+				OutErrorMessage = FString::Printf(TEXT("Unable to load profile data"));
+				
+				return false;	
+			}
+		}
+		else
+		{
+			return false;
+		}
+
+		return true;
+	}
 }
diff --git a/Engine/Plugins/Experimental/ProceduralVegetationEditor/Source/ProceduralVegetation/Private/Helpers/PVJSONHelper.h b/Engine/Plugins/Experimental/ProceduralVegetationEditor/Source/ProceduralVegetation/Private/Helpers/PVJSONHelper.h
--- a/Engine/Plugins/Experimental/ProceduralVegetationEditor/Source/ProceduralVegetation/Private/Helpers/PVJSONHelper.h
+++ b/Engine/Plugins/Experimental/ProceduralVegetationEditor/Source/ProceduralVegetation/Private/Helpers/PVJSONHelper.h
@@ -4,6 +4,7 @@
 
 #include "CoreMinimal.h"
 
+struct FPVPlantProfile;
 struct FManagedArrayCollection;
 class FJsonObject;
 
@@ -28,4 +29,6 @@ namespace PV::JSON
 	TSharedPtr<FJsonObject> LoadMetaFileIntoJsonObject(const FString& FilePath, FString& OutErrorMessage);
 
 	bool LoadMetaJsonToCollection(FManagedArrayCollection& Collection, TSharedPtr<FJsonObject> LoadedData);
+
+	bool LoadProfileData(FString FilePath, TArray<FPVPlantProfile>& OutProfileData, FString& OutErrorMessage);
 }
diff --git a/Engine/Plugins/Experimental/ProceduralVegetationEditor/Source/ProceduralVegetation/Private/Implementations/PVMeshBuilder.cpp b/Engine/Plugins/Experimental/ProceduralVegetationEditor/Source/ProceduralVegetation/Private/Implementations/PVMeshBuilder.cpp
--- a/Engine/Plugins/Experimental/ProceduralVegetationEditor/Source/ProceduralVegetation/Private/Implementations/PVMeshBuilder.cpp
+++ b/Engine/Plugins/Experimental/ProceduralVegetationEditor/Source/ProceduralVegetation/Private/Implementations/PVMeshBuilder.cpp
@@ -45,7 +45,7 @@ void MergeMaterials(const TArray<FTrunkGenerationMaterialSetup>& InMaterialSetup
 	}
 }
 
-void FPVMeshBuilder::GenerateGeometryCollection(const FManagedArrayCollection& InSkeletonCollection, const FPVMeshBuilderParams& MeshBuilderParams,
+void FPVMeshBuilder::GenerateGeometryCollection(const FManagedArrayCollection& InSkeletonCollection, const FManagedArrayCollection& InPlantProfileCollection, const FPVMeshBuilderParams& MeshBuilderParams,
                                                 FGeometryCollection& OutGeometryCollection)
 {
 	InSkeletonCollection.CopyTo(&OutGeometryCollection);
@@ -75,7 +75,7 @@ void FPVMeshBuilder::GenerateGeometryCollection(const FManagedArrayCollection& I
 	PV::AttributesHelper::ComputeGradients(OutGeometryCollection);
 
 	MeshBuilderParams.MaterialSettings.ApplyMaterialSettings(OutGeometryCollection);
-	GetLocalDynamicMeshData(OutGeometryCollection, MeshBuilderParams, MeshesDataCollection);
+	GetLocalDynamicMeshData(OutGeometryCollection, InPlantProfileCollection, MeshBuilderParams, MeshesDataCollection);
 
 	uint32 VertexCount = 0;
 	uint32 TriangleCount = 0;
@@ -365,11 +365,11 @@ void FPVMeshBuilder::RemoveNonBudPoints(PV::Facades::FBranchFacade& BranchFacade
 	}
 }
 
-void FPVMeshBuilder::GenerateDynamicMesh(FManagedArrayCollection& Collection, const FPVMeshBuilderParams& MeshBuilderParams,
+void FPVMeshBuilder::GenerateDynamicMesh(FManagedArrayCollection& Collection, const FManagedArrayCollection& InPlantProfileCollection, const FPVMeshBuilderParams& MeshBuilderParams,
                                          TObjectPtr<UDynamicMesh>& OutMesh)
 {
 	TArray<FLocalDynamicMeshData> MeshesDataCollection;
-	GetLocalDynamicMeshData(Collection, MeshBuilderParams, MeshesDataCollection);
+	GetLocalDynamicMeshData(Collection, InPlantProfileCollection, MeshBuilderParams, MeshesDataCollection);
 
 	uint32 VertexCount = 0;
 	uint32 TriangleCount = 0;
@@ -841,7 +841,7 @@ TMap<int32, TArray<int32>> FPVMeshBuilder::GetPointsIndicesToFoliageIndicesMap(c
 
 void FPVMeshBuilder::GenerateBranchMeshData(const bool bPrimitiveIsTrunk, const int32 GenerationNumber, const TArray<int32>& PrimitivePoints,
                                             const FDisplacementData& DisplacementData, const FPVMeshBuilderParams& MeshBuilderParams,
-                                            const TArray<int32>& TargetMeshDivisions, const FManagedArrayCollection& Collection,
+                                            const TArray<int32>& TargetMeshDivisions, const FManagedArrayCollection& Collection, const FManagedArrayCollection& InPlantProfileCollection,
                                             FLocalDynamicMeshData& OutLocalMeshData)
 {
 	if (PrimitivePoints.Num() < 2)
@@ -851,25 +851,16 @@ void FPVMeshBuilder::GenerateBranchMeshData(const bool bPrimitiveIsTrunk, const
 	
 	const PV::Facades::FPointFacade PointFacade(Collection);
 	const PV::Facades::FBudVectorsFacade BudVectorFacade(Collection);
-	const PV::Facades::FPlantProfileFacade PlantProfileFacade(Collection);
+	const PV::Facades::FPlantProfileFacade PlantProfileFacade(InPlantProfileCollection);
 
 	bool bShouldApplyProfile =
 		PlantProfileFacade.NumProfileEntries() > 0
-		&& MeshBuilderParams.PlantProfileOptions.Num() > 0
 		&& (MeshBuilderParams.bApplyProfileToBranches || bPrimitiveIsTrunk);
 
 	TArray<float> ProfilePoints;
 	if (bShouldApplyProfile)
 	{
-		if (const int32 ProfileId = MeshBuilderParams.PlantProfileOptions.Find(MeshBuilderParams.SelectedPlantProfile) - 1;
-			ProfileId > -1)
-		{
-			ProfilePoints = PlantProfileFacade.GetProfilePoints(ProfileId);
-		}
-		else
-		{
-			bShouldApplyProfile = false;
-		}
+		ProfilePoints = PlantProfileFacade.GetProfilePoints(0);
 	}
 
 	const TManagedArray<FVector3f>& PointPositions = PointFacade.GetPositions();
@@ -1259,7 +1250,7 @@ bool IsDisplacementTextureValid(const TObjectPtr<UTexture2D>& Texture)
 #endif
 }
 
-void FPVMeshBuilder::GetLocalDynamicMeshData(FManagedArrayCollection& Collection, const FPVMeshBuilderParams& MeshBuilderParams,
+void FPVMeshBuilder::GetLocalDynamicMeshData(FManagedArrayCollection& Collection, const FManagedArrayCollection& InPlantProfileCollection, const FPVMeshBuilderParams& MeshBuilderParams,
                                              TArray<FLocalDynamicMeshData>& MeshesData)
 {
 	const PV::Facades::FBranchFacade BranchFacade(Collection);
@@ -1332,7 +1323,7 @@ void FPVMeshBuilder::GetLocalDynamicMeshData(FManagedArrayCollection& Collection
 				MeshesData[PrimitiveIndex].MaterialID = BranchFacade.GetBranchUVMaterial(PrimitiveIndex);
 
 				GenerateBranchMeshData(PrimitiveIsTrunk, GenerationNumber, PrimitivePoints, DisplacementData, MeshBuilderParams, TargetMeshDivisions,
-					Collection, MeshesData[PrimitiveIndex]);
+					Collection, InPlantProfileCollection, MeshesData[PrimitiveIndex]);
 			});
 }
 
diff --git a/Engine/Plugins/Experimental/ProceduralVegetationEditor/Source/ProceduralVegetation/Private/Implementations/PVMeshBuilder.h b/Engine/Plugins/Experimental/ProceduralVegetationEditor/Source/ProceduralVegetation/Private/Implementations/PVMeshBuilder.h
--- a/Engine/Plugins/Experimental/ProceduralVegetationEditor/Source/ProceduralVegetation/Private/Implementations/PVMeshBuilder.h
+++ b/Engine/Plugins/Experimental/ProceduralVegetationEditor/Source/ProceduralVegetation/Private/Implementations/PVMeshBuilder.h
@@ -142,32 +142,17 @@ struct FPVMeshBuilderParams
 		))
 	FPVFloatRamp ScaleRetentionGradient;
 
-	UPROPERTY(EditAnywhere, Category="Plant Profile Settings", DisplayName="Profile",
-		meta=(PCG_NotOverridable, EditCondition="bIsPlantProfileDropdownEnabled", EditConditionHides, HideEditConditionToggle, GetOptions="GetPlantProfileOptions",
-			Tooltip="Select a profile (cross sectional shape) from predefined options\n\nNone applies a simple circular cross sectional shape"))
-	FString SelectedPlantProfile = "None";
-
-	UPROPERTY(Transient)
-	TArray<FString> PlantProfileOptions = {"None"};
-
-	UPROPERTY(Transient)
-	bool bIsPlantProfileDropdownEnabled = false;
-
 	UPROPERTY(EditAnywhere, Category="Plant Profile Settings", DisplayName="Profile Falloff",
-		meta=(XAxisMin=0.0f, XAxisMax=1.0f, YAxisMin=0.0f, YAxisMax=1.0f, EditCondition="bIsPlantProfileDropdownEnabled", EditConditionHides,
-			HideEditConditionToggle, Tooltip=
-			"Gradient controlling how the profile shape is applied across the length of the trunk and branches\n\nCan be used to create a base where the trunk widens and transitions into main roots (Root Flare)"
+		meta=(XAxisMin=0.0f, XAxisMax=1.0f, YAxisMin=0.0f, YAxisMax=1.0f,
+			Tooltip = "Gradient controlling how the profile shape is applied across the length of the trunk and branches\n\nCan be used to create a base where the trunk widens and transitions into main roots (Root Flare)"
 		))
 	FPVFloatRamp PlantProfileFallOff;
 
 	UPROPERTY(EditAnywhere, Category="Plant Profile Settings", DisplayName="Profile Scale",
-		meta=(ClampMin=0.0f, ClampMax=10.0f, UIMin=0.0f, UIMax=10.0f, EditCondition="bIsPlantProfileDropdownEnabled", EditConditionHides,
-			HideEditConditionToggle, Tooltip="Scale of the profile applied\n\nCan be used to create a wider base in conjection with Profile Falloff"))
+		meta=(ClampMin=0.0f, ClampMax=10.0f, UIMin=0.0f, UIMax=10.0f, Tooltip="Scale of the profile applied\n\nCan be used to create a wider base in conjection with Profile Falloff"))
 	float PlantProfileScale = 1.0f;
 
-	UPROPERTY(EditAnywhere, Category="Plant Profile Settings",
-		meta=(EditCondition="bIsPlantProfileDropdownEnabled", EditConditionHides, HideEditConditionToggle, Tooltip=
-			"Should the profile be applied to branches or just the trunk"))
+	UPROPERTY(EditAnywhere, Category="Plant Profile Settings", meta=(Tooltip = "Should the profile be applied to branches or just the trunk"))
 	bool bApplyProfileToBranches = false;
 
 	UPROPERTY(EditAnywhere, Category = "Displacement Settings",
@@ -308,14 +293,14 @@ public:
 
 struct FPVMeshBuilder
 {
-	static void GenerateGeometryCollection(const FManagedArrayCollection& InSkeletonCollection, const FPVMeshBuilderParams& MeshBuilderParams,
+	static void GenerateGeometryCollection(const FManagedArrayCollection& InSkeletonCollection, const FManagedArrayCollection& InPlantProfileCollection, const FPVMeshBuilderParams& MeshBuilderParams,
 	                                       FGeometryCollection& OutGeometryCollection);
 
 	static void AddPointsToSkeleton(FManagedArrayCollection& OutCollection);
 	
 	static void RemoveNonBudPoints(PV::Facades::FBranchFacade& BranchFacade, PV::Facades::FPointFacade& PointFacade);
 
-	static void GenerateDynamicMesh(FManagedArrayCollection& Collection, const FPVMeshBuilderParams& MeshBuilderParams,
+	static void GenerateDynamicMesh(FManagedArrayCollection& Collection, const FManagedArrayCollection& InPlantProfileCollection, const FPVMeshBuilderParams& MeshBuilderParams,
 	                                TObjectPtr<UDynamicMesh>& OutMesh);
 
 	static bool ExtractDisplacementData(const TObjectPtr<UTexture2D>& Texture, TArray<float>& OutValues, FString& OutError);
@@ -348,11 +333,11 @@ private:
 
 	static void GenerateBranchMeshData(const bool bPrimitiveIsTrunk, const int32 GenerationNumber, const TArray<int32>& PrimitivePoints,
 	                                   const FDisplacementData& DisplacementData, const FPVMeshBuilderParams& MeshBuilderParams,
-	                                   const TArray<int32>& TargetMeshDivisions, const FManagedArrayCollection& Collection,
+	                                   const TArray<int32>& TargetMeshDivisions, const FManagedArrayCollection& Collection, const FManagedArrayCollection& InPlantProfileCollection, 
 	                                   FLocalDynamicMeshData& OutLocalMeshData);
 
 	static void UpdateFoliagePivotPoints(const TSet<int32>& PointsToRemove, FManagedArrayCollection& OutCollection);
 
-	static void GetLocalDynamicMeshData(FManagedArrayCollection& Collection, const FPVMeshBuilderParams& MeshBuilderParams,
+	static void GetLocalDynamicMeshData(FManagedArrayCollection& Collection, const FManagedArrayCollection& InPlantProfileCollection, const FPVMeshBuilderParams& MeshBuilderParams,
 	                                    TArray<FLocalDynamicMeshData>& MeshesData);
 };
diff --git a/Engine/Plugins/Experimental/ProceduralVegetationEditor/Source/ProceduralVegetation/Private/Nodes/PVMeshBuilderSettings.cpp b/Engine/Plugins/Experimental/ProceduralVegetationEditor/Source/ProceduralVegetation/Private/Nodes/PVMeshBuilderSettings.cpp
--- a/Engine/Plugins/Experimental/ProceduralVegetationEditor/Source/ProceduralVegetation/Private/Nodes/PVMeshBuilderSettings.cpp
+++ b/Engine/Plugins/Experimental/ProceduralVegetationEditor/Source/ProceduralVegetation/Private/Nodes/PVMeshBuilderSettings.cpp
@@ -1,10 +1,13 @@
 ﻿// Copyright Epic Games, Inc. All Rights Reserved.
 
 #include "PVMeshBuilderSettings.h"
+
+#include "PVCommon.h"
 #include "DataTypes/PVData.h"
 #include "DataTypes/PVGrowthData.h"
 #include "DataTypes/PVMeshData.h"
 #include "Data/PCGBasePointData.h"
+#include "DataTypes/PVPlantProfileData.h"
 #include "Facades/PVBranchFacade.h"
 #include "Facades/PVProfileFacade.h"
 #include "GeometryCollection/ManagedArrayCollection.h"
@@ -130,6 +133,17 @@ FPCGElementPtr UPVMeshBuilderSettings::CreateElement() const
 	return MakeShared<FPVMeshBuilderElement>();
 }
 
+TArray<FPCGPinProperties> UPVMeshBuilderSettings::InputPinProperties() const
+{
+	TArray<FPCGPinProperties> Properties = Super::InputPinProperties();
+
+	FPCGPinProperties& PlantProfilePin = Properties.Emplace_GetRef(PV::Pins::MeshBuilderPlantProfileInputLabel, FPCGDataTypeIdentifier{ FPVDataTypeInfoPlantProfile::AsId()});
+	PlantProfilePin.SetAllowMultipleConnections(false);
+	PlantProfilePin.SetAdvancedPin();
+
+	return Properties;
+}
+
 bool FPVMeshBuilderElement::ExecuteInternal(FPCGContext* InContext) const
 {
 	TRACE_CPUPROFILER_EVENT_SCOPE(FPVMeshBuilderElement::Execute);
@@ -142,8 +156,15 @@ bool FPVMeshBuilderElement::ExecuteInternal(FPCGContext* InContext) const
 	UPVMeshBuilderSettings* MutableSettings = const_cast<UPVMeshBuilderSettings*>(Settings);
 	check(MutableSettings);
 
-	MutableSettings->MesherSettings.PlantProfileOptions.Empty();
-	MutableSettings->MesherSettings.PlantProfileOptions.Add(TEXT("None"));
+	FManagedArrayCollection PlantProfileCollection;
+
+	for (const FPCGTaggedData& Input : InContext->InputData.GetInputsByPin(PV::Pins::MeshBuilderPlantProfileInputLabel))
+	{
+		if(const UPVPlantProfileData* Profile = Cast<UPVPlantProfileData>(Input.Data))
+		{
+			PlantProfileCollection = Profile->GetCollection();
+		}		
+	}
 
 	const TArray<FPCGTaggedData>& Inputs = InContext->InputData.GetInputsByPin(PCGPinConstants::DefaultInputLabel);
 	if (!Inputs.IsEmpty())
@@ -179,21 +200,6 @@ bool FPVMeshBuilderElement::ExecuteInternal(FPCGContext* InContext) const
 				MutableSettings->CustomMaterialSettings = MaterialSettings;
 			}
 
-			if (const PV::Facades::FPlantProfileFacade PlantProfileFacade = PV::Facades::FPlantProfileFacade(SkeletonCollection);
-				PlantProfileFacade.NumProfileEntries() > 0)
-			{
-				for (int32 i = 0; i < PlantProfileFacade.NumProfileEntries(); i++)
-				{
-					MutableSettings->MesherSettings.PlantProfileOptions.Add(FString::FromInt(i));
-				}
-
-				MutableSettings->MesherSettings.bIsPlantProfileDropdownEnabled = true;
-			}
-			else
-			{
-				MutableSettings->MesherSettings.bIsPlantProfileDropdownEnabled = false;
-			}
-
 			if (!Settings->DisplacementWarnings.IsEmpty())
 			{
 				PCGLog::LogWarningOnGraph(FText::FromString(FString::Format(TEXT("{0} "), {*Settings->DisplacementWarnings})), InContext);
@@ -202,7 +208,7 @@ bool FPVMeshBuilderElement::ExecuteInternal(FPCGContext* InContext) const
 			UPVMeshData* OutManagedArrayCollectionData = FPCGContext::NewObject_AnyThread<UPVMeshData>(InContext);
 
 			FGeometryCollection OutGeometryCollection;
-			FPVMeshBuilder::GenerateGeometryCollection(SkeletonCollection, Settings->MesherSettings, OutGeometryCollection);
+			FPVMeshBuilder::GenerateGeometryCollection(SkeletonCollection, PlantProfileCollection, Settings->MesherSettings, OutGeometryCollection);
 
 			OutManagedArrayCollectionData->Initialize(MoveTemp(OutGeometryCollection));
 
diff --git a/Engine/Plugins/Experimental/ProceduralVegetationEditor/Source/ProceduralVegetation/Private/Nodes/PVMeshBuilderSettings.h b/Engine/Plugins/Experimental/ProceduralVegetationEditor/Source/ProceduralVegetation/Private/Nodes/PVMeshBuilderSettings.h
--- a/Engine/Plugins/Experimental/ProceduralVegetationEditor/Source/ProceduralVegetation/Private/Nodes/PVMeshBuilderSettings.h
+++ b/Engine/Plugins/Experimental/ProceduralVegetationEditor/Source/ProceduralVegetation/Private/Nodes/PVMeshBuilderSettings.h
@@ -31,7 +31,7 @@ public:
 
 protected:
 	virtual FPCGElementPtr CreateElement() const override;
-
+	virtual TArray<FPCGPinProperties> InputPinProperties() const override;
 	virtual FPCGDataTypeIdentifier GetInputPinTypeIdentifier() const override;
 	virtual FPCGDataTypeIdentifier GetOutputPinTypeIdentifier() const override;
 
@@ -50,12 +50,6 @@ public:
 
 	UPROPERTY()
 	FString DisplacementWarnings;
-	
-	UFUNCTION()
-	TArray<FString> GetPlantProfileOptions() const
-	{
-		return MesherSettings.PlantProfileOptions;
-	}
 };
 
 class FPVMeshBuilderElement : public FPVBaseElement
diff --git a/Engine/Plugins/Experimental/ProceduralVegetationEditor/Source/ProceduralVegetation/Private/Nodes/PVPlantProfileLoaderSettings.cpp b/Engine/Plugins/Experimental/ProceduralVegetationEditor/Source/ProceduralVegetation/Private/Nodes/PVPlantProfileLoaderSettings.cpp
new file mode 100644
--- /dev/null
+++ b/Engine/Plugins/Experimental/ProceduralVegetationEditor/Source/ProceduralVegetation/Private/Nodes/PVPlantProfileLoaderSettings.cpp
@@ -0,0 +1,129 @@
+// Copyright Epic Games, Inc. All Rights Reserved.
+
+#include "PVPlantProfileLoaderSettings.h"
+#include "PCGContext.h"
+#include "ProceduralVegetationModule.h"
+#include "DataAssets/ProceduralVegetationPlantProfileDataAsset.h"
+#include "DataTypes/PVPlantProfileData.h"
+#include "Facades/PVProfileFacade.h"
+
+#define LOCTEXT_NAMESPACE "PVPlantProfileLoaderSettings"
+
+#if WITH_EDITOR
+FText UPVPlantProfileLoaderSettings::GetDefaultNodeTitle() const
+{ 
+	return LOCTEXT("NodeTitle", "Plant Profile Loader"); 
+}
+
+FText UPVPlantProfileLoaderSettings::GetNodeTooltipText() const
+{	
+	return LOCTEXT("NodeTooltip", 
+		"This allows the user to load a plant profile data"
+		"It contains points data for a shape of the trunk. "
+		"\n\nPress Ctrl + L to lock/unlock node output"
+	);
+}
+
+void UPVPlantProfileLoaderSettings::PostEditChangeProperty(FPropertyChangedEvent& PropertyChangedEvent)
+{
+	Super::PostEditChangeProperty(PropertyChangedEvent);
+	
+	if (PropertyChangedEvent.Property
+		&& PropertyChangedEvent.Property->GetFName() == GET_MEMBER_NAME_CHECKED(UPVPlantProfileLoaderSettings, PlantProfileData))
+	{
+		if (PlantProfileData)
+		{
+			Profiles = PlantProfileData->Profiles;	
+		}
+		else
+		{
+			Profiles.Empty();
+		}
+	}
+}
+#endif
+
+TArray<FPCGPinProperties> UPVPlantProfileLoaderSettings::InputPinProperties() const
+{
+	return TArray<FPCGPinProperties>();
+}
+
+TArray<FPCGPinProperties> UPVPlantProfileLoaderSettings::OutputPinProperties() const
+{
+	TArray<FPCGPinProperties> Properties;
+
+	if (PlantProfileData)
+	{
+		for (const auto& [ProfileName, ProfilePoints] : PlantProfileData->Profiles)
+		{
+			if (!ProfileName.IsEmpty() && ProfilePoints.Num() > 0)
+			{
+				FPCGPinProperties& Pin = Properties.Emplace_GetRef(*ProfileName, GetOutputPinTypeIdentifier());
+				Pin.bAllowMultipleData = false;
+			}
+		}
+	}
+
+	return Properties;
+}
+
+FPCGDataTypeIdentifier UPVPlantProfileLoaderSettings::GetOutputPinTypeIdentifier() const
+{
+	return FPCGDataTypeIdentifier{ FPVDataTypeInfoPlantProfile::AsId() };
+}
+
+FPCGElementPtr UPVPlantProfileLoaderSettings::CreateElement() const
+{
+	return MakeShared<FPVPlantProfileLoaderElement>();
+}
+
+bool FPVPlantProfileLoaderElement::ExecuteInternal(FPCGContext* InContext) const
+{
+	TRACE_CPUPROFILER_EVENT_SCOPE(FPVPresetLoaderElement::ExecuteInternal);
+
+	check(InContext);
+
+	const UPVPlantProfileLoaderSettings* Settings = InContext->GetInputSettings<UPVPlantProfileLoaderSettings>();
+	check(Settings);
+
+	if (!Settings->PlantProfileData)
+	{
+		PCGLog::LogWarningOnGraph(NSLOCTEXT("PVPlantProfileLoaderSettings", "MissingDataAsset", "Preset data asset is not set"), InContext);
+		return true;
+	}
+
+	bool bValidDataExist = false;
+	for (const auto& [ProfileName, ProfilePoints] : Settings->PlantProfileData->Profiles)
+	{
+		if (!ProfileName.IsEmpty() && ProfilePoints.Num() > 0)
+		{
+			FManagedArrayCollection Collection;
+
+			PV::Facades::FPlantProfileFacade Facade(Collection);
+
+			Facade.AddProfileEntry(ProfilePoints);
+			
+			UPVPlantProfileData* OutVariantData = FPCGContext::NewObject_AnyThread<UPVPlantProfileData>(InContext);
+			OutVariantData->Initialize(CopyTemp(Collection));
+			
+			FPCGTaggedData& CollectionOutput = InContext->OutputData.TaggedData.Emplace_GetRef();
+			CollectionOutput.Data = OutVariantData;
+			CollectionOutput.Pin = *ProfileName;
+			bValidDataExist = true;
+		}
+		else
+		{
+			PCGLog::LogWarningOnGraph(FText::FromString(FString::Format(TEXT("Skipping variation {0} due to invalid data."), {*ProfileName})), InContext);
+		}
+	}
+
+	if (!bValidDataExist)
+	{
+		PCGLog::LogErrorOnGraph(FText::FromString(FString::Format(TEXT("No valid variation found in data asset {0}"), { Settings->PlantProfileData.GetFName().ToString()})), InContext);
+		return true;
+	}
+	
+	return true;
+}
+
+#undef LOCTEXT_NAMESPACE
\ No newline at end of file
diff --git a/Engine/Plugins/Experimental/ProceduralVegetationEditor/Source/ProceduralVegetation/Private/Nodes/PVPlantProfileLoaderSettings.h b/Engine/Plugins/Experimental/ProceduralVegetationEditor/Source/ProceduralVegetation/Private/Nodes/PVPlantProfileLoaderSettings.h
new file mode 100644
--- /dev/null
+++ b/Engine/Plugins/Experimental/ProceduralVegetationEditor/Source/ProceduralVegetation/Private/Nodes/PVPlantProfileLoaderSettings.h
@@ -0,0 +1,45 @@
+// Copyright Epic Games, Inc. All Rights Reserved.
+#pragma once
+
+#include "CoreMinimal.h"
+#include "Nodes/PVBaseSettings.h"
+#include "PVPlantProfileLoaderSettings.generated.h"
+
+struct FPVPlantProfile;
+
+UCLASS(BlueprintType, ClassGroup = (Procedural))
+class UPVPlantProfileLoaderSettings : public UPVBaseSettings
+{
+	GENERATED_BODY()
+
+public:
+	
+#if WITH_EDITOR
+	virtual FName GetDefaultNodeName() const override { return FName(TEXT("ProceduralVegetationPlantProfileLoader")); }
+	virtual FText GetDefaultNodeTitle() const override;
+	virtual FText GetNodeTooltipText() const override;
+	virtual EPCGSettingsType GetType() const override { return EPCGSettingsType::InputOutput; }
+	virtual TArray<EPVRenderType> GetDefaultRenderType() const override { return TArray{ EPVRenderType::PointData }; }
+
+	virtual void PostEditChangeProperty(FPropertyChangedEvent& PropertyChangedEvent) override;
+#endif
+
+protected:
+	virtual TArray<FPCGPinProperties> InputPinProperties() const override;
+	virtual TArray<FPCGPinProperties> OutputPinProperties() const override;
+	virtual FPCGDataTypeIdentifier GetOutputPinTypeIdentifier() const override;
+	virtual FPCGElementPtr CreateElement() const override;
+
+public:
+	UPROPERTY(EditAnywhere, Category = Settings, meta=(Tooltip="Loads the selected procedural vegetation plant profile data\n\n. A procedural vegatation profile data contains points data for a shape of the trunk."))
+	TObjectPtr<class UProceduralVegetationPlantProfileDataAsset> PlantProfileData = nullptr;
+
+	UPROPERTY(VisibleAnywhere, Category = Settings)
+	TArray<FPVPlantProfile> Profiles;
+};
+
+class FPVPlantProfileLoaderElement : public FPVBaseElement
+{
+protected:
+	virtual bool ExecuteInternal(FPCGContext* InContext) const override;
+};
diff --git a/Engine/Plugins/Experimental/ProceduralVegetationEditor/Source/ProceduralVegetation/Private/PVCommon.h b/Engine/Plugins/Experimental/ProceduralVegetationEditor/Source/ProceduralVegetation/Private/PVCommon.h
--- a/Engine/Plugins/Experimental/ProceduralVegetationEditor/Source/ProceduralVegetation/Private/PVCommon.h
+++ b/Engine/Plugins/Experimental/ProceduralVegetationEditor/Source/ProceduralVegetation/Private/PVCommon.h
@@ -18,4 +18,5 @@ namespace PV::Pins
 	const FName GrowerRandomAngleInputLabel = TEXT("RandomAngle");
 	
 	const FName FoliageDistributorFoliageInputLabel = TEXT("Foliage");
+	const FName MeshBuilderPlantProfileInputLabel = TEXT("PlantProfile");
 }
\ No newline at end of file
diff --git a/Engine/Plugins/Experimental/ProceduralVegetationEditor/Source/ProceduralVegetation/Private/ProceduralVegetationPreset.cpp b/Engine/Plugins/Experimental/ProceduralVegetationEditor/Source/ProceduralVegetation/Private/ProceduralVegetationPreset.cpp
--- a/Engine/Plugins/Experimental/ProceduralVegetationEditor/Source/ProceduralVegetation/Private/ProceduralVegetationPreset.cpp
+++ b/Engine/Plugins/Experimental/ProceduralVegetationEditor/Source/ProceduralVegetation/Private/ProceduralVegetationPreset.cpp
@@ -5,6 +5,7 @@
 #include "Misc/PackageName.h"
 #include "UObject/Package.h"
 #include "AssetRegistry/IAssetRegistry.h"
+#include "DataAssets/ProceduralVegetationPlantProfileDataAsset.h"
 #include "Facades/PVBranchFacade.h"
 #include "Facades/PVProfileFacade.h"
 #include "HAL/FileManager.h"
@@ -214,18 +215,21 @@ void UProceduralVegetationPreset::CreateProfileDataAsset()
 	FString PackagePath = FPaths::Combine(FPackageName::GetLongPackagePath(GetPackage()->GetName()), ProfilePackageName);
 	UPackage* PlantProfilePackage = CreatePackage(*PackagePath);
 	
-	UPlantProfileAsset* NewPlantProfileAsset = NewObject<UPlantProfileAsset>(PlantProfilePackage, *ProfilePackageName, RF_Public | RF_Standalone);
- 
+	UProceduralVegetationPlantProfileDataAsset* NewPlantProfileAsset = NewObject<UProceduralVegetationPlantProfileDataAsset>(PlantProfilePackage, *ProfilePackageName, RF_Public | RF_Standalone);
+
+	int ProfileIndex = 0;
 	for (auto& Pair : Variants)
 	{
 		PV::Facades::FPlantProfileFacade Facade = PV::Facades::FPlantProfileFacade(Pair.Value);
 		for (int32 Index = 0; Index < Facade.NumProfileEntries(); Index++)
 		{
-			FPlantProfile Profile;
-			Profile.ProfileName = FString::Format(TEXT("Profile_{0}"), {Index});
+			FPVPlantProfile Profile;
+			Profile.Name = FString::Format(TEXT("Profile_{0}"), {Index});
 			auto Points = Facade.GetProfilePoints(Index);
-			Profile.ProfilePoints = Points;
+			Profile.Points = Points;
 			NewPlantProfileAsset->Profiles.Add(Profile);
+
+			ProfileIndex++;
 		}
 	}
 	
diff --git a/Engine/Plugins/Experimental/ProceduralVegetationEditor/Source/ProceduralVegetation/Public/DataTypes/PVPlantProfileData.h b/Engine/Plugins/Experimental/ProceduralVegetationEditor/Source/ProceduralVegetation/Public/DataTypes/PVPlantProfileData.h
new file mode 100644
--- /dev/null
+++ b/Engine/Plugins/Experimental/ProceduralVegetationEditor/Source/ProceduralVegetation/Public/DataTypes/PVPlantProfileData.h
@@ -0,0 +1,30 @@
+// Copyright Epic Games, Inc. All Rights Reserved.
+
+#pragma once
+
+#include "PVData.h"
+#include "PVPlantProfileData.generated.h"
+
+USTRUCT()
+struct FPVDataTypeInfoPlantProfile : public FPCGDataTypeInfo
+{
+	GENERATED_BODY()
+
+	PCG_DECLARE_TYPE_INFO(PROCEDURALVEGETATION_API);
+
+#if WITH_EDITOR
+	// If the type should be hidden to the user.
+	virtual bool Hidden() const override { return true; };
+#endif // WITH_EDITOR
+};
+
+UCLASS(BlueprintType, ClassGroup = (Procedural))
+class UPVPlantProfileData : public UPVData
+{
+	GENERATED_BODY()
+
+public:
+	//~ Begin UPCGData interface
+	PCG_ASSIGN_DEFAULT_TYPE_INFO(FPVDataTypeInfoPlantProfile)
+	//~ End UPCGData interface
+};
\ No newline at end of file
diff --git a/Engine/Plugins/Experimental/ProceduralVegetationEditor/Source/ProceduralVegetation/Public/ProceduralVegetationPreset.h b/Engine/Plugins/Experimental/ProceduralVegetationEditor/Source/ProceduralVegetation/Public/ProceduralVegetationPreset.h
--- a/Engine/Plugins/Experimental/ProceduralVegetationEditor/Source/ProceduralVegetation/Public/ProceduralVegetationPreset.h
+++ b/Engine/Plugins/Experimental/ProceduralVegetationEditor/Source/ProceduralVegetation/Public/ProceduralVegetationPreset.h
@@ -76,26 +76,5 @@ public:
 	static void ShowHideInternalProperties(bool bDebugEnable);
 private:
 	void CreateProfileDataAsset();
-	
 	void UpdateFoliageAndMaterialPath();
-};
-
-USTRUCT(BlueprintType)
-struct PROCEDURALVEGETATION_API FPlantProfile
-{
-	GENERATED_BODY()
-public:
-	UPROPERTY(VisibleAnywhere, Category="PlantProfile")
-	FString ProfileName;
-	UPROPERTY(VisibleAnywhere, Category="PlantProfile")
-	TArray<float> ProfilePoints;
-};
-
-UCLASS(BlueprintType)
-class PROCEDURALVEGETATION_API UPlantProfileAsset : public UDataAsset
-{
-	GENERATED_BODY()
-public:
-	UPROPERTY(VisibleAnywhere, Category="PlantProfile")
-	TArray<FPlantProfile> Profiles;
 };
\ No newline at end of file
