[PVE] Fixed localization warning

#Synced-CL 51211376

diff --git a/Engine/Plugins/Experimental/ProceduralVegetationEditor/Source/ProceduralVegetation/Private/Nodes/PVGrowthDataLoader.cpp b/Engine/Plugins/Experimental/ProceduralVegetationEditor/Source/ProceduralVegetation/Private/Nodes/PVGrowthDataLoader.cpp
--- a/Engine/Plugins/Experimental/ProceduralVegetationEditor/Source/ProceduralVegetation/Private/Nodes/PVGrowthDataLoader.cpp
+++ b/Engine/Plugins/Experimental/ProceduralVegetationEditor/Source/ProceduralVegetation/Private/Nodes/PVGrowthDataLoader.cpp
@@ -7,7 +7,7 @@
 #include "DataTypes/PVGrowthData.h"
 #include "Facades/PVPointFacade.h"
 
-#define LOCTEXT_NAMESPACE "PVPresetLoaderSettings"
+#define LOCTEXT_NAMESPACE "PVGrowthDataLoaderSettings"
 
 #if WITH_EDITOR
 FText UPVGrowthDataLoaderSettings::GetDefaultNodeTitle() const
