[PVE] Fixed crash on selecting embedded graph node.
Fix/How: It was crashing because we were always trying to find the edGraph node in the Main graph. The fix is to use a focused graph instead.

#Synced-CL 51211376
#Synced-CL 51211376

diff --git a/Engine/Plugins/Experimental/ProceduralVegetationEditor/Source/ProceduralVegetationEditor/Private/PVEditor.cpp b/Engine/Plugins/Experimental/ProceduralVegetationEditor/Source/ProceduralVegetationEditor/Private/PVEditor.cpp
--- a/Engine/Plugins/Experimental/ProceduralVegetationEditor/Source/ProceduralVegetationEditor/Private/PVEditor.cpp
+++ b/Engine/Plugins/Experimental/ProceduralVegetationEditor/Source/ProceduralVegetationEditor/Private/PVEditor.cpp
@@ -449,7 +449,10 @@ void FPVEditor::ChangeNodeInspection(UPCGNode* InNode)
 
 		NodeBeingInspected = InNode;
 
-		SetNodeInspected(GetEditorNodeFromPCGNode(NodeBeingInspected), true);
+		if(UPCGEditorGraphNodeBase* EditorNode = GetEditorNodeFromPCGNode(NodeBeingInspected))
+		{
+			SetNodeInspected(EditorNode, true);	
+		}
 	}
 }
 
@@ -482,17 +485,22 @@ UPCGNode* FPVEditor::GetFirstSelectedNode()
 
 UPCGEditorGraphNodeBase* FPVEditor::GetEditorNodeFromPCGNode(const UPCGNode* InPCGNode) const
 {
-	const TObjectPtr<UEdGraphNode>* EdGraphNode = GetMainEditorGraph()->Nodes.FindByPredicate(
-		[&](UEdGraphNode* GraphNode)
-			{
-				const UPCGEditorGraphNodeBase* PCGEdGraphNode = Cast<UPCGEditorGraphNodeBase>(GraphNode);
-				return PCGEdGraphNode
-					? PCGEdGraphNode->GetPCGNode() == InPCGNode
-					: false;
-			});
-	return EdGraphNode
-		? CastChecked<UPCGEditorGraphNodeBase>(*EdGraphNode)
-		: nullptr;
+	if (UPCGEditorGraph* FocusedGraph = GetFocusedEditorGraph())
+	{
+		const TObjectPtr<UEdGraphNode>* EdGraphNode = FocusedGraph->Nodes.FindByPredicate
+		(
+			[&](UEdGraphNode* GraphNode)
+				{
+					const UPCGEditorGraphNodeBase* PCGEdGraphNode = Cast<UPCGEditorGraphNodeBase>(GraphNode);
+					return PCGEdGraphNode
+						? PCGEdGraphNode->GetPCGNode() == InPCGNode
+						: false;
+				});
+		
+		return EdGraphNode ? CastChecked<UPCGEditorGraphNodeBase>(*EdGraphNode) : nullptr;	
+	}
+	
+	return nullptr;
 }
 
 
