Fix BudLateralMeristem attribute not being accuratly computed by PV::AttributesHelper::ComputeBudLateralMeristem

- Set "LateralMeristem"  to PointScale * 0.001
- Correctly compute ParentDot
- Implement better estimate of Davinci value by first estimating the LateralElongation
- Fix bug in PVPlantTraversalHelper which caused traversal of the same branch multiple times
- Fix AddGrowthMissingData not properly initializing LengthFromRoot and PointScale

#Synced-CL 51214504

diff --git a/Engine/Plugins/Experimental/ProceduralVegetationEditor/Source/ProceduralVegetation/Private/Helpers/PVAttributesHelper.cpp b/Engine/Plugins/Experimental/ProceduralVegetationEditor/Source/ProceduralVegetation/Private/Helpers/PVAttributesHelper.cpp
--- a/Engine/Plugins/Experimental/ProceduralVegetationEditor/Source/ProceduralVegetation/Private/Helpers/PVAttributesHelper.cpp
+++ b/Engine/Plugins/Experimental/ProceduralVegetationEditor/Source/ProceduralVegetation/Private/Helpers/PVAttributesHelper.cpp
@@ -29,7 +29,7 @@ void PV::AttributesHelper::ComputeNjordPixelIndex(const FComputeNjordPixelIndexA
 		1, 
 		[&](int32 BranchIndex, int32& NjordPixelIndex)->EForEachResult
 		{
-			const bool bIsTrunk = GetBranchParentIndex(InAttributes.BranchParentNumber, InAttributes.BranchNumber, BranchIndex) == INDEX_NONE;
+			const bool bIsTrunk = IsTrunk(InAttributes.BranchParentNumber, BranchIndex);
 			const TArray<int32>& BranchPoints = InAttributes.BranchPoints[BranchIndex];
 			for (int32 i = 0; i < BranchPoints.Num(); ++i)
 			{
@@ -451,6 +451,29 @@ void PV::AttributesHelper::ComputeBudLateralMeristem(const FComputeBudLateralMer
 {
 	using namespace PV::PlantTraversalHelper;
 
+	for (int32 i = 0; i < InAttributes.BudLateralMeristem.Num(); ++i)
+	{
+		InAttributes.BudLateralMeristem[i].Davinci = 0;
+	}
+
+	const float EstimatedLateralElongation = Invoke([&]() -> float
+	{
+		// The LateralElongation is a user parameter but it can be fairly accuratly estimated by taking the scale of the
+		// second to last point on the tunk (multiplied by 0.01).
+		// NOTE: This may not be accurate for multiplant growth data as the different plants may have different LateralElongation.
+		for (const int32 TrunkIndex : GetTrunkIndices(InAttributes.BranchParentNumber))
+		{
+			const TArray<int32>& BranchPoints = InAttributes.BranchPoints[TrunkIndex];
+			if (ensure(BranchPoints.Num() > 1))
+			{
+				const int32 SecondToLastPointIndex = BranchPoints[BranchPoints.Num() - 2];
+				return InAttributes.PointScale[SecondToLastPointIndex] * 0.01f;
+			}
+			break;
+		}
+		return 0.033f;
+	});
+
 	RecursiveWalkBranches_Reversed(
 		InAttributes.BranchParentNumber,
 		InAttributes.BranchChildren,
@@ -460,28 +483,33 @@ void PV::AttributesHelper::ComputeBudLateralMeristem(const FComputeBudLateralMer
 			const TArray<int32>& BranchPointIndices = InAttributes.BranchPoints[BranchIndex];
 			check(BranchPointIndices.Num() > 1);
 
-			const FVector3f& FirstBranchPointPos = InAttributes.PointPosition[BranchPointIndices[0]];
-			const FVector3f& SecondBranchPointPos = InAttributes.PointPosition[BranchPointIndices[1]];
-
-			const FVector3f BranchDir = (SecondBranchPointPos - FirstBranchPointPos).GetSafeNormal();
-			const FVector3f FirstPointNormal = InAttributes.BudDirection[BranchPointIndices[0]].Apical;;
-			
-			const float ParentDot = FVector3f::DotProduct(BranchDir, FirstPointNormal);
+			const bool bIsTrunk = IsTrunk(InAttributes.BranchParentNumber, BranchIndex);
 
 			float PreviousPointDavinciValue = 0;
 			for (int32 i = BranchPointIndices.Num() - 1; i>= 0; --i)
 			{
 				const int32 PointIndex = BranchPointIndices[i];
 				const float ExisitingDavinciValue = InAttributes.BudLateralMeristem[PointIndex].Davinci;
-				const float PointScale = InAttributes.PointScale[PointIndex];
-				const float AccumulatedDavinciValue = PreviousPointDavinciValue + ExisitingDavinciValue + PointScale;
+				const float AccumulatedDavinciValue = PreviousPointDavinciValue + ExisitingDavinciValue + EstimatedLateralElongation;
 				PreviousPointDavinciValue = AccumulatedDavinciValue;
 
-				InAttributes.BudLateralMeristem[PointIndex].LateralMeristem = 0.0;
+				InAttributes.BudLateralMeristem[PointIndex].LateralMeristem = InAttributes.PointScale[PointIndex] * 0.01f;
 				InAttributes.BudLateralMeristem[PointIndex].Multiplier = 1;
 				InAttributes.BudLateralMeristem[PointIndex].Inactive = 0.0;
 				InAttributes.BudLateralMeristem[PointIndex].Davinci = AccumulatedDavinciValue;
-				InAttributes.BudLateralMeristem[PointIndex].ParentDot = ParentDot;
+
+				if (!bIsTrunk && i == 1)
+				{
+					InAttributes.BudLateralMeristem[PointIndex].ParentDot = FVector3f::DotProduct(
+						InAttributes.BudDirection[BranchPointIndices[0]].Apical, 
+						InAttributes.BudDirection[BranchPointIndices[1]].Apical
+					);
+				}
+				else
+				{
+					InAttributes.BudLateralMeristem[PointIndex].ParentDot = 0;
+				}
+				
 				InAttributes.BudLateralMeristem[PointIndex].RootDistance = InAttributes.PointLengthFromRoot[PointIndex];
 				InAttributes.BudLateralMeristem[PointIndex].Degredation = 0.0;
 			}
diff --git a/Engine/Plugins/Experimental/ProceduralVegetationEditor/Source/ProceduralVegetation/Private/Helpers/PVJSONHelper.cpp b/Engine/Plugins/Experimental/ProceduralVegetationEditor/Source/ProceduralVegetation/Private/Helpers/PVJSONHelper.cpp
--- a/Engine/Plugins/Experimental/ProceduralVegetationEditor/Source/ProceduralVegetation/Private/Helpers/PVJSONHelper.cpp
+++ b/Engine/Plugins/Experimental/ProceduralVegetationEditor/Source/ProceduralVegetation/Private/Helpers/PVJSONHelper.cpp
@@ -523,12 +523,38 @@ namespace PV::JSON
 
 	void AddGrowthMissingData(FManagedArrayCollection& Collection)
 	{
+		if (!PV::FPointScaleAttribute::HasAttribute(Collection))
+		{
+			auto PointScaleAttributes = PV::FPointScaleAttribute::AddAttribute(Collection, 0);
+			auto BudLateralMeristemAttribute = PV::FBudLateralMeristemAttribute::FindAttribute(Collection);
+			if (BudLateralMeristemAttribute)
+			{
+				for (int32 i = 0; i < BudLateralMeristemAttribute.Num(); ++i)
+				{
+					PointScaleAttributes[i] = BudLateralMeristemAttribute[i].LateralMeristem * 100.f;
+				}
+			}
+		}
+
 		if (!PV::FPointLengthFromRootAttribute::HasAttribute(Collection)
 			|| !PV::FPointLengthFromSeedAttribute::HasAttribute(Collection))
 		{
-			PV::FPointLengthFromRootAttribute::AddAttribute(Collection);
-			PV::FPointLengthFromSeedAttribute::AddAttribute(Collection);
-			PV::AttributesHelper::ComputeLengthFromRoot(Collection);
+			auto PointLengthFromRootAttribute = PV::FPointLengthFromRootAttribute::AddAttribute(Collection);
+			auto PointLengthFromSeedAttribute = PV::FPointLengthFromSeedAttribute::AddAttribute(Collection);
+
+			auto BudLateralMeristemAttribute = PV::FBudLateralMeristemAttribute::FindAttribute(Collection);
+			if (BudLateralMeristemAttribute)
+			{
+				for (int32 i = 0; i < BudLateralMeristemAttribute.Num(); ++i)
+				{
+					PointLengthFromRootAttribute[i] = BudLateralMeristemAttribute[i].RootDistance;
+					PointLengthFromSeedAttribute[i] = BudLateralMeristemAttribute[i].RootDistance;
+				}
+			}
+			else
+			{
+				PV::AttributesHelper::ComputeLengthFromRoot(Collection);
+			}
 		}
 
 		if (!PV::FPointGroundGradientAttribute::HasAttribute(Collection))
diff --git a/Engine/Plugins/Experimental/ProceduralVegetationEditor/Source/ProceduralVegetation/Private/Helpers/PVPlantTraversalHelper.h b/Engine/Plugins/Experimental/ProceduralVegetationEditor/Source/ProceduralVegetation/Private/Helpers/PVPlantTraversalHelper.h
--- a/Engine/Plugins/Experimental/ProceduralVegetationEditor/Source/ProceduralVegetation/Private/Helpers/PVPlantTraversalHelper.h
+++ b/Engine/Plugins/Experimental/ProceduralVegetationEditor/Source/ProceduralVegetation/Private/Helpers/PVPlantTraversalHelper.h
@@ -34,12 +34,17 @@ namespace PV::PlantTraversalHelper
 		return OutBranchIndices;
 	}
 
+	inline bool IsTrunk(FBranchParentNumberAttributeConstView BranchParentNumberAttribute, int32 BranchIndex)
+	{
+		return BranchParentNumberAttribute.IsValidIndex(BranchIndex) && BranchParentNumberAttribute[BranchIndex] == 0;
+	}
+
 	inline TArray<int32> GetTrunkIndices(FBranchParentNumberAttributeConstView BranchParentNumberAttribute)
 	{
 		TArray<int32> TrunkIndicesResult;
 		for (int32 i = 0; i < BranchParentNumberAttribute.Num(); i++)
 		{
-			if (BranchParentNumberAttribute[i] == 0)
+			if (IsTrunk(BranchParentNumberAttribute, i))
 			{
 				TrunkIndicesResult.Add(i);
 			}
@@ -106,11 +111,16 @@ namespace PV::PlantTraversalHelper
 			}
 
 			const TArray<int32>& BranchChildren = BranchChildrenAttribute[BranchIndex];
-			for (int32 BranchChild : BranchChildren)
+			for (int32 ChildBranchNumber : BranchChildren)
 			{
-				const int32 ChildBranchIndex = BranchNumberAttribute.Find(BranchChild);
+				const int32 ChildBranchIndex = BranchNumberAttribute.Find(ChildBranchNumber);
 				check(ChildBranchIndex != INDEX_NONE);
 
+				if (BranchParentNumberAttribute[ChildBranchIndex] != BranchNumberAttribute[BranchIndex])
+				{
+					continue;
+				}
+
 				if (WalkBranches_Recursive(ChildBranchIndex, BranchState) == EForEachResult::Break)
 				{
 					return EForEachResult::Break;
@@ -189,27 +199,27 @@ namespace PV::PlantTraversalHelper
 		PV::FBranchFoliageIDsAttributeConstView BranchFoliageIDsAttribute,
 		PV::FFoliageLengthFromRootAttributeConstView FoliageLengthFromRootAttribute,
 		int32 BranchIndex,
-		int32 PointIndex)
+		int32 BranchPointIndex)
 	{
-		const TArray<int32> BranchPoints = BranchPointsAttribute[BranchIndex];
-		const int32 PointBranchIndex = BranchPoints.Find(PointIndex);
-		if (PointBranchIndex == INDEX_NONE)
+		const TArray<int32>& BranchPoints = BranchPointsAttribute[BranchIndex];
+		if (!BranchPoints.IsValidIndex(BranchPointIndex))
 		{
 			return TArray<int32>();
 		}
 
-		const int32 PreviousPointBranchIndex = PointBranchIndex > 0 ? PointBranchIndex - 1 : INDEX_NONE;
-		if (PreviousPointBranchIndex == INDEX_NONE)
+		const int32 PreviousBranchPointIndex = BranchPointIndex > 0 ? BranchPointIndex - 1 : INDEX_NONE;
+		if (PreviousBranchPointIndex == INDEX_NONE)
 		{
 			return TArray<int32>();
 		}
 
-		const int32 PreviousPointIndex = BranchPoints[PreviousPointBranchIndex];
+		const int32 PointIndex = BranchPoints[BranchPointIndex];
+		const int32 PreviousPointIndex = BranchPoints[PreviousBranchPointIndex];
 
 		const float BranchPointLengthFromRoot = PointLengthFromRootAttribute[PointIndex];
 		const float PreviousBranchPointLengthFromRoot = PointLengthFromRootAttribute[PreviousPointIndex];
 
-		const TArray<int32> FoliageEntryIds = BranchFoliageIDsAttribute[BranchIndex];
+		const TArray<int32>& FoliageEntryIds = BranchFoliageIDsAttribute[BranchIndex];
 
 		TArray<int32> OutFoliageEntryIds;
 		OutFoliageEntryIds.Reserve(FoliageEntryIds.Num());
@@ -744,8 +754,7 @@ namespace PV::PlantTraversalHelper
 		const bool bHasFusedPoints = HasFusedPoints(BranchPointsAttribute, BranchParentNumberAttribute, BranchNumberAttribute);
 		for (int32 BranchIndex = 0; BranchIndex < BranchPointsAttribute.Num(); ++BranchIndex)
 		{
-			const bool bSkipFirstPoint = bHasFusedPoints 
-				&& GetBranchParentIndex(BranchParentNumberAttribute, BranchNumberAttribute, BranchIndex) != INDEX_NONE;
+			const bool bSkipFirstPoint = bHasFusedPoints && !IsTrunk(BranchParentNumberAttribute, BranchIndex);
 
 			const int32 StartBranchPointIndex = bSkipFirstPoint ? 1 : 0;
 			const TArray<int32>& PointIndices = BranchPointsAttribute[BranchIndex];
