[PVE] Fixed localization warning

#Synced-CL 51291474

diff --git a/Engine/Plugins/Experimental/ProceduralVegetationEditor/Source/ProceduralVegetation/Private/Nodes/PVGrowthDataLoader.cpp b/Engine/Plugins/Experimental/ProceduralVegetationEditor/Source/ProceduralVegetation/Private/Nodes/PVGrowthDataLoader.cpp
--- a/Engine/Plugins/Experimental/ProceduralVegetationEditor/Source/ProceduralVegetation/Private/Nodes/PVGrowthDataLoader.cpp
+++ b/Engine/Plugins/Experimental/ProceduralVegetationEditor/Source/ProceduralVegetation/Private/Nodes/PVGrowthDataLoader.cpp
@@ -110,7 +110,7 @@ bool FPVGrowthDataLoaderElement::ExecuteInternal(FPCGContext* InContext) const
 
 	if (!Settings->GrowthAsset)
 	{
-		PCGLog::LogWarningOnGraph(NSLOCTEXT("PVPresetLoaderSettings", "MissingDataAsset", "Growth asset is not set"), InContext);
+		PCGLog::LogWarningOnGraph(NSLOCTEXT("PVGrowthDataLoaderSettings", "MissingDataAsset", "Growth asset is not set"), InContext);
 		return true;
 	}
 
