[PVE] Fix for "Build Engine Localization: Localization warnings in ue5-main"

slack-thread: https://epicgames.slack.com/archives/C042RSGQVPD/p1771539966820979

#Synced-CL 51007179

diff --git a/Engine/Plugins/Experimental/ProceduralVegetationEditor/Source/ProceduralVegetationEditor/Private/Widgets/SPVExportSelectionDialog.cpp b/Engine/Plugins/Experimental/ProceduralVegetationEditor/Source/ProceduralVegetationEditor/Private/Widgets/SPVExportSelectionDialog.cpp
--- a/Engine/Plugins/Experimental/ProceduralVegetationEditor/Source/ProceduralVegetationEditor/Private/Widgets/SPVExportSelectionDialog.cpp
+++ b/Engine/Plugins/Experimental/ProceduralVegetationEditor/Source/ProceduralVegetationEditor/Private/Widgets/SPVExportSelectionDialog.cpp
@@ -115,7 +115,7 @@ void SPVExportSelectionDialog::Construct(const FArguments& InArgs)
 						+ SSegmentedControl<EPVExportType>::Slot(EPVExportType::Batch)
 						.HAlign(HAlign_Center)
 						.VAlign(VAlign_Center)
-						.ToolTip(LOCTEXT("PVExportTypeAllLabel", "Batch export all output nodes"))
+						.ToolTip(LOCTEXT("PVExportTypeBatchLabel", "Batch export all output nodes"))
 						[
 							SNew(STextBlock)
 							.Text(LOCTEXT("PVBatchExportLabel", "Batch"))
@@ -123,10 +123,10 @@ void SPVExportSelectionDialog::Construct(const FArguments& InArgs)
 						+ SSegmentedControl<EPVExportType>::Slot(EPVExportType::Selected)
 						.HAlign(HAlign_Center)
 						.VAlign(VAlign_Center)
-						.ToolTip(LOCTEXT("PVExportTypeAllLabel", "Export only selected output nodes in the graph"))
+						.ToolTip(LOCTEXT("PVExportTypeSelectedLabel", "Export only selected output nodes in the graph"))
 						[
 							SNew(STextBlock)
-							.Text(LOCTEXT("PVSelectionExportLabel", "Selected"))
+							.Text(LOCTEXT("PVSelectedExportLabel", "Selected"))
 						]
 					]
 					+ SVerticalBox::Slot()
