[PVE] Change NodeBeingInspected from type UPCGNode to UPCGEditorGraphNodeBase, to avoid searching for the EdGraph Node from PCGNode

#Synced-CL 51007179

diff --git a/Engine/Plugins/Experimental/ProceduralVegetationEditor/Source/ProceduralVegetationEditor/Private/PVEditor.cpp b/Engine/Plugins/Experimental/ProceduralVegetationEditor/Source/ProceduralVegetationEditor/Private/PVEditor.cpp
--- a/Engine/Plugins/Experimental/ProceduralVegetationEditor/Source/ProceduralVegetationEditor/Private/PVEditor.cpp
+++ b/Engine/Plugins/Experimental/ProceduralVegetationEditor/Source/ProceduralVegetationEditor/Private/PVEditor.cpp
@@ -243,14 +243,14 @@ void FPVEditor::OnSelectedNodesChanged(const TSet<UObject*>& InNewSelection)
 	{
 		if (UPCGEditorGraphNodeBase* EdNode = Cast<UPCGEditorGraphNodeBase>(Selection))
 		{
-			SelectedNodes.Add(EdNode->GetPCGNode());
+			SelectedNodes.Add(EdNode);
 		}
 	}
 
-	const UPVBaseSettings* InspectionSettings = GetNodeSettings(NodeBeingInspected);
+	const UPVBaseSettings* InspectionSettings = GetNodeSettings(GetNodeBeingInspected());
 	if (!InspectionSettings || !InspectionSettings->bLockInspection)
 	{
-		ChangeNodeInspection(GetFirstSelectedNode());
+		ChangeNodeInspection(GetFirstSelectedEdNode());
 	}
 
 	// TODO: Hacky way to execute embedded subgraphs
@@ -337,7 +337,10 @@ void FPVEditor::RegisterExtraTabFactories(FWorkflowAllowedTabSet& TabSet)
 	if (CollectionSpreadSheetWidget.IsValid())
 	{
 		FPCGGenericTabFactoryParams Params(CollectionSpreadSheetTabId, SharedThis<FPCGEditor>(this),
-			[this]() { return CollectionSpreadSheetWidget.ToSharedRef(); });
+			[this]()
+				{
+					return CollectionSpreadSheetWidget.ToSharedRef();
+				});
 		Params.Label = LOCTEXT("CollectionSpreadsheetTab", "Collection Spreadsheet");
 		TabSet.RegisterFactory(MakeShared<FPCGGenericTabFactory>(Params));
 	}
@@ -438,31 +441,30 @@ void FPVEditor::OnNodeInspected(const FPCGStack& InPCGStack)
 	UpdateStats();
 }
 
-void FPVEditor::ChangeNodeInspection(UPCGNode* InNode)
+void FPVEditor::ChangeNodeInspection(UPCGEditorGraphNodeBase* InNode)
 {
-	if (UPVBaseSettings* Settings = GetNodeSettings(InNode))
+	if (!InNode)
+	{
+		return;
+	}
+	
+	if (UPVBaseSettings* Settings = GetNodeSettings(InNode->GetPCGNode()))
 	{
 		if (EditorViewport.IsValid())
 		{
 			EditorViewport->OnNodeInspectionChanged(Settings);
 		}
 
-		NodeBeingInspected = InNode;
-
-		if(UPCGEditorGraphNodeBase* EditorNode = GetEditorNodeFromPCGNode(NodeBeingInspected))
-		{
-			SetNodeInspected(EditorNode, true);	
-		}
+		SetNodeBeingInspected(InNode);
 	}
 }
 
 void FPVEditor::OnLockNodeSelection()
 {
-	const auto SelectedNode = GetFirstSelectedNode();
+	UPCGEditorGraphNodeBase* const SelectedNode = GetFirstSelectedEdNode();
+	UPVBaseSettings* SelectedNodeSettings = GetNodeSettings(SelectedNode->GetPCGNode());
 
-	UPVBaseSettings* SelectedNodeSettings = GetNodeSettings(SelectedNode);
-
-	if (UPVBaseSettings* NodeBeingInspectedSettings = GetNodeSettings(NodeBeingInspected))
+	if (UPVBaseSettings* NodeBeingInspectedSettings = GetNodeSettings(GetNodeBeingInspected()))
 	{
 		NodeBeingInspectedSettings->bLockInspection = !NodeBeingInspectedSettings->bLockInspection;
 	}
@@ -470,45 +472,56 @@ void FPVEditor::OnLockNodeSelection()
 	if (NodeBeingInspected != SelectedNode)
 	{
 		SelectedNodeSettings->bLockInspection = true;
-		ChangeNodeInspection(GetFirstSelectedNode());
+		ChangeNodeInspection(SelectedNode);
 	}
 
 	UpdateStats();
 }
 
-UPCGNode* FPVEditor::GetFirstSelectedNode()
+TObjectPtr<UPCGNode> FPVEditor::GetFirstSelectedNode()
+{
+	if (const TObjectPtr<UPCGEditorGraphNodeBase> EdNode = GetFirstSelectedEdNode())
+	{
+		return EdNode->GetPCGNode();
+	}
+	return nullptr;
+}
+
+TObjectPtr<UPCGEditorGraphNodeBase> FPVEditor::GetFirstSelectedEdNode()
 {
 	return SelectedNodes.Num()
 		? SelectedNodes[0]
 		: nullptr;
 }
 
-UPCGEditorGraphNodeBase* FPVEditor::GetEditorNodeFromPCGNode(const UPCGNode* InPCGNode) const
+bool FPVEditor::IsNodeSelected(TObjectPtr<UPCGNode> InNode)
 {
-	if (UPCGEditorGraph* FocusedGraph = GetFocusedEditorGraph())
-	{
-		const TObjectPtr<UEdGraphNode>* EdGraphNode = FocusedGraph->Nodes.FindByPredicate
-		(
-			[&](UEdGraphNode* GraphNode)
-				{
-					const UPCGEditorGraphNodeBase* PCGEdGraphNode = Cast<UPCGEditorGraphNodeBase>(GraphNode);
-					return PCGEdGraphNode
-						? PCGEdGraphNode->GetPCGNode() == InPCGNode
-						: false;
-				});
-		
-		return EdGraphNode ? CastChecked<UPCGEditorGraphNodeBase>(*EdGraphNode) : nullptr;	
-	}
-	
-	return nullptr;
+	return SelectedNodes.ContainsByPredicate(
+		[&](const TObjectPtr<UPCGEditorGraphNodeBase>& Node)
+			{
+				return Node->GetPCGNode() == InNode;
+			}
+	);
 }
 
+TObjectPtr<UPCGNode> FPVEditor::GetNodeBeingInspected()
+{
+	return NodeBeingInspected
+		? NodeBeingInspected->GetPCGNode()
+		: nullptr;
+}
+
+void FPVEditor::SetNodeBeingInspected(const TObjectPtr<UPCGEditorGraphNodeBase> InNode)
+{
+	NodeBeingInspected = InNode;
+	SetNodeInspected(NodeBeingInspected, true);
+}
 
 void FPVEditor::UpdateInspectedCollection()
 {
 	if (NodeBeingInspected)
 	{
-		if (const UPVData* const Data = GetDataFromPin(GetOutPinFromNode(NodeBeingInspected)))
+		if (const UPVData* const Data = GetDataFromPin(GetOutPinFromNode(GetNodeBeingInspected())))
 		{
 			CollectionBeingInspected = Data->GetSharedCollection();
 		}
@@ -520,8 +533,10 @@ void FPVEditor::UpdateInspectedCollection()
 void FPVEditor::UpdateCollectionSpreadSheet()
 {
 	if (!CollectionSpreadSheetWidget)
+	{
 		return;
-	
+	}
+
 	if (!NodeBeingInspected)
 	{
 		CollectionSpreadSheetWidget->SetData(FString());
@@ -542,9 +557,10 @@ void FPVEditor::UpdateCollectionSpreadSheet()
 
 	if (bIsCollectionSpreadSheetVisible)
 	{
-		const TArray<TObjectPtr<UPCGPin>> AllPins = GetAllPVPinsFromNode(NodeBeingInspected);
+		const TObjectPtr<UPCGNode> InspectedNode = GetNodeBeingInspected();
+		const TArray<TObjectPtr<UPCGPin>> AllPins = GetAllPVPinsFromNode(InspectedNode);
 
-		CollectionSpreadSheetWidget->SetData(NodeBeingInspected->GetNodeTitle(EPCGNodeTitleType::FullTitle).ToString());
+		CollectionSpreadSheetWidget->SetData(InspectedNode->GetNodeTitle(EPCGNodeTitleType::FullTitle).ToString());
 
 		TMap<FString, SCollectionInfo>& CollectionInfoMap = CollectionSpreadSheetWidget->GetCollectionTable()->GetCollectionInfoMap();
 		CollectionInfoMap.Empty();
@@ -558,7 +574,7 @@ void FPVEditor::UpdateCollectionSpreadSheet()
 				);
 			}
 		}
-		
+
 		// TODO: Check if there is a way to select default output
 		// if (GetOutPinFromNode(NodeBeingInspected))
 		// {
@@ -624,11 +640,12 @@ void FPVEditor::UpdateStats()
 	GatherStats(OverlayStats);
 	EditorViewport->PopulateStatsOverlayText(OverlayStats);
 
-	if (const UPVBaseSettings* NodeBeingInspectedSettings = GetNodeSettings(NodeBeingInspected))
+	const TObjectPtr<UPCGNode> InspectedNode = GetNodeBeingInspected();
+	if (const UPVBaseSettings* NodeBeingInspectedSettings = GetNodeSettings(InspectedNode))
 	{
 		EditorViewport->SetOverlayText(
 			NodeBeingInspectedSettings->bLockInspection
-			? NodeBeingInspected->GetNodeTitle(EPCGNodeTitleType::FullTitle)
+			? InspectedNode->GetNodeTitle(EPCGNodeTitleType::FullTitle)
 			: FText::GetEmpty()
 		);
 	}
@@ -649,7 +666,7 @@ void FPVEditor::OnExport()
 
 				UPCGPin* const InPin = GetInPinFromNode(Node);
 				UPCGPin* const UpstreamOutPin = GetUpstreamPin(InPin);
-				
+
 				const bool bIsMeshDataPin = IsPinOfType<FPVDataTypeInfoMesh>(UpstreamOutPin);
 				const bool bIsFoliageMeshDataPin = IsPinOfType<FPVDataTypeInfoFoliageMesh>(UpstreamOutPin);
 				if (!bIsMeshDataPin && !bIsFoliageMeshDataPin)
@@ -659,10 +676,11 @@ void FPVEditor::OnExport()
 
 				if (const UPVData* const NodeData = GetDataFromPin(InPin))
 				{
-					const TObjectPtr<UPVExportEntry>& ExportEntry = ExportEntries.Emplace_GetRef(NewObject<UPVExportEntry>(ProceduralVegetationGraph));
-					ExportEntry->Initialize(NodeData->GetSharedCollection(), Node, SelectedNodes.Contains(Node), bIsFoliageMeshDataPin);
+					const TObjectPtr<UPVExportEntry>& ExportEntry = ExportEntries.
+						Emplace_GetRef(NewObject<UPVExportEntry>(ProceduralVegetationGraph));
+					ExportEntry->Initialize(NodeData->GetSharedCollection(), Node, IsNodeSelected(Node), bIsFoliageMeshDataPin);
 				}
-				
+
 				return true;
 			}
 	);
@@ -678,7 +696,7 @@ void FPVEditor::OnExport()
 	}
 
 	UE_LOG(LogProceduralVegetationEditor, Log, TEXT("Export Started"));
-	
+
 	FPVExporter Exporter(ProceduralVegetationBeingEdited, ExportEntries);
 	Exporter.Export();
 }
diff --git a/Engine/Plugins/Experimental/ProceduralVegetationEditor/Source/ProceduralVegetationEditor/Private/PVEditor.h b/Engine/Plugins/Experimental/ProceduralVegetationEditor/Source/ProceduralVegetationEditor/Private/PVEditor.h
--- a/Engine/Plugins/Experimental/ProceduralVegetationEditor/Source/ProceduralVegetationEditor/Private/PVEditor.h
+++ b/Engine/Plugins/Experimental/ProceduralVegetationEditor/Source/ProceduralVegetationEditor/Private/PVEditor.h
@@ -73,11 +73,15 @@ private:
 
 	void OnNodeInspected(const FPCGStack& InPCGStack);
 
-	void ChangeNodeInspection(UPCGNode* InNode);
+	void ChangeNodeInspection(UPCGEditorGraphNodeBase* InNode);
 	void OnLockNodeSelection();
 
-	UPCGNode* GetFirstSelectedNode();
-	UPCGEditorGraphNodeBase* GetEditorNodeFromPCGNode(const UPCGNode* Node) const;
+	TObjectPtr<UPCGNode> GetFirstSelectedNode();
+	TObjectPtr<UPCGEditorGraphNodeBase> GetFirstSelectedEdNode();
+	bool IsNodeSelected(TObjectPtr<UPCGNode> InNode);
+	
+	TObjectPtr<UPCGNode> GetNodeBeingInspected();
+	void SetNodeBeingInspected(TObjectPtr<UPCGEditorGraphNodeBase> InNode);
 
 	void UpdateInspectedCollection();
 
@@ -97,8 +101,8 @@ private:
 
 	TSharedPtr<SPVEditorViewport> EditorViewport = nullptr;
 
-	TObjectPtr<UPCGNode> NodeBeingInspected = nullptr;
-	TArray<TObjectPtr<UPCGNode>> SelectedNodes;
+	TObjectPtr<UPCGEditorGraphNodeBase> NodeBeingInspected = nullptr;
+	TArray<TObjectPtr<UPCGEditorGraphNodeBase>> SelectedNodes;
 
 	TSharedPtr<SCollectionSpreadSheetWidget> CollectionSpreadSheetWidget;
 	TSharedPtr<FManagedArrayCollection> CollectionBeingInspected = nullptr;
