Fix static analysis warning

#Synced-CL 51214504

diff --git a/Engine/Plugins/Experimental/ProceduralVegetationEditor/Source/ProceduralVegetation/Private/Nodes/PVSimplifyPlantSkeletonSettings.cpp b/Engine/Plugins/Experimental/ProceduralVegetationEditor/Source/ProceduralVegetation/Private/Nodes/PVSimplifyPlantSkeletonSettings.cpp
--- a/Engine/Plugins/Experimental/ProceduralVegetationEditor/Source/ProceduralVegetation/Private/Nodes/PVSimplifyPlantSkeletonSettings.cpp
+++ b/Engine/Plugins/Experimental/ProceduralVegetationEditor/Source/ProceduralVegetation/Private/Nodes/PVSimplifyPlantSkeletonSettings.cpp
@@ -54,7 +54,7 @@ bool FPVSimplifyPlantSkeletonElement::ExecuteInternal(FPCGContext* InContext) co
 
 			UPVGrowthData* OutManagedArrayCollectionData = FPCGContext::NewObject_AnyThread<UPVGrowthData>(InContext);
 
-			if (Settings->SimplificationAmount > 0)
+			if (Settings->SimplificationAmount > 0.f)
 			{
 				PV::SimplifyPlantSkeleton(OutCollection, Settings->SimplificationAmount);
 			}
