[PVE] Add new parametric settings to Graft node

- Recomputes branch parents before distributing grafts
- Leaf phyllotaxy settings added to hormone based settings
- Phyllotaxy types are shared with the foliage distributor
- Parametric based distribution settings added to the grafter node
- Adds support for multi-seed graft inputs

#Synced-CL 51286731

diff --git a/Engine/Plugins/Experimental/ProceduralVegetationEditor/Source/ProceduralVegetation/Private/Facades/PVBranchFacade.cpp b/Engine/Plugins/Experimental/ProceduralVegetationEditor/Source/ProceduralVegetation/Private/Facades/PVBranchFacade.cpp
--- a/Engine/Plugins/Experimental/ProceduralVegetationEditor/Source/ProceduralVegetation/Private/Facades/PVBranchFacade.cpp
+++ b/Engine/Plugins/Experimental/ProceduralVegetationEditor/Source/ProceduralVegetation/Private/Facades/PVBranchFacade.cpp
@@ -21,7 +21,8 @@ namespace PV::Facades
 		, TrunkURangeAttribute(InCollection, AttributeNames::TrunkURange, GroupNames::DetailsGroup)
 		, BranchParentNumbers(InCollection, AttributeNames::BranchParentNumber, GroupNames::BranchGroup)
 		, PlantNumbers(InCollection, AttributeNames::PlantNumber, GroupNames::BranchGroup)
-	{}
+	{
+	}
 
 	FBranchFacade::FBranchFacade(const FManagedArrayCollection& InCollection)
 		: Collection(nullptr)
@@ -38,7 +39,8 @@ namespace PV::Facades
 		, TrunkURangeAttribute(InCollection, AttributeNames::TrunkURange, GroupNames::DetailsGroup)
 		, BranchParentNumbers(InCollection, AttributeNames::BranchParentNumber, GroupNames::BranchGroup)
 		, PlantNumbers(InCollection, AttributeNames::PlantNumber, GroupNames::BranchGroup)
-	{}
+	{
+	}
 
 	bool FBranchFacade::IsValid() const
 	{
@@ -71,7 +73,9 @@ namespace PV::Facades
 	int32 FBranchFacade::GetRootPoint(int32 Index) const
 	{
 		const TArray<int32>& Points = GetPoints(Index);
-		return Points.Num() ? Points[0] : INDEX_NONE;
+		return Points.Num()
+			? Points[0]
+			: INDEX_NONE;
 	}
 
 	void FBranchFacade::SetPoints(int32 Index, TArray<int32> InPoints)
@@ -485,7 +489,7 @@ namespace PV::Facades
 		{
 			Parents.Add();
 		}
-		
+
 		return Parents.AddElements(NumElements);
 	}
 
@@ -506,10 +510,10 @@ namespace PV::Facades
 			const int32 BranchParentNumber = BranchParentNumbers[i];
 
 			if (BranchParentNumber != 0)
-				[[likely]]
-				{
-					BranchNumbersToBranchChildren.FindOrAdd(BranchParentNumber).Add(BranchNumber);
-				}
+			[[likely]]
+			{
+				BranchNumbersToBranchChildren.FindOrAdd(BranchParentNumber).Add(BranchNumber);
+			}
 		}
 
 		for (int32 i = 0; i < NumberOfElements; i++)
@@ -521,6 +525,65 @@ namespace PV::Facades
 		}
 	}
 
+	TArray<int32> FBranchFacade::ComputeParentsForBranchNumber(const int32 BranchNumber, const TMap<int32, int32>& BranchNumbersToBranchParentNumbers,
+	                                                           TMap<int32, TArray<int32>>& Memo, TSet<int32>& VisitedParents)
+	{
+		if (const TArray<int32>* Cached = Memo.Find(BranchNumber))
+		{
+			return *Cached;
+		}
+		
+		if (VisitedParents.Contains(BranchNumber))
+		{
+			return TArray<int32>({0});
+		}
+		
+		const int32* BranchParentNumberPtr = BranchNumbersToBranchParentNumbers.Find(BranchNumber);
+		if (!BranchParentNumberPtr)
+		{
+			return TArray<int32>();
+		}
+		
+		const int32 BranchParentNumber = *BranchParentNumberPtr;
+		TArray<int32> Parents;
+		if (BranchParentNumber == 0)
+		{
+			Parents = {0};
+		}
+		else
+		{
+			Parents = ComputeParentsForBranchNumber(BranchParentNumber, BranchNumbersToBranchParentNumbers, Memo, VisitedParents);
+			Parents.Add(BranchParentNumber);
+		}
+
+		Memo.Add(BranchNumber, Parents);
+		return Parents;
+	}
+
+	void FBranchFacade::RecomputeBranchParents()
+	{
+		const int32 NumberOfElements = BranchNumbers.Num();
+		TMap<int32, int32> BranchNumbersToBranchParentNumbers;
+		BranchNumbersToBranchParentNumbers.Reserve(NumberOfElements);
+
+		for (int32 i = 0; i < NumberOfElements; i++)
+		{
+			BranchNumbersToBranchParentNumbers.Add(BranchNumbers[i], BranchParentNumbers[i]);
+		}
+
+		TMap<int32, TArray<int32>> Memo;
+		Memo.Reserve(NumberOfElements);
+
+		for (int32 i = 0; i < NumberOfElements; i++)
+		{
+			const int32 BranchNumber = BranchNumbers[i];
+			TSet<int32> VisitedParents;
+			const TArray<int32> ParentsForBranchNumber = ComputeParentsForBranchNumber(BranchNumber, BranchNumbersToBranchParentNumbers, Memo,
+				VisitedParents);
+			Parents.ModifyAt(i, ParentsForBranchNumber);
+		}
+	}
+
 	void FBranchFacade::FillChildren(const TArray<TArray<int32>>& InputArray)
 	{
 		FILL_COLLECTION_ATTRIBUTE(Children, InputArray);
@@ -562,7 +625,7 @@ namespace PV::Facades
 		{
 			return BranchParentNumbers[Index] == 0;
 		}
-		
+
 		return false;
 	}
 
@@ -615,6 +678,7 @@ namespace PV::Facades
 	{
 		GetValueFromAccessor<int32>(PlantNumbers, OutValue, Index, -1);
 	}
+
 	void FBranchFacade::SetBranchNumbersFromIndex(const TArray<int32>& InputArray, const int32 StartIndex)
 	{
 		FILL_COLLECTION_ATTRIBUTE_FROM_INDEX(BranchNumbers, InputArray, StartIndex);
diff --git a/Engine/Plugins/Experimental/ProceduralVegetationEditor/Source/ProceduralVegetation/Private/Helpers/PVPhyllotaxyHelper.h b/Engine/Plugins/Experimental/ProceduralVegetationEditor/Source/ProceduralVegetation/Private/Helpers/PVPhyllotaxyHelper.h
new file mode 100644
--- /dev/null
+++ b/Engine/Plugins/Experimental/ProceduralVegetationEditor/Source/ProceduralVegetation/Private/Helpers/PVPhyllotaxyHelper.h
@@ -0,0 +1,25 @@
+// Copyright Epic Games, Inc. All Rights Reserved.
+
+#pragma once
+
+#include "PVPhyllotaxyHelper.generated.h"
+
+UENUM(BlueprintType)
+enum class EPhyllotaxyType : uint8
+{
+	Alternate UMETA(DisplayName = "Alternate"),
+	Opposite UMETA(DisplayName = "Opposite"),
+	Decussate UMETA(DisplayName = "Decussate"),
+	Whorled UMETA(DisplayName = "Whorled"),
+	Spiral UMETA(DisplayName = "Spiral")
+};
+ 
+UENUM(BlueprintType)
+enum class EPhyllotaxyFormation : uint8
+{
+	Distichous UMETA(DisplayName = "Distichous"),
+	Tristichous UMETA(DisplayName = "Tristichous"),
+	Pentastichous UMETA(DisplayName = "Pentastichous"),
+	Octastichous UMETA(DisplayName = "Octastichous"),
+	Parastichous UMETA(DisplayName = "Parastichous")
+};
\ No newline at end of file
diff --git a/Engine/Plugins/Experimental/ProceduralVegetationEditor/Source/ProceduralVegetation/Private/Implementations/PVFoliage.cpp b/Engine/Plugins/Experimental/ProceduralVegetationEditor/Source/ProceduralVegetation/Private/Implementations/PVFoliage.cpp
--- a/Engine/Plugins/Experimental/ProceduralVegetationEditor/Source/ProceduralVegetation/Private/Implementations/PVFoliage.cpp
+++ b/Engine/Plugins/Experimental/ProceduralVegetationEditor/Source/ProceduralVegetation/Private/Implementations/PVFoliage.cpp
@@ -11,7 +11,8 @@
 #include "Facades/PVBudVectorsFacade.h"
 #include "Facades/PVFoliageFacade.h"
 #include "Facades/PVPointFacade.h"
- 
+#include "Helpers/PVPhyllotaxyHelper.h"
+
 struct ConditionAttributes
 {
 	TArray<float> UpAlignment;
diff --git a/Engine/Plugins/Experimental/ProceduralVegetationEditor/Source/ProceduralVegetation/Private/Implementations/PVFoliage.h b/Engine/Plugins/Experimental/ProceduralVegetationEditor/Source/ProceduralVegetation/Private/Implementations/PVFoliage.h
--- a/Engine/Plugins/Experimental/ProceduralVegetationEditor/Source/ProceduralVegetation/Private/Implementations/PVFoliage.h
+++ b/Engine/Plugins/Experimental/ProceduralVegetationEditor/Source/ProceduralVegetation/Private/Implementations/PVFoliage.h
@@ -6,31 +6,13 @@
 #include "Misc/EnumRange.h"
  
 #include "GeometryCollection/ManagedArrayCollection.h"
- 
+
 #include "PVFoliage.generated.h"
- 
+
+enum class EPhyllotaxyFormation : uint8;
+enum class EPhyllotaxyType : uint8;
 struct FPVFloatRamp;
  
-UENUM(BlueprintType)
-enum class EPhyllotaxyType : uint8
-{
-	Alternate UMETA(DisplayName = "Alternate"),
-	Opposite UMETA(DisplayName = "Opposite"),
-	Decussate UMETA(DisplayName = "Decussate"),
-	Whorled UMETA(DisplayName = "Whorled"),
-	Spiral UMETA(DisplayName = "Spiral")
-};
- 
-UENUM(BlueprintType)
-enum class EPhyllotaxyFormation : uint8
-{
-	Distichous UMETA(DisplayName = "Distichous"),
-	Tristichous UMETA(DisplayName = "Tristichous"),
-	Pentastichous UMETA(DisplayName = "Pentastichous"),
-	Octastichous UMETA(DisplayName = "Octastichous"),
-	Parastichous UMETA(DisplayName = "Parastichous")
-};
- 
 struct FDistributionSettings
 {
 	const float EthyleneThreshold;
diff --git a/Engine/Plugins/Experimental/ProceduralVegetationEditor/Source/ProceduralVegetation/Private/Implementations/PVGrafter.cpp b/Engine/Plugins/Experimental/ProceduralVegetationEditor/Source/ProceduralVegetation/Private/Implementations/PVGrafter.cpp
--- a/Engine/Plugins/Experimental/ProceduralVegetationEditor/Source/ProceduralVegetation/Private/Implementations/PVGrafter.cpp
+++ b/Engine/Plugins/Experimental/ProceduralVegetationEditor/Source/ProceduralVegetation/Private/Implementations/PVGrafter.cpp
@@ -8,11 +8,13 @@
 
 #include "Facades/PVAttributesNames.h"
 #include "Facades/PVBranchFacade.h"
-#include "Facades/PVBudVectorsFacade.h"
+#include "Facades/PVPlantFacade.h"
 #include "Facades/PVPointFacade.h"
+#include "Facades/PVTreeFacade.h"
 #include "Helpers/PVPlantTraversalHelper.h"
 #include "Helpers/PVFoliageJSONHelper.h"
 #include "Helpers/PVUtilities.h"
+#include "Helpers/PVAttributesHelper.h"
 
 
 template <typename Fn>
@@ -332,7 +334,7 @@ void FPVGrafter::SetGraft(FManagedArrayCollection& OutCollection, const FManaged
 				}
 				PointFacadeOut.SetBudDirections(TargetPointIndex + PointFacadeOutStartIndex, TransformedValues);
 			});
-	
+
 	ForEachGraftPointExcludingRoot(NumOfPointsInSourceGraft, SourceGraftTrunkRootPointIndex,
 		[&](const int32 SourcePointIndex, const int32 TargetPointIndex)
 			{
@@ -602,11 +604,54 @@ bool FPVGrafter::RemoveFoliageDataIfPresent(FManagedArrayCollection& OutCollecti
 	return bRemovedFoliageData;
 }
 
-void FPVGrafter::DistributeGraft(FManagedArrayCollection& OutCollection, const FManagedArrayCollection& SourceSkeletonCollection,
-                                 FManagedArrayCollection& SourceGraftCollection, const FDistributionSettings& DistributionSettings,
-                                 const FScaleSettings& ScaleSettings, const FVectorSettings& VectorSettings, const FMiscSettings& MiscSettings)
+TArray<FManagedArrayCollection> FPVGrafter::ExtractIndividualPlants(const FManagedArrayCollection& InCollection)
+{
+	const PV::Facades::FPlantFacade PlantFacadeSourceGraft(InCollection);
+	
+	TArray<int32> PlantNumbersInGraft = PlantFacadeSourceGraft.GetPlantNumbers();
+	TArray<FManagedArrayCollection> IndividualPlants;
+	if (PlantNumbersInGraft.Num() > 1)
+	{
+		IndividualPlants.Reserve(PlantNumbersInGraft.Num());
+		for (const int32 PlantNumber : PlantNumbersInGraft)
+		{
+			FManagedArrayCollection NewCollection;
+			InCollection.CopyTo(&NewCollection);
+			PV::Facades::FBranchFacade BranchFacade(NewCollection);
+			TArray<int> BranchesToRemove;
+			BranchesToRemove.Reserve(BranchFacade.GetElementCount());
+			for (int32 BranchIndex = 0; BranchIndex < BranchFacade.GetElementCount(); ++BranchIndex)
+			{
+				int32 PlantNumberOfBranch;
+				BranchFacade.GetPlantNumber(BranchIndex, PlantNumberOfBranch);
+				if (PlantNumberOfBranch != PlantNumber)
+				{
+					BranchesToRemove.Add(BranchIndex);
+				}
+			}
+			
+			PV::Facades::FTreeFacade::RemoveBranches(BranchFacade, BranchesToRemove, NewCollection);
+			IndividualPlants.Add(NewCollection);
+		}
+	}
+	else
+	{
+		IndividualPlants.Add(InCollection);
+	}
+	
+	return IndividualPlants;
+}
+
+void FPVGrafter::DistributeGraftWithHormoneBasedSettings(FManagedArrayCollection& OutCollection,
+                                                         const FManagedArrayCollection& SourceSkeletonCollection,
+                                                         FManagedArrayCollection& SourceGraftCollection,
+                                                         const FHormoneBasedSettings::FDistributionSettings& DistributionSettings,
+                                                         const FHormoneBasedSettings::FPhyllotaxySettings& PhyllotaxySettings,
+                                                         const FHormoneBasedSettings::FScaleSettings& ScaleSettings,
+                                                         const FHormoneBasedSettings::FVectorSettings& VectorSettings,
+                                                         const FMiscSettings& MiscSettings)
 {
-	TRACE_CPUPROFILER_EVENT_SCOPE(PVGrafter::DistributeGraft);
+	TRACE_CPUPROFILER_EVENT_SCOPE(PVGrafter::DistributeGraftWithHormoneBasedSettings);
 
 	FRandomStream RandomStream;
 	RandomStream.Initialize(MiscSettings.RandomSeed);
@@ -631,21 +676,59 @@ void FPVGrafter::DistributeGraft(FManagedArrayCollection& OutCollection, const F
 		BranchChildrenAttribute,
 		BranchNumberAttribute,
 		BranchPointsAttribute);
+	
+	// We extract all individual plants from the graft source as separate collections for simplicity
+	TArray<FManagedArrayCollection> GraftSources = ExtractIndividualPlants(SourceGraftCollection);
 
 	TManagedArrayAccessor<TArray<FVector3f>> BudDirectionAttribute(OutCollection, PV::AttributeNames::BudDirection,
 		PV::GroupNames::PointGroup);
 	TManagedArrayAccessor<float> BranchGradientsAttribute(OutCollection, PV::AttributeNames::BranchGradient,
 		PV::GroupNames::PointGroup);
-	TManagedArrayAccessor<TArray<float>> LeafPhyllotaxyAttribute(OutCollection, PV::AttributeNames::LeafPhyllotaxy,
-		PV::GroupNames::DetailsGroup);
 	TManagedArrayAccessor<TArray<float>> BudHormoneLevels(OutCollection, PV::AttributeNames::BudHormoneLevels,
 		PV::GroupNames::PointGroup);
-
-	float PhyllotaxyFormationLeaf = LeafPhyllotaxyAttribute[0][1];
-	float MinBudsLeaf = LeafPhyllotaxyAttribute[0][3];
-	float MaxBudsLeaf = LeafPhyllotaxyAttribute[0][4];
-	int32 ResetPhyllotaxyLeaf = static_cast<int>(LeafPhyllotaxyAttribute[0][6]);
-	float PhyllotaxyOffsetLeaf = LeafPhyllotaxyAttribute[0][7];
+	
+	float PhyllotaxyFormationLeaf = 180;
+	float MinBudsLeaf = 1;
+	float MaxBudsLeaf = 1;
+	bool ResetPhyllotaxyLeaf = PhyllotaxySettings.ResetPhyllotaxy;
+	float PhyllotaxyOffsetLeaf = PhyllotaxySettings.PhyllotaxyOffset;
+	
+	switch (PhyllotaxySettings.PhyllotaxyType)
+	{
+	case EPhyllotaxyType::Alternate: PhyllotaxyFormationLeaf = 180 + PhyllotaxySettings.PhyllotaxyAdditionalAngle;
+		MinBudsLeaf = MaxBudsLeaf = 1;
+		break;
+	case EPhyllotaxyType::Opposite: PhyllotaxyFormationLeaf = PhyllotaxySettings.PhyllotaxyAdditionalAngle;
+		MinBudsLeaf = MaxBudsLeaf = 2;
+		break;
+	case EPhyllotaxyType::Decussate: PhyllotaxyFormationLeaf = 90 + PhyllotaxySettings.PhyllotaxyAdditionalAngle;
+		MinBudsLeaf = MaxBudsLeaf = 2;
+		break;
+	case EPhyllotaxyType::Whorled: PhyllotaxyFormationLeaf = 90 + PhyllotaxySettings.PhyllotaxyAdditionalAngle;
+		MinBudsLeaf = PhyllotaxySettings.MinimumNodeBuds;
+		MaxBudsLeaf = PhyllotaxySettings.MaximumNodeBuds;
+		break;
+	case EPhyllotaxyType::Spiral:
+		{
+			// These leaf arrangement configurations where Distichous is two leaves, Tristichous is three etc.
+			switch (PhyllotaxySettings.PhyllotaxyFormation)
+			{
+			case EPhyllotaxyFormation::Distichous: PhyllotaxyFormationLeaf = 180 + PhyllotaxySettings.PhyllotaxyAdditionalAngle;
+				break;
+			case EPhyllotaxyFormation::Tristichous: PhyllotaxyFormationLeaf = 120 + PhyllotaxySettings.PhyllotaxyAdditionalAngle;
+				break;
+			case EPhyllotaxyFormation::Pentastichous: PhyllotaxyFormationLeaf = 144 + PhyllotaxySettings.PhyllotaxyAdditionalAngle;
+				break;
+			case EPhyllotaxyFormation::Octastichous: PhyllotaxyFormationLeaf = 135 + PhyllotaxySettings.PhyllotaxyAdditionalAngle;
+				break;
+			case EPhyllotaxyFormation::Parastichous: PhyllotaxyFormationLeaf = PhyllotaxySettings.PhyllotaxyAdditionalAngle;
+				break;
+			}
+ 
+			MinBudsLeaf = MaxBudsLeaf = 1;
+			break;
+		}
+	}
 
 	const int32 NumBranches = BranchFacadeSourceSkeleton.GetElementCount();
 	const TManagedArray<int32>& SourceSkeletonBudNumbers = PointFacadeSourceSkeleton.GetBudNumbersAttribute();
@@ -706,7 +789,7 @@ void FPVGrafter::DistributeGraft(FManagedArrayCollection& OutCollection, const F
 		FVector3f InitialAxillaryDirection(BudDirectionAttribute[BranchPoints[0]][1]);
 		const FVector3f InitialUpVector(BudDirectionAttribute[BranchPoints[0]][5]);
 
-		if (ResetPhyllotaxyLeaf == 1)
+		if (ResetPhyllotaxyLeaf)
 		{
 			InitialAxillaryDirection = FVector3f::CrossProduct(InitialApicalDirection, InitialUpVector);
 		}
@@ -779,7 +862,7 @@ void FPVGrafter::DistributeGraft(FManagedArrayCollection& OutCollection, const F
 							DistanceAlpha,
 							CurrentBudNumber,
 							CurrentBudNumber);
-						
+
 						CurrentBudNumber++;
 					}
 					AttachmentPointsPointIndices.Add(IndexAdded);
@@ -921,7 +1004,508 @@ void FPVGrafter::DistributeGraft(FManagedArrayCollection& OutCollection, const F
 				.bIsTipInstance = bIsTipInstance
 			};
 
-			SetGraft(OutCollection, SourceSkeletonCollection, SourceGraftCollection, Attributes);
+			check(GraftSources.Num() > 0);
+			const FManagedArrayCollection& PickedSourceGraft = GraftSources[RandomStream.RandRange(0, GraftSources.Num() - 1)];
+			SetGraft(OutCollection, SourceSkeletonCollection, PickedSourceGraft, Attributes);
+		}
+	}
+}
+
+TMap<int32, float> FPVGrafter::ComputeBranchLengths(PV::FPointLengthFromRootAttributeConstView PointLengthFromRootAttribute,
+                                                    PV::FBranchPointsAttributeConstView BranchPointsAttribute)
+{
+	const int32 NumOfBranches = BranchPointsAttribute.Num();
+	TMap<int32, float> BranchIndicesToBranchLengths;
+	BranchIndicesToBranchLengths.Reserve(NumOfBranches);
+
+	for (int32 BranchIndex = 0; BranchIndex < NumOfBranches; BranchIndex++)
+	{
+		const TArray<int32>& BranchPoints = BranchPointsAttribute[BranchIndex];
+		if (BranchPoints.Num() < 2)
+		{
+			BranchIndicesToBranchLengths.Add(BranchIndex, 0.0f);
+			continue;
+		}
+
+		TArray<int32> SortedPointIndices = BranchPoints;
+		SortedPointIndices.Sort([PointLengthFromRootAttribute](const int32& A, const int32& B)
+			{
+				return PointLengthFromRootAttribute[A] < PointLengthFromRootAttribute[B];
+			});
+
+		const int32 RootPointIndex = SortedPointIndices[0];
+		const int32 BranchTipPointIndex = SortedPointIndices.Last();
+
+		const float RootPointLFR = PointLengthFromRootAttribute[RootPointIndex];
+		const float TipPointLFR = PointLengthFromRootAttribute[BranchTipPointIndex];
+		const float BranchLength = TipPointLFR - RootPointLFR;
+
+		BranchIndicesToBranchLengths.Add(BranchIndex, BranchLength);
+	}
+
+	return BranchIndicesToBranchLengths;
+}
+
+FVector3f FPVGrafter::SampleVectorByFloat(const float SampleValue, const TArray<int32>& SortedPointIndices,
+                                          PV::FPointPlantGradientAttributeConstView PointPlantGradientAttribute,
+                                          PV::FPointPositionAttributeConstView PointPositionAttribute)
+{
+	if (SortedPointIndices.Num() < 2)
+	{
+		return FVector3f();
+	}
+
+	const int32 FirstPointIndex = SortedPointIndices[0];
+	const int32 LastPointIndex = SortedPointIndices.Last();
+	float Value0 = 1.0f - PointPlantGradientAttribute[FirstPointIndex];
+	float Value1 = 1.0f - PointPlantGradientAttribute[LastPointIndex];
+	FVector3f OutValue(Value0, Value0, Value0);
+
+	if (FMath::IsWithinInclusive(SampleValue, Value0, Value1))
+	{
+		for (int32 i = 1; i < SortedPointIndices.Num(); i++)
+		{
+			Value1 = 1.0f - PointPlantGradientAttribute[SortedPointIndices[i]];
+			if (FMath::IsWithinInclusive(SampleValue, Value0, Value1))
+			{
+				FVector3f P0 = PointPositionAttribute[SortedPointIndices[i - 1]];
+				FVector3f P1 = PointPositionAttribute[SortedPointIndices[i]];
+
+				float Blend = FMath::GetMappedRangeValueClamped(
+					FVector2f(Value0, Value1),
+					FVector2f(0.0f, 1.0f),
+					SampleValue);
+
+				OutValue = FMath::Lerp(P0, P1, Blend);
+			}
+
+			Value0 = Value1;
+		}
+	}
+	else
+	{
+		if (FMath::Abs(Value0 - SampleValue) > FMath::Abs(Value1 - SampleValue))
+		{
+			OutValue = FVector3f(Value1, Value1, Value1);
+		}
+	}
+
+	return OutValue;
+}
+
+float FPVGrafter::SampleFloatByFloat(const float SampleValue, const TArray<int32>& SortedPointIndices,
+                                     PV::FPointPlantGradientAttributeConstView PointPlantGradientAttribute,
+                                     PV::FPointLengthFromRootAttributeConstView PointLengthFromRootAttribute)
+{
+	if (SortedPointIndices.Num() < 2)
+	{
+		return -1.0f;
+	}
+
+	const int32 FirstPointIndex = SortedPointIndices[0];
+	const int32 LastPointIndex = SortedPointIndices.Last();
+	float Value0 = 1.0f - PointPlantGradientAttribute[FirstPointIndex];
+	float Value1 = 1.0f - PointPlantGradientAttribute[LastPointIndex];
+	float OutValue = Value0;
+
+	if (FMath::IsWithinInclusive(SampleValue, Value0, Value1))
+	{
+		for (int32 i = 1; i < SortedPointIndices.Num(); i++)
+		{
+			Value1 = 1.0f - PointPlantGradientAttribute[SortedPointIndices[i]];
+			if (FMath::IsWithinInclusive(SampleValue, Value0, Value1))
+			{
+				float P0 = PointLengthFromRootAttribute[SortedPointIndices[i - 1]];
+				float P1 = PointLengthFromRootAttribute[SortedPointIndices[i]];
+
+				float Blend = FMath::GetMappedRangeValueClamped(
+					FVector2f(Value0, Value1),
+					FVector2f(0.0f, 1.0f),
+					SampleValue);
+
+				OutValue = FMath::Lerp(P0, P1, Blend);
+			}
+
+			Value0 = Value1;
+		}
+	}
+	else
+	{
+		if (FMath::Abs(Value0 - SampleValue) > FMath::Abs(Value1 - SampleValue))
+		{
+			OutValue = Value1;
+		}
+	}
+
+	return OutValue;
+}
+
+bool FPVGrafter::FindPointByNormalizedLengthFromRoot(float SampleValue, const TArray<int32>& SortedPointIndices,
+                                                     PV::FPointLengthFromRootAttributeConstView PointLengthFromRootAttribute,
+                                                     const float LFRofRoot, const float LFRofTip, int32& PreviousPointIndex, int32& NextPointIndex,
+                                                     float& OutAlpha)
+{
+	if (SortedPointIndices.Num() < 2)
+	{
+		return false;
+	}
+
+	SampleValue = FMath::Clamp(SampleValue, 0.0f, 1.0f);
+
+	const float SampleLengthFromRoot = FMath::GetMappedRangeValueClamped(
+		FVector2f(0.0f, 1.0f),
+		FVector2f(LFRofRoot, LFRofTip),
+		SampleValue);
+
+	for (int32 i = 0; i < SortedPointIndices.Num() - 1; ++i)
+	{
+		const int32 P = SortedPointIndices[i];
+		const int32 N = SortedPointIndices[i + 1];
+		const float LFR0 = PointLengthFromRootAttribute[P];
+		const float LFR1 = PointLengthFromRootAttribute[N];
+
+		if (FMath::IsWithinInclusive(SampleLengthFromRoot, LFR0, LFR1)
+			|| FMath::IsWithinInclusive(SampleLengthFromRoot, LFR1, LFR0))
+		{
+			const float Denominator = (LFR1 - LFR0);
+			OutAlpha = (FMath::Abs(Denominator) <= KINDA_SMALL_NUMBER)
+				? 0.0f
+				: ((SampleLengthFromRoot - LFR0) / Denominator);
+			OutAlpha = FMath::Clamp(OutAlpha, 0.0f, 1.0f);
+			PreviousPointIndex = P;
+			NextPointIndex = N;
+
+			return true;
+		}
+	}
+
+	PreviousPointIndex = SortedPointIndices[0];
+	NextPointIndex = SortedPointIndices.Last();
+	OutAlpha = (SampleValue < 0.5f)
+		? 0.0f
+		: 1.0f;
+
+	return true;
+}
+
+void FPVGrafter::ComputeRotationAnglesForGraft(const FVector3f& ApicalDirection, const FVector3f& BudDirectionUp, const float AxilAngle,
+                                               const float TwistAroundApical, FVector3f& OutUpDirection, FVector3f& OutNormalDirection)
+{
+	const FVector3f ApicalDirectionNormalized = ApicalDirection.GetSafeNormal();
+	if (ApicalDirectionNormalized.IsNearlyZero())
+	{
+		OutUpDirection = FVector3f(1, 0, 0);
+		OutNormalDirection = FVector3f(0, 0, 1);
+		return;
+	}
+
+	FVector3f BudDirectionNormalized = BudDirectionUp.GetSafeNormal();
+	if (BudDirectionNormalized.IsNearlyZero())
+	{
+		FVector Axis1, Axis2;
+		FVector(ApicalDirectionNormalized).FindBestAxisVectors(Axis1, Axis2);
+		BudDirectionNormalized = FVector3f(Axis2);
+	}
+
+	FVector3f AxilAxis = FVector3f::CrossProduct(BudDirectionNormalized, ApicalDirectionNormalized).GetSafeNormal();
+	if (AxilAxis.IsNearlyZero())
+	{
+		FVector Axis1, Axis2;
+		FVector(ApicalDirectionNormalized).FindBestAxisVectors(Axis1, Axis2);
+		AxilAxis = FVector3f(Axis2);
+	}
+
+	const FQuat4f AxilQuat(AxilAxis, AxilAngle);
+	const FVector3f AxilAdjusted = AxilQuat.RotateVector(BudDirectionNormalized);
+
+	const FQuat4f TwistQuat(ApicalDirectionNormalized, TwistAroundApical);
+
+	OutUpDirection = TwistQuat.RotateVector(AxilAdjusted).GetSafeNormal();
+	OutNormalDirection = TwistQuat.RotateVector(AxilQuat.RotateVector(ApicalDirectionNormalized)).GetSafeNormal();
+}
+
+void FPVGrafter::DistributeGraftWithParametricSettings(FManagedArrayCollection& OutCollection,
+                                                       const FManagedArrayCollection& SourceSkeletonCollection,
+                                                       FManagedArrayCollection& SourceGraftCollection,
+                                                       const FParametricSettings::FSpacingSettings& SpacingSettings,
+                                                       const FParametricSettings::FAngleSettings& AngleSettings,
+                                                       const FParametricSettings::FScaleSettings& ScaleSettings,
+                                                       const FMiscSettings& MiscSettings)
+{
+	TRACE_CPUPROFILER_EVENT_SCOPE(PVGrafter::DistributeGraftWithHormoneBasedSettings);
+
+	FRandomStream RandomStream;
+	RandomStream.Initialize(MiscSettings.RandomSeed);
+
+	PV::Facades::FBranchFacade BranchFacadeSourceSkeleton(SourceSkeletonCollection);
+	PV::Facades::FPointFacade PointFacadeSourceSkeleton(SourceSkeletonCollection);
+	PV::Facades::FPointFacade PointFacadeOut(OutCollection);
+
+	auto PointPlantGradientAttribute = PV::FPointPlantGradientAttribute::GetAttribute(OutCollection);
+	auto PointBranchGradientAttribute = PV::FPointBranchGradientAttribute::GetAttribute(OutCollection);
+	auto BranchParentNumberAttribute = PV::FBranchParentNumberAttribute::GetAttribute(OutCollection);
+	auto BranchChildrenAttribute = PV::FBranchChildrenAttribute::GetAttribute(OutCollection);
+	auto BranchNumberAttribute = PV::FBranchNumberAttribute::GetAttribute(OutCollection);
+	auto BranchPointsAttribute = PV::FBranchPointsAttribute::GetAttribute(OutCollection);
+	auto PointPositionAttribute = PV::FPointPositionAttribute::GetAttribute(OutCollection);
+	auto BudDirectionAttribute = PV::FBudDirectionAttribute::GetAttribute(OutCollection);
+
+	// Only need to do this temporarily since legacy skeleton has inconsistent data in children
+	// TODO: Remove this once support for legacy skeletons is droppped
+	PV::Facades::FBranchFacade BranchFacadeOut(OutCollection);
+	BranchFacadeOut.RecomputeBranchChildren();
+	BranchFacadeOut.RecomputeBranchParents();
+
+	PV::AttributesHelper::ComputePointPlantAndBranchGradient(
+		{
+			PointPlantGradientAttribute,
+			PointBranchGradientAttribute,
+			BranchPointsAttribute,
+			BranchParentNumberAttribute,
+			BranchNumberAttribute,
+			BranchChildrenAttribute
+		});
+	
+	// We extract all individual plants from the graft source as separate collections for simplicity
+	TArray<FManagedArrayCollection> GraftSources = ExtractIndividualPlants(SourceGraftCollection);
+
+	const TManagedArray<int32>& SourceSkeletonBudNumbers = PointFacadeSourceSkeleton.GetBudNumbersAttribute();
+	check(SourceSkeletonBudNumbers.Num() > 0);
+	const int32 MaxBudNumber = *Algo::MaxElement(SourceSkeletonBudNumbers.GetConstArray());
+	int32 CurrentBudNumber = MaxBudNumber + 1;
+
+	auto PointLengthFromRootAttribute = PV::FPointLengthFromRootAttribute::GetAttribute(OutCollection);
+
+	TMap<int32, float> BranchIndicesToBranchLengths = ComputeBranchLengths(PointLengthFromRootAttribute, BranchPointsAttribute);
+	float MaxBranchLength = -1.0f;
+	float AverageBranchLength = 0.0f;
+	for (const auto IndexLengthPair : BranchIndicesToBranchLengths)
+	{
+		const float BranchLength = IndexLengthPair.Value;
+		MaxBranchLength = FMath::Max(BranchLength, MaxBranchLength);
+		AverageBranchLength = AverageBranchLength + BranchLength;
+	}
+
+	AverageBranchLength = AverageBranchLength / BranchIndicesToBranchLengths.Num();
+	AverageBranchLength = FMath::Lerp(AverageBranchLength, MaxBranchLength, 0.33f);
+
+	const int32 NumOfBranches = BranchParentNumberAttribute.Num();
+	for (int32 BranchIndex = 0; BranchIndex < NumOfBranches; BranchIndex++)
+	{
+		// Check generation limit
+		const TArray<int32>& BranchParents = BranchFacadeOut.GetParents(BranchIndex);
+		if (SpacingSettings.LimitStartGeneration && BranchParents.Num() < SpacingSettings.StartGeneration)
+		{
+			continue;
+		}
+
+		if (SpacingSettings.LimitEndGeneration && BranchParents.Num() > SpacingSettings.EndGeneration)
+		{
+			continue;
+		}
+
+		const float LengthRatio = BranchIndicesToBranchLengths[BranchIndex] / MaxBranchLength;
+		int32 LoopNumber = SpacingSettings.BranchDensity;
+		if (SpacingSettings.SpacingBasis == FParametricSettings::EDistributionBasis::Branch)
+		{
+			LoopNumber = static_cast<int32>(LoopNumber * LengthRatio);
+		}
+		else if (SpacingSettings.SpacingBasis == FParametricSettings::EDistributionBasis::Plant)
+		{
+			LoopNumber = static_cast<int32>(LoopNumber * (AverageBranchLength / MaxBranchLength));
+		}
+
+		LoopNumber = FMath::Max(LoopNumber, 1);
+
+		const TArray<int32>& BranchPoints = BranchPointsAttribute[BranchIndex];
+
+		if (BranchPoints.Num() < 2)
+		{
+			continue;
+		}
+
+		TArray<int32> SortedPointIndices = BranchPoints;
+		SortedPointIndices.Sort([PointLengthFromRootAttribute](const int32& A, const int32& B)
+			{
+				return PointLengthFromRootAttribute[A] < PointLengthFromRootAttribute[B];
+			});
+
+		const int32 FirstPointIndex = SortedPointIndices[0];
+		const int32 TipPointIndex = SortedPointIndices.Last();
+		const float PlantGradientAtRoot = 1.0f - PointPlantGradientAttribute[FirstPointIndex];
+		const float LFRofRoot = PointLengthFromRootAttribute[FirstPointIndex];
+		const float LFRofTip = PointLengthFromRootAttribute[TipPointIndex];
+
+		for (int32 j = 0; j < LoopNumber; j++)
+		{
+			float BaseValue = static_cast<float>(j) / FMath::Max(LoopNumber - 1, 1);
+			float PlantGradient;
+			float BranchGradient = 0.0f;
+			float AttachmentPointLengthFromRoot = 0.0f;
+			FVector3f Pos(0.0f, 0.0f, 0.0f);
+			int32 PointIndexA = INDEX_NONE;
+			int32 PointIndexB = INDEX_NONE;
+			float PointAlpha = 0.0f;
+
+			if (SpacingSettings.SpacingBasis == FParametricSettings::EDistributionBasis::Plant)
+			{
+				PlantGradient = BaseValue;
+				PlantGradient = SpacingSettings.SpacingRamp->GetRichCurveConst()->Eval(PlantGradient);
+				PlantGradient = FMath::GetMappedRangeValueClamped(
+					FVector2f(0.0f, 1.0f),
+					FVector2f(SpacingSettings.RelativeStart, SpacingSettings.RelativeEnd),
+					PlantGradient);
+				if (PlantGradient < PlantGradientAtRoot)
+				{
+					continue;
+				}
+
+				Pos = SampleVectorByFloat(PlantGradient, SortedPointIndices, PointPlantGradientAttribute, PointPositionAttribute);
+				AttachmentPointLengthFromRoot = SampleFloatByFloat(PlantGradient, SortedPointIndices, PointPlantGradientAttribute,
+					PointLengthFromRootAttribute);
+				BranchGradient = FMath::GetMappedRangeValueClamped(
+					FVector2f(LFRofRoot, LFRofTip),
+					FVector2f(0.0f, 1.0f),
+					AttachmentPointLengthFromRoot);
+
+				FindPointByNormalizedLengthFromRoot(BranchGradient, SortedPointIndices, PointLengthFromRootAttribute, LFRofRoot, LFRofTip,
+					PointIndexA,
+					PointIndexB, PointAlpha);
+			}
+			else if (SpacingSettings.SpacingBasis == FParametricSettings::EDistributionBasis::Branch)
+			{
+				BranchGradient = SpacingSettings.SpacingRamp->GetRichCurveConst()->Eval(BaseValue);
+				BranchGradient = FMath::GetMappedRangeValueClamped(
+					FVector2f(0.0f, 1.0f),
+					FVector2f(SpacingSettings.RelativeStart, SpacingSettings.RelativeEnd),
+					BranchGradient);
+
+				if (SpacingSettings.BranchDensity == 1)
+				{
+					BranchGradient = SpacingSettings.RelativeEnd;
+					BaseValue = 1.0f;
+				}
+
+				if (!FindPointByNormalizedLengthFromRoot(BranchGradient, SortedPointIndices, PointLengthFromRootAttribute, LFRofRoot, LFRofTip,
+					PointIndexA, PointIndexB, PointAlpha))
+				{
+					continue;
+				}
+
+				const FVector3f P0 = PointPositionAttribute[PointIndexA];
+				const FVector3f P1 = PointPositionAttribute[PointIndexB];
+				Pos = FMath::Lerp(P0, P1, PointAlpha);
+
+				AttachmentPointLengthFromRoot = FMath::GetMappedRangeValueClamped(
+					FVector2f(0.0f, 1.0f),
+					FVector2f(LFRofRoot, LFRofTip),
+					BranchGradient);
+
+				const float PlantGradient0 = 1.0f - PointPlantGradientAttribute[PointIndexA];
+				const float PlantGradient1 = 1.0f - PointPlantGradientAttribute[PointIndexB];
+				PlantGradient = FMath::Lerp(PlantGradient0, PlantGradient1, PointAlpha);
+			}
+
+			float ScaleRampValue = ScaleSettings.ScaleRamp->GetRichCurveConst()->Eval(BranchGradient);
+
+			if (SpacingSettings.SpacingBasis == FParametricSettings::EDistributionBasis::Branch)
+			{
+				ScaleRampValue = FMath::GetMappedRangeValueClamped(
+					FVector2f(0.0f, 1.0f),
+					FVector2f(0.0f, LengthRatio),
+					ScaleRampValue);
+			}
+			
+			const float RandomScaleFactor = RandomStream.FRandRange(
+				ScaleSettings.RandomizeScaleAngleMinimum,
+				ScaleSettings.RandomizeScaleAngleMaximum);
+
+			const float Scale = ScaleRampValue * RandomScaleFactor * ScaleSettings.BaseScale;
+
+			const FVector3f BudDirectionApicalA = BudDirectionAttribute[PointIndexA].Apical;
+			const FVector3f BudDirectionApicalB = BudDirectionAttribute[PointIndexB].Apical;
+			const FVector3f BudDirectionUpA = BudDirectionAttribute[PointIndexA].UpVector;
+			const FVector3f BudDirectionUpB = BudDirectionAttribute[PointIndexB].UpVector;
+
+			const FVector3f ApicalDirection = FMath::Lerp(BudDirectionApicalA, BudDirectionApicalB, PointAlpha);
+			const FVector3f BudDirectionUp = FMath::Lerp(BudDirectionUpA, BudDirectionUpB, PointAlpha);
+
+			const float AxilRampValue = AngleSettings.AxilAngleRamp->GetRichCurveConst()->Eval(BaseValue);
+			float AxilAngleInDegrees = AngleSettings.AxilAngle * AxilRampValue;
+			
+			AxilAngleInDegrees += RandomStream.FRandRange(AngleSettings.RandomizeAxilAngleMinimum, AngleSettings.RandomizeAxilAngleMaximum);
+			AxilAngleInDegrees = FMath::Clamp(AxilAngleInDegrees, -90.0f, 90.0f);
+
+			const float Twist = FMath::DegreesToRadians(AngleSettings.Rotation) * static_cast<float>(j);
+
+			FVector3f UpDirection, NormalDirection;
+			ComputeRotationAnglesForGraft(
+				ApicalDirection,
+				BudDirectionUp,
+				FMath::DegreesToRadians(AxilAngleInDegrees),
+				Twist,
+				UpDirection,
+				NormalDirection);
+
+			const bool bIsTipInstance =
+				FMath::IsNearlyEqual(AttachmentPointLengthFromRoot, LFRofTip, 0.001f) &&
+				(BranchGradient >= 0.999f || FMath::IsNearlyEqual(BranchGradient, SpacingSettings.RelativeEnd, 0.001f));
+
+			int32 AttachmentPointIndex = INDEX_NONE;
+			if (bIsTipInstance)
+			{
+				AttachmentPointIndex = TipPointIndex;
+
+				const int32 SecondToLastPointIndex = SortedPointIndices[SortedPointIndices.Num() - 2];
+				const FVector3f TipPosition = PointPositionAttribute[TipPointIndex];
+				const FVector3f SecondToLastPointPosition = PointPositionAttribute[SecondToLastPointIndex];
+				FVector3f TipTangent = (TipPosition - SecondToLastPointPosition).GetSafeNormal();
+				if (TipTangent.IsNearlyZero())
+				{
+					TipTangent = ApicalDirection.GetSafeNormal();
+				}
+
+				const FVector3f OldUp = UpDirection.GetSafeNormal();
+				FQuat4f AlignQuat = FQuat4f::FindBetweenNormals(OldUp, TipTangent);
+
+				UpDirection = TipTangent;
+				NormalDirection = AlignQuat.RotateVector(NormalDirection).GetSafeNormal();
+
+				if (FMath::Abs(FVector3f::DotProduct(UpDirection, NormalDirection)) > 0.999f)
+				{
+					FVector Axis1, Axis2;
+					FVector(UpDirection).FindBestAxisVectors(Axis1, Axis2);
+					NormalDirection = FVector3f(Axis1);
+				}
+			}
+			else
+			{
+				AttachmentPointIndex = PV::Utilities::AddInterpolatedPointToCollection(
+					OutCollection,
+					PointIndexA,
+					PointIndexB,
+					PointAlpha,
+					CurrentBudNumber,
+					CurrentBudNumber);
+
+				CurrentBudNumber++;
+			}
+
+			FPVGraftAttributes Attributes
+			{
+				.BranchIndex = BranchIndex,
+				.PointIndex = AttachmentPointIndex,
+				.AttachmentPoint = Pos,
+				.UpVector = UpDirection,
+				.NormalVector = NormalDirection,
+				.Scale = Scale,
+				.LengthFromRoot = AttachmentPointLengthFromRoot,
+				.bIsTipInstance = bIsTipInstance
+			};
+
+			check(GraftSources.Num() > 0);
+			const FManagedArrayCollection& PickedSourceGraft = GraftSources[RandomStream.RandRange(0, GraftSources.Num() - 1)]; 
+			SetGraft(OutCollection, SourceSkeletonCollection, PickedSourceGraft, Attributes);
 		}
 	}
 }
diff --git a/Engine/Plugins/Experimental/ProceduralVegetationEditor/Source/ProceduralVegetation/Private/Implementations/PVGrafter.h b/Engine/Plugins/Experimental/ProceduralVegetationEditor/Source/ProceduralVegetation/Private/Implementations/PVGrafter.h
--- a/Engine/Plugins/Experimental/ProceduralVegetationEditor/Source/ProceduralVegetation/Private/Implementations/PVGrafter.h
+++ b/Engine/Plugins/Experimental/ProceduralVegetationEditor/Source/ProceduralVegetation/Private/Implementations/PVGrafter.h
@@ -5,40 +5,94 @@
 #include "CoreMinimal.h"
 #include "Utils/PVAttributes.h"
 #include "GeometryCollection/ManagedArrayCollection.h"
-
+#include "Helpers/PVPhyllotaxyHelper.h"
 
 struct FPVFloatRamp;
 
-
 struct FPVGrafter
 {
-	struct FDistributionSettings
-	{
-		const float EthyleneThreshold;
-		const float InstanceSpacing;
-		const FPVFloatRamp* InstanceSpacingRamp;
-		const float InstanceSpacingRampEffect;
-		const int32 MaxPerBranch;
-	};
-
-	struct FScaleSettings
+	struct FHormoneBasedSettings
 	{
-		const float BaseScale;
-		const float BranchScaleImpact;
-		const float MinScale;
-		const float MaxScale;
-		const float RandomScaleMin;
-		const float RandomScaleMax;
-		const FPVFloatRamp* ScaleRamp;
+		struct FDistributionSettings
+		{
+			const float EthyleneThreshold;
+			const float InstanceSpacing;
+			const FPVFloatRamp* InstanceSpacingRamp;
+			const float InstanceSpacingRampEffect;
+			const int32 MaxPerBranch;
+		};
+		
+		struct FPhyllotaxySettings
+		{
+			const bool ResetPhyllotaxy;
+			const EPhyllotaxyType PhyllotaxyType;
+			const EPhyllotaxyFormation PhyllotaxyFormation;
+			const int32 MinimumNodeBuds;
+			const int32 MaximumNodeBuds;
+			const float PhyllotaxyAdditionalAngle;
+			const float PhyllotaxyOffset;
+		};
+
+		struct FScaleSettings
+		{
+			const float BaseScale;
+			const float BranchScaleImpact;
+			const float MinScale;
+			const float MaxScale;
+			const float RandomScaleMin;
+			const float RandomScaleMax;
+			const FPVFloatRamp* ScaleRamp;
+		};
+
+		struct FVectorSettings
+		{
+			const bool OverrideAxilAngle;
+			const float AxilAngle;
+			const FPVFloatRamp* AxilAngleRamp;
+			const float AxilAngleRampUpperValue;
+			const float AxilAngleRampEffect;
+		};
 	};
 
-	struct FVectorSettings
+	struct FParametricSettings
 	{
-		const bool OverrideAxilAngle;
-		const float AxilAngle;
-		const FPVFloatRamp* AxilAngleRamp;
-		const float AxilAngleRampUpperValue;
-		const float AxilAngleRampEffect;
+		enum class EDistributionBasis : uint8
+		{
+			Plant,
+			Branch
+		};
+
+		struct FSpacingSettings
+		{
+			const int32 BranchDensity;
+			const float RelativeStart;
+			const float RelativeEnd;
+			const bool LimitStartGeneration;
+			const int32 StartGeneration;
+			const bool LimitEndGeneration;
+			const int32 EndGeneration;
+			const EDistributionBasis SpacingBasis;
+			const FPVFloatRamp* SpacingRamp;
+		};
+
+		struct FAngleSettings
+		{
+			const float Rotation;
+			const float AxilAngle;
+			const float RandomizeAxilAngleMinimum;
+			const float RandomizeAxilAngleMaximum;
+			const EDistributionBasis AxilAngleRampBasis;
+			const FPVFloatRamp* AxilAngleRamp;
+		};
+
+		struct FScaleSettings
+		{
+			const EDistributionBasis ScaleRampBasis;
+			const float BaseScale;
+			const float RandomizeScaleAngleMinimum;
+			const float RandomizeScaleAngleMaximum;
+			const FPVFloatRamp* ScaleRamp;
+		};
 	};
 
 	struct FMiscSettings
@@ -78,8 +132,57 @@ struct FPVGrafter
 	                     const FManagedArrayCollection& SourceGraftCollection, const FPVGraftAttributes& InAttributes);
 
 	static bool RemoveFoliageDataIfPresent(FManagedArrayCollection& OutCollection);
+	
+	static TArray<FManagedArrayCollection> ExtractIndividualPlants(const FManagedArrayCollection& InCollection);
+
+	static void DistributeGraftWithHormoneBasedSettings(FManagedArrayCollection& OutCollection,
+	                                                    const FManagedArrayCollection& SourceSkeletonCollection,
+	                                                    FManagedArrayCollection& SourceGraftCollection,
+	                                                    const FHormoneBasedSettings::FDistributionSettings& DistributionSettings,
+	                                                    const FHormoneBasedSettings::FPhyllotaxySettings& PhyllotaxySettings,
+	                                                    const FHormoneBasedSettings::FScaleSettings& ScaleSettings,
+	                                                    const FHormoneBasedSettings::FVectorSettings& VectorSettings,
+	                                                    const FMiscSettings& MiscSettings);
+
+	static TMap<int32, float> ComputeBranchLengths(
+		PV::FPointLengthFromRootAttributeConstView PointLengthFromRootAttribute,
+		PV::FBranchPointsAttributeConstView BranchPointsAttribute);
 
-	static void DistributeGraft(FManagedArrayCollection& OutCollection, const FManagedArrayCollection& SourceSkeletonCollection,
-	                            FManagedArrayCollection& SourceGraftCollection, const FDistributionSettings& DistributionSettings,
-	                            const FScaleSettings& ScaleSettings, const FVectorSettings& VectorSettings, const FMiscSettings& MiscSettings);
+	static FVector3f SampleVectorByFloat(
+		const float SampleValue,
+		const TArray<int32>& SortedPointIndices,
+		PV::FPointPlantGradientAttributeConstView PointPlantGradientAttribute,
+		PV::FPointPositionAttributeConstView PointPositionAttribute);
+
+	static float SampleFloatByFloat(
+		const float SampleValue,
+		const TArray<int32>& SortedPointIndices,
+		PV::FPointPlantGradientAttributeConstView PointPlantGradientAttribute,
+		PV::FPointLengthFromRootAttributeConstView PointLengthFromRootAttribute);
+
+	static bool FindPointByNormalizedLengthFromRoot(
+		float SampleValue,
+		const TArray<int32>& SortedPointIndices,
+		PV::FPointLengthFromRootAttributeConstView PointLengthFromRootAttribute,
+		const float LFRofRoot,
+		const float LFRofTip,
+		int32& PreviousPointIndex,
+		int32& NextPointIndex,
+		float& OutAlpha);
+
+	static void ComputeRotationAnglesForGraft(
+		const FVector3f& ApicalDirection,
+		const FVector3f& BudDirectionUp,
+		const float AxilAngle,
+		const float TwistAroundApical,
+		FVector3f& OutUpDirection,
+		FVector3f& OutNormalDirection);
+
+	static void DistributeGraftWithParametricSettings(FManagedArrayCollection& OutCollection,
+	                                                  const FManagedArrayCollection& SourceSkeletonCollection,
+	                                                  FManagedArrayCollection& SourceGraftCollection,
+	                                                  const FParametricSettings::FSpacingSettings& SpacingSettings,
+	                                                  const FParametricSettings::FAngleSettings& AngleSettings,
+	                                                  const FParametricSettings::FScaleSettings& ScaleSettings,
+	                                                  const FMiscSettings& MiscSettings);
 };
diff --git a/Engine/Plugins/Experimental/ProceduralVegetationEditor/Source/ProceduralVegetation/Private/Nodes/PVFoliageDistributorSettings.h b/Engine/Plugins/Experimental/ProceduralVegetationEditor/Source/ProceduralVegetation/Private/Nodes/PVFoliageDistributorSettings.h
--- a/Engine/Plugins/Experimental/ProceduralVegetationEditor/Source/ProceduralVegetation/Private/Nodes/PVFoliageDistributorSettings.h
+++ b/Engine/Plugins/Experimental/ProceduralVegetationEditor/Source/ProceduralVegetation/Private/Nodes/PVFoliageDistributorSettings.h
@@ -5,7 +5,8 @@
 #include "Utils/PVFloatRamp.h"
  
 #include "DataTypes/PVData.h"
- 
+#include "Helpers/PVPhyllotaxyHelper.h"
+
 #include "Implementations/PVFoliage.h"
  
 #include "Nodes/PVBaseSettings.h"
diff --git a/Engine/Plugins/Experimental/ProceduralVegetationEditor/Source/ProceduralVegetation/Private/Nodes/PVGrafterSettings.cpp b/Engine/Plugins/Experimental/ProceduralVegetationEditor/Source/ProceduralVegetation/Private/Nodes/PVGrafterSettings.cpp
--- a/Engine/Plugins/Experimental/ProceduralVegetationEditor/Source/ProceduralVegetation/Private/Nodes/PVGrafterSettings.cpp
+++ b/Engine/Plugins/Experimental/ProceduralVegetationEditor/Source/ProceduralVegetation/Private/Nodes/PVGrafterSettings.cpp
@@ -5,17 +5,64 @@
 #include "ProceduralVegetationModule.h"
 #include "DataTypes/PVGrowthData.h"
 #include "ProceduralVegetationPreset.h"
-#include "Facades/PVAttributesNames.h"
-#include "Facades/PVBranchFacade.h"
 #include "Facades/PVPointFacade.h"
-#include "Implementations/PVCarve.h"
+#include "Helpers/PCGHelpers.h"
+#include "Implementations/PVGrafter.h"
 
 #define LOCTEXT_NAMESPACE "PVGrafterSettings"
 
 #if WITH_EDITOR
+UPVGrafterSettings::UPVGrafterSettings()
+{
+	if (PCGHelpers::IsNewObjectAndNotDefault(this))
+	{
+		FRichCurve* RichCurve = HormoneBasedSettings.InstanceSpacingRamp.GetRichCurve();
+		RichCurve->Reset();
+		const FKeyHandle Handle0 = RichCurve->AddKey(0.0f, 0.0f);
+		const FKeyHandle Handle1 = RichCurve->AddKey(1.0f, 1.0f);
+		RichCurve->SetKeyInterpMode(Handle0, RCIM_Linear);
+		RichCurve->SetKeyInterpMode(Handle1, RCIM_Linear);
+
+		RichCurve = HormoneBasedSettings.ScaleRamp.GetRichCurve();
+		RichCurve->Reset();
+		const FKeyHandle Handle2 = RichCurve->AddKey(0.0f, 0.0f);
+		const FKeyHandle Handle3 = RichCurve->AddKey(1.0f, 1.0f);
+		RichCurve->SetKeyInterpMode(Handle2, RCIM_Linear);
+		RichCurve->SetKeyInterpMode(Handle3, RCIM_Linear);
+
+		RichCurve = HormoneBasedSettings.AxilAngleRamp.GetRichCurve();
+		RichCurve->Reset();
+		const FKeyHandle Handle4 = RichCurve->AddKey(0.0f, 1.0f);
+		const FKeyHandle Handle5 = RichCurve->AddKey(1.0f, 0.5f);
+		RichCurve->SetKeyInterpMode(Handle4, RCIM_Linear);
+		RichCurve->SetKeyInterpMode(Handle5, RCIM_Linear);
+
+		RichCurve = ParametricSettings.SpacingRamp.GetRichCurve();
+		RichCurve->Reset();
+		const FKeyHandle Handle6 = RichCurve->AddKey(0.0f, 0.0f);
+		const FKeyHandle Handle7 = RichCurve->AddKey(1.0f, 1.0f);
+		RichCurve->SetKeyInterpMode(Handle6, RCIM_Linear);
+		RichCurve->SetKeyInterpMode(Handle7, RCIM_Linear);
+
+		RichCurve = ParametricSettings.AxilAngleRamp.GetRichCurve();
+		RichCurve->Reset();
+		const FKeyHandle Handle8 = RichCurve->AddKey(0.0f, 0.0f);
+		const FKeyHandle Handle9 = RichCurve->AddKey(1.0f, 1.0f);
+		RichCurve->SetKeyInterpMode(Handle8, RCIM_Linear);
+		RichCurve->SetKeyInterpMode(Handle9, RCIM_Linear);
+
+		RichCurve = ParametricSettings.ScaleRamp.GetRichCurve();
+		RichCurve->Reset();
+		const FKeyHandle Handle10 = RichCurve->AddKey(0.0f, 1.0f);
+		const FKeyHandle Handle11 = RichCurve->AddKey(1.0f, 0.25f);
+		RichCurve->SetKeyInterpMode(Handle10, RCIM_Linear);
+		RichCurve->SetKeyInterpMode(Handle11, RCIM_Linear);
+	}
+}
+
 FText UPVGrafterSettings::GetDefaultNodeTitle() const
 {
-	return LOCTEXT("NodeTitle", "Grafter (In Development)");
+	return LOCTEXT("NodeTitle", "Grafter");
 }
 
 FText UPVGrafterSettings::GetNodeTooltipText() const
@@ -54,6 +101,33 @@ FPCGElementPtr UPVGrafterSettings::CreateElement() const
 	return MakeShared<FPVGrafterElement>();
 }
 
+#if WITH_EDITOR
+bool UPVGrafterSettings::CanEditChange(const FProperty* InProperty) const
+{
+	bool bEditCondition = Super::CanEditChange(InProperty);
+	
+	if (!bEditCondition && (this->GetClass() != InProperty->GetOwnerClass()))
+	{
+		return bEditCondition;
+	}
+	
+	const bool ModeIsParametric = this->Mode == EPVDistributionSettingsMode::ParametricSettings;
+	const bool ModeIsHormoneBased = this->Mode == EPVDistributionSettingsMode::HormoneBasedSettings;
+	
+	if (InProperty->GetOwnerProperty()->GetFName() == GET_MEMBER_NAME_CHECKED(UPVGrafterSettings, ParametricSettings))
+	{
+		bEditCondition = ModeIsParametric;
+	}
+	
+	if (InProperty->GetOwnerProperty()->GetFName() == GET_MEMBER_NAME_CHECKED(UPVGrafterSettings, HormoneBasedSettings))
+	{
+		bEditCondition = ModeIsHormoneBased;
+	}
+	
+	return bEditCondition;
+}
+#endif
+
 bool FPVGrafterElement::ExecuteInternal(FPCGContext* InContext) const
 {
 	TRACE_CPUPROFILER_EVENT_SCOPE(FPVGrafterElement::ExecuteInternal);
@@ -87,36 +161,6 @@ bool FPVGrafterElement::ExecuteInternal(FPCGContext* InContext) const
 
 	UPVGrowthData* OutManagedArrayCollectionData = FPCGContext::NewObject_AnyThread<UPVGrowthData>(InContext);
 
-	const FPVGrafter::FDistributionSettings DistributionSettings = {
-		.EthyleneThreshold = Settings->EthyleneThreshold,
-		.InstanceSpacing = Settings->InstanceSpacing,
-		.InstanceSpacingRamp = &Settings->InstanceSpacingRamp,
-		.InstanceSpacingRampEffect = Settings->InstanceSpacingRampEffect,
-		.MaxPerBranch = Settings->MaxPerBranch
-	};
-
-	const FPVGrafter::FScaleSettings ScaleSettings = {
-		.BaseScale = Settings->BaseScale,
-		.BranchScaleImpact = Settings->BranchScaleImpact,
-		.MinScale = Settings->MinScale,
-		.MaxScale = Settings->MaxScale,
-		.RandomScaleMin = Settings->RandomScaleMin,
-		.RandomScaleMax = Settings->RandomScaleMax,
-		.ScaleRamp = &Settings->ScaleRamp
-	};
-
-	const FPVGrafter::FVectorSettings VectorSettings = {
-		.OverrideAxilAngle = Settings->OverrideAxilAngle,
-		.AxilAngle = Settings->AxilAngle,
-		.AxilAngleRamp = &Settings->AxilAngleRamp,
-		.AxilAngleRampUpperValue = Settings->AxilAngleRampUpperValue,
-		.AxilAngleRampEffect = Settings->AxilAngleRampEffect
-	};
-
-	const FPVGrafter::FMiscSettings MiscSettings = {
-		.RandomSeed = Settings->RandomSeed
-	};
-
 	if (FPVGrafter::RemoveFoliageDataIfPresent(OutCollection))
 	{
 		PCGLog::LogWarningOnGraph(LOCTEXT("FoliageDataRemovedFromSourceSkeleton", "Foliage data removed from source skeleton"), InContext);
@@ -127,8 +171,108 @@ bool FPVGrafterElement::ExecuteInternal(FPCGContext* InContext) const
 		PCGLog::LogWarningOnGraph(LOCTEXT("FoliageDataRemovedFromGraft", "Foliage data removed from source graft"), InContext);
 	}
 
-	FPVGrafter::DistributeGraft(OutCollection, SourceSkeletonCollection, SourceGraftCollection,
-		DistributionSettings, ScaleSettings, VectorSettings, MiscSettings);
+	if (Settings->Mode == EPVDistributionSettingsMode::HormoneBasedSettings)
+	{
+		const FPVGrafter::FHormoneBasedSettings::FDistributionSettings DistributionSettings = {
+			.EthyleneThreshold = Settings->HormoneBasedSettings.EthyleneThreshold,
+			.InstanceSpacing = Settings->HormoneBasedSettings.InstanceSpacing,
+			.InstanceSpacingRamp = &Settings->HormoneBasedSettings.InstanceSpacingRamp,
+			.InstanceSpacingRampEffect = Settings->HormoneBasedSettings.InstanceSpacingRampEffect,
+			.MaxPerBranch = Settings->HormoneBasedSettings.MaxPerBranch
+		};
+		
+		const FPVGrafter::FHormoneBasedSettings::FPhyllotaxySettings PhyllotaxySettings = {
+			.ResetPhyllotaxy = Settings->HormoneBasedSettings.ResetPhyllotaxy,
+			.PhyllotaxyType = Settings->HormoneBasedSettings.PhyllotaxyType,
+			.PhyllotaxyFormation = Settings->HormoneBasedSettings.PhyllotaxyFormation,
+			.MinimumNodeBuds = Settings->HormoneBasedSettings.MinimumNodeBuds,
+			.MaximumNodeBuds = Settings->HormoneBasedSettings.MaximumNodeBuds,
+			.PhyllotaxyAdditionalAngle = Settings->HormoneBasedSettings.PhyllotaxyAdditionalAngle,
+			.PhyllotaxyOffset = Settings->HormoneBasedSettings.PhyllotaxyOffset
+		}; 
+
+		const FPVGrafter::FHormoneBasedSettings::FScaleSettings ScaleSettings = {
+			.BaseScale = Settings->HormoneBasedSettings.BaseScale,
+			.BranchScaleImpact = Settings->HormoneBasedSettings.BranchScaleImpact,
+			.MinScale = Settings->HormoneBasedSettings.MinScale,
+			.MaxScale = Settings->HormoneBasedSettings.MaxScale,
+			.RandomScaleMin = Settings->HormoneBasedSettings.RandomScaleMin,
+			.RandomScaleMax = Settings->HormoneBasedSettings.RandomScaleMax,
+			.ScaleRamp = &Settings->HormoneBasedSettings.ScaleRamp
+		};
+
+		const FPVGrafter::FHormoneBasedSettings::FVectorSettings VectorSettings = {
+			.OverrideAxilAngle = Settings->HormoneBasedSettings.OverrideAxilAngle,
+			.AxilAngle = Settings->HormoneBasedSettings.AxilAngle,
+			.AxilAngleRamp = &Settings->HormoneBasedSettings.AxilAngleRamp,
+			.AxilAngleRampUpperValue = Settings->HormoneBasedSettings.AxilAngleRampUpperValue,
+			.AxilAngleRampEffect = Settings->HormoneBasedSettings.AxilAngleRampEffect
+		};
+
+		const FPVGrafter::FMiscSettings MiscSettings = {
+			.RandomSeed = Settings->RandomSeed
+		};
+
+		FPVGrafter::DistributeGraftWithHormoneBasedSettings(
+			OutCollection,
+			SourceSkeletonCollection,
+			SourceGraftCollection,
+			DistributionSettings,
+			PhyllotaxySettings,
+			ScaleSettings,
+			VectorSettings,
+			MiscSettings);
+	}
+	else
+	{
+		const FPVGrafter::FParametricSettings::FSpacingSettings SpacingSettings = {
+			.BranchDensity = Settings->ParametricSettings.BranchDensity,
+			.RelativeStart = Settings->ParametricSettings.RelativeStart,
+			.RelativeEnd = Settings->ParametricSettings.RelativeEnd,
+			.LimitStartGeneration = Settings->ParametricSettings.LimitStartGeneration,
+			.StartGeneration = Settings->ParametricSettings.StartGeneration,
+			.LimitEndGeneration = Settings->ParametricSettings.LimitEndGeneration,
+			.EndGeneration = Settings->ParametricSettings.EndGeneration,
+			.SpacingBasis = Settings->ParametricSettings.SpacingBasis == EPVDistributionBasis::Plant
+			? FPVGrafter::FParametricSettings::EDistributionBasis::Plant
+			: FPVGrafter::FParametricSettings::EDistributionBasis::Branch,
+			.SpacingRamp = &Settings->ParametricSettings.SpacingRamp
+		};
+
+		const FPVGrafter::FParametricSettings::FAngleSettings AngleSettings = {
+			.Rotation = Settings->ParametricSettings.Rotation,
+			.AxilAngle = Settings->ParametricSettings.AxilAngle,
+			.RandomizeAxilAngleMinimum = Settings->ParametricSettings.RandomizeAxilAngleMinimum,
+			.RandomizeAxilAngleMaximum = Settings->ParametricSettings.RandomizeAxilAngleMaximum,
+			.AxilAngleRampBasis = Settings->ParametricSettings.AxilAngleRampBasis == EPVDistributionBasis::Plant
+			? FPVGrafter::FParametricSettings::EDistributionBasis::Plant
+			: FPVGrafter::FParametricSettings::EDistributionBasis::Branch,
+			.AxilAngleRamp = &Settings->ParametricSettings.AxilAngleRamp
+		};
+
+		const FPVGrafter::FParametricSettings::FScaleSettings ScaleSettings = {
+			.ScaleRampBasis = Settings->ParametricSettings.ScaleRampBasis == EPVDistributionBasis::Plant
+			? FPVGrafter::FParametricSettings::EDistributionBasis::Plant
+			: FPVGrafter::FParametricSettings::EDistributionBasis::Branch,
+			.BaseScale = Settings->ParametricSettings.BaseScale,
+			.RandomizeScaleAngleMinimum = Settings->ParametricSettings.RandomizeScaleAngleMinimum,
+			.RandomizeScaleAngleMaximum = Settings->ParametricSettings.RandomizeScaleAngleMaximum,
+			.ScaleRamp = &Settings->ParametricSettings.ScaleRamp
+		};
+
+		const FPVGrafter::FMiscSettings MiscSettings = {
+			.RandomSeed = Settings->RandomSeed
+		};
+
+		FPVGrafter::DistributeGraftWithParametricSettings(
+			OutCollection,
+			SourceSkeletonCollection,
+			SourceGraftCollection,
+			SpacingSettings,
+			AngleSettings,
+			ScaleSettings,
+			MiscSettings);
+	}
 
 	OutManagedArrayCollectionData->Initialize(MoveTemp(OutCollection));
 	InContext->OutputData.TaggedData.Emplace(OutManagedArrayCollectionData);
diff --git a/Engine/Plugins/Experimental/ProceduralVegetationEditor/Source/ProceduralVegetation/Private/Nodes/PVGrafterSettings.h b/Engine/Plugins/Experimental/ProceduralVegetationEditor/Source/ProceduralVegetation/Private/Nodes/PVGrafterSettings.h
--- a/Engine/Plugins/Experimental/ProceduralVegetationEditor/Source/ProceduralVegetation/Private/Nodes/PVGrafterSettings.h
+++ b/Engine/Plugins/Experimental/ProceduralVegetationEditor/Source/ProceduralVegetation/Private/Nodes/PVGrafterSettings.h
@@ -2,9 +2,9 @@
 #pragma once
 
 #include "CoreMinimal.h"
+#include "Helpers/PVPhyllotaxyHelper.h"
 #include "Utils/PVFloatRamp.h"
 #include "Nodes/PVBaseSettings.h"
-#include "Implementations/PVGrafter.h"
 
 #include "PVGrafterSettings.generated.h"
 
@@ -14,35 +14,96 @@ namespace PVGrafterInputPins
 	static const FName Graft("Graft");
 }
 
-UCLASS(BlueprintType, ClassGroup = (Procedural))
-class UPVGrafterSettings : public UPVBaseSettings
+UENUM()
+enum class EPVDistributionSettingsMode : uint8
+{
+	ParametricSettings UMETA(DisplayName="Use Parametric Settings"),
+	HormoneBasedSettings UMETA(DisplayName="Use Hormone Based Settings")
+};
+
+UENUM()
+enum class EPVDistributionBasis : uint8
+{
+	Plant UMETA(DisplayName="Plant"),
+	Branch UMETA(DisplayName="Branch")
+};
+
+USTRUCT()
+struct FPVDistributionParametricSettings
 {
 	GENERATED_BODY()
-	
-	friend class FPVGrafterElement;
 
-public:
-#if WITH_EDITOR
-	virtual FName GetDefaultNodeName() const override
-	{
-		return FName(TEXT("Grafter"));
-	}
+	UPROPERTY(EditAnywhere, Category="Spacing Settings", meta=(PCG_Overridable, ClampMin=0, ClampMax=100))
+	int32 BranchDensity = 45;
 
-	virtual FText GetDefaultNodeTitle() const override;
-	virtual FText GetNodeTooltipText() const override;
+	UPROPERTY(EditAnywhere, Category="Spacing Settings", meta=(PCG_Overridable, ClampMin=0.0f, ClampMax=1.0f))
+	float RelativeStart = 0.0f;
 
-	virtual TArray<EPVRenderType> GetDefaultRenderType() const override
-	{
-		return TArray{EPVRenderType::PointData};
-	}
-#endif
+	UPROPERTY(EditAnywhere, Category="Spacing Settings", meta=(PCG_Overridable, ClampMin=0.0f, ClampMax=1.0f))
+	float RelativeEnd = 1.0f;
 
-protected:
-	virtual TArray<FPCGPinProperties> InputPinProperties() const override;
-	virtual FPCGDataTypeIdentifier GetOutputPinTypeIdentifier() const override;
-	virtual FPCGElementPtr CreateElement() const override;
+	UPROPERTY(EditAnywhere, Category="Spacing Settings", meta=(PCG_Overridable))
+	bool LimitStartGeneration = false;
+
+	UPROPERTY(EditAnywhere, Category="Spacing Settings",
+		meta=(PCG_Overridable, ClampMin=0, ClampMax=10, EditCondition="LimitStartGeneration", EditConditionHides))
+	int32 StartGeneration = 0;
+
+	UPROPERTY(EditAnywhere, Category="Spacing Settings", meta=(PCG_Overridable))
+	bool LimitEndGeneration = false;
+
+	UPROPERTY(EditAnywhere, Category="Spacing Settings",
+		meta=(PCG_Overridable, ClampMin=0, ClampMax=10, EditCondition="LimitEndGeneration", EditConditionHides))
+	int32 EndGeneration = 0;
+
+	UPROPERTY(EditAnywhere, Category="Spacing Settings", meta=(PCG_Overridable))
+	EPVDistributionBasis SpacingBasis = EPVDistributionBasis::Plant;
+
+	UPROPERTY(EditAnywhere, Category="Spacing Settings",
+		meta=(PCG_Overridable, XAxisMin=0.0f, XAxisMax=1.0f, YAxisMin=0.0f, YAxisMax=1.0f))
+	FPVFloatRamp SpacingRamp;
+
+	UPROPERTY(EditAnywhere, Category="Angle Settings", meta=(PCG_Overridable, ClampMin=0.0f, ClampMax=360.0f))
+	float Rotation = 0.0f;
+
+	UPROPERTY(EditAnywhere, Category="Angle Settings", meta=(PCG_Overridable, ClampMin=-90.0f, ClampMax=90.0f))
+	float AxilAngle = 0.0f;
+
+	UPROPERTY(EditAnywhere, Category="Angle Settings", meta=(PCG_Overridable, ClampMin=-45.0f, ClampMax=0.0f))
+	float RandomizeAxilAngleMinimum = 0.0f;
+
+	UPROPERTY(EditAnywhere, Category="Angle Settings", meta=(PCG_Overridable, ClampMin=0.0f, ClampMax=45.0f))
+	float RandomizeAxilAngleMaximum = 0.0f;
+
+	UPROPERTY(EditAnywhere, Category="Angle Settings", meta=(PCG_Overridable))
+	EPVDistributionBasis AxilAngleRampBasis = EPVDistributionBasis::Plant;
+
+	UPROPERTY(EditAnywhere, Category="Angle Settings",
+		meta=(PCG_Overridable, XAxisMin=0.0f, XAxisMax=1.0f, YAxisMin=0.0f, YAxisMax=1.0f))
+	FPVFloatRamp AxilAngleRamp;
+
+	UPROPERTY(EditAnywhere, Category="Parametric Scale Settings", meta=(PCG_Overridable))
+	EPVDistributionBasis ScaleRampBasis = EPVDistributionBasis::Plant;
+
+	UPROPERTY(EditAnywhere, Category="Parametric Scale Settings", meta=(PCG_Overridable, ClampMin=0.0f, ClampMax=10.0f))
+	float BaseScale = 0.2f;
+
+	UPROPERTY(EditAnywhere, Category="Parametric Scale Settings", meta=(PCG_Overridable, ClampMin=0.0f, ClampMax=1.0f))
+	float RandomizeScaleAngleMinimum = 1.0f;
+
+	UPROPERTY(EditAnywhere, Category="Parametric Scale Settings", meta=(PCG_Overridable, ClampMin=1.0f, ClampMax=2.0f))
+	float RandomizeScaleAngleMaximum = 1.0f;
+
+	UPROPERTY(EditAnywhere, Category="Parametric Scale Settings",
+		meta=(PCG_Overridable, XAxisMin=0.0f, XAxisMax=1.0f, YAxisMin=0.0f, YAxisMax=1.0f))
+	FPVFloatRamp ScaleRamp;
+};
+
+USTRUCT()
+struct FPVDistributionHormoneBasedSettings
+{
+	GENERATED_BODY()
 
-private:
 	UPROPERTY(EditAnywhere, Category="Distribution Settings",
 		meta=(PCG_Overridable, ClampMin=0.0f, ClampMax=1.0f, UIMin=0.0f, UIMax=1.0f, Tooltip=
 			"Threshold controlling foliage release based on an ethylene-like signal.\n\nHigher Threshold → more buds/branches retained, making the plant denser and bushier.\nLower Threshold → fewer buds/branches retained, resulting in a sparser, pruned structure."
@@ -68,47 +129,85 @@ private:
 
 	UPROPERTY(EditAnywhere, Category="Distribution Settings",
 		meta=(PCG_Overridable, ClampMin=-1, ClampMax=1000, Tooltip=
-			"Caps the number of instances per branch\n\nlimit on how many foliage instances can be placed per branch to avoid overcrowding."))
+			"Caps the number of instances per branch.\n\nLimit on how many foliage instances can be placed per branch to avoid overcrowding."))
 	int32 MaxPerBranch = -1;
 
-	UPROPERTY(EditAnywhere, Category="Scale Settings",
+	UPROPERTY(EditAnywhere, Category="Phyllotaxy Settings",
+		meta=(PCG_Overridable, Tooltip=
+			"Reset phyllotaxy per branch.\n\nBranch orientation is used to align the initial phyllotaxy. This setting will reset it per branch so that its no longer tied to parent orientation."
+		))
+	bool ResetPhyllotaxy = false;
+
+	UPROPERTY(EditAnywhere, Category="Phyllotaxy Settings",
+		meta=(PCG_Overridable, Tooltip=
+			"Pattern for arranging buds/leaves around the stem.\n\nSelects the phyllotactic pattern used for placement (e.g., Alternate, Opposite, Whorled, Spiral). This defines the sequence and angular spacing of buds/leaves around the stem and along its length."
+		))
+	EPhyllotaxyType PhyllotaxyType = EPhyllotaxyType::Alternate;
+
+	UPROPERTY(EditAnywhere, Category="Phyllotaxy Settings",
+		meta=(PCG_Overridable, EditCondition="PhyllotaxyType == EPhyllotaxyType::Spiral", EditConditionHides))
+	EPhyllotaxyFormation PhyllotaxyFormation = EPhyllotaxyFormation::Distichous;
+
+	UPROPERTY(EditAnywhere, Category="Phyllotaxy Settings",
+		meta=(PCG_Overridable, ClampMin=1, ClampMax=10, UIMin=1, UIMax=10, Tooltip=
+			"Minimum buds/leaves per node.\n\nSets the lower bound on the number of buds/leaves generated at each node. Guarantees at least this many instances (subject to spacing and other constraints). Should be ≤ MaximumNodeBuds. Applied when Whorled phyllotaxy is selected."
+		))
+	int32 MinimumNodeBuds = 2;
+
+	UPROPERTY(EditAnywhere, Category="Phyllotaxy Settings",
+		meta=(PCG_Overridable, ClampMin=1, ClampMax=10, UIMin=1, UIMax=10, Tooltip=
+			"Maximum buds/leaves per node.\n\nCaps the number of buds/leaves spawned. Use with MinimumNodeBuds to bound density and keep arrangements predictable. Applied when Whorled phyllotaxy is selected."
+		))
+	int32 MaximumNodeBuds = 3;
+
+	UPROPERTY(EditAnywhere, Category="Phyllotaxy Settings",
+		meta=(PCG_Overridable, ClampMin=0.0f, ClampMax=360.0f, UIMin=0.0f, UIMax=360.0f, Tooltip=
+			"Extra divergence angle (°) added to the phyllotaxy pattern\n\nAdds an angular offset to the base phyllotactic divergence."))
+	float PhyllotaxyAdditionalAngle = 0.0f;
+	
+	UPROPERTY(EditAnywhere, Category="Phyllotaxy Settings",
+		meta=(PCG_Overridable, ClampMin=0.0f, ClampMax=360.0f, UIMin=0.0f, UIMax=360.0f, Tooltip=
+			"Initial offset value for the phyllotaxy (rotation around parent).\n\nValues are in degrees."))
+	float PhyllotaxyOffset = 0.0f;
+
+	UPROPERTY(EditAnywhere, Category="Hormone based Scale Settings",
 		meta=(PCG_Overridable, ClampMin=0.0f, UIMin=0.0f, UIMax=5.0f, Tooltip=
 			"Default size multiplier for placed foliage\n\nSets the base uniform scale applied to each instance before randomness and ramps are considered. Increase to enlarge all instances proportionally; decrease for smaller foliage. Final size still respects Min/Max limits."
 		))
 	float BaseScale = 1.0f;
 
-	UPROPERTY(EditAnywhere, Category="Scale Settings",
+	UPROPERTY(EditAnywhere, Category="Hormone based Scale Settings",
 		meta=(PCG_Overridable, ClampMin=0.0f, ClampMax=1.0f, UIMin=0.0f, UIMax=1.0f, Tooltip=
 			"Impact of branch size on foliage size\n\nThe size of the foliage being scattered will be chosen in relation to the physical size of the branch it is distributed on. The largest branches will be distributed the largest foliage pieces."
 		))
 	float BranchScaleImpact = 0.0f;
 
-	UPROPERTY(EditAnywhere, Category="Scale Settings",
+	UPROPERTY(EditAnywhere, Category="Hormone based Scale Settings",
 		meta=(PCG_Overridable, ClampMin=0.0f, UIMin=0.0f, UIMax=5.0f, Tooltip=
 			"Minimum allowed instance size.\n\nClamps the final computed scale so it never goes below this value. Use to prevent tiny, hard-to-see instances."
 		))
 	float MinScale = 0.5f;
 
-	UPROPERTY(EditAnywhere, Category="Scale Settings",
+	UPROPERTY(EditAnywhere, Category="Hormone based Scale Settings",
 		meta=(PCG_Overridable, ClampMin=0.0f, UIMin=0.0f, UIMax=5.0f, Tooltip=
 			"Maximum allowed instance size.\n\nClamps the final computed scale so it never exceeds this value. Use to cap unusually large instances for visual consistency."
 		))
 	float MaxScale = 1.0f;
 
-	UPROPERTY(EditAnywhere, Category="Scale Settings",
+	UPROPERTY(EditAnywhere, Category="Hormone based Scale Settings",
 		meta=(PCG_Overridable, ClampMin=0.0f, UIMin=0.0f, UIMax=5.0f, Tooltip=
 			"Lower bound for per-instance size variation.\n\nSets the minimum of the random scale range used to vary instance sizes. Each instance samples within Random Scale Min–Max, then combines with Base Scale and any ramps (respecting Min/Max Scale)."
 		))
 	float RandomScaleMin = 1.0f;
 
-	UPROPERTY(EditAnywhere, Category="Scale Settings",
+	UPROPERTY(EditAnywhere, Category="Hormone based Scale Settings",
 		meta=(PCG_Overridable, ClampMin=1.0f, UIMin=1.0f, UIMax=5.0f, Tooltip=
 			"Upper bound for per-instance size variation.\n\nSets the maximum of the random scale range used to vary instance sizes. Each instance samples within Random Scale Min–Max, then combines with Base Scale and any ramps (respecting Min/Max Scale)."
 		))
 	float RandomScaleMax = 1.0f;
 
-	UPROPERTY(EditAnywhere, Category="Scale Settings",
-		meta=(XAxisMin=0.0f, XAxisMax=1.0f, YAxisMin=0.0f, YAxisMax=1.0f, Tooltip=
+	UPROPERTY(EditAnywhere, Category="Hormone based Scale Settings",
+		meta=(PCG_Overridable, XAxisMin=0.0f, XAxisMax=1.0f, YAxisMin=0.0f, YAxisMax=1.0f, Tooltip=
 			"Curve that varies scale across the height of the plant\n\nUse to taper foliage toward tips or emphasize size in targeted regions."))
 	FPVFloatRamp ScaleRamp;
 
@@ -125,7 +224,7 @@ private:
 	float AxilAngle = 45.0f;
 
 	UPROPERTY(EditAnywhere, Category="Vector Settings",
-		meta=(XAxisMin=0.0f, XAxisMax=1.0f, YAxisMin=0.0f, YAxisMax=1.0f, Tooltip=
+		meta=(PCG_Overridable, XAxisMin=0.0f, XAxisMax=1.0f, YAxisMin=0.0f, YAxisMax=1.0f, Tooltip=
 			"Ramp that varies the axil angle across the height of the plant\n\nIts output is remapped to an angle range defined by AxilAngle and AxilAngleRampUpperValue, modulated by AxilAngleRampEffect."
 		))
 	FPVFloatRamp AxilAngleRamp;
@@ -141,8 +240,54 @@ private:
 			"Adjusts this parameter for finer control.\n\nDefines the upper axil angle reached when AxilAngleRamp evaluates to 1. Together with AxilAngle, it forms the min–max range for ramp-driven orientation."
 		))
 	float AxilAngleRampEffect = 0.0f;
+};
+
+UCLASS(BlueprintType, ClassGroup = (Procedural))
+class UPVGrafterSettings : public UPVBaseSettings
+{
+	GENERATED_BODY()
+
+	friend class FPVGrafterElement;
+
+public:
+	UPVGrafterSettings();
+
+#if WITH_EDITOR
+	virtual FName GetDefaultNodeName() const override
+	{
+		return FName(TEXT("Grafter"));
+	}
+
+	virtual FText GetDefaultNodeTitle() const override;
+	virtual FText GetNodeTooltipText() const override;
+
+	virtual TArray<EPVRenderType> GetDefaultRenderType() const override
+	{
+		return TArray{EPVRenderType::PointData};
+	}
+#endif
+
+protected:
+	virtual TArray<FPCGPinProperties> InputPinProperties() const override;
+	virtual FPCGDataTypeIdentifier GetOutputPinTypeIdentifier() const override;
+	virtual FPCGElementPtr CreateElement() const override;
+#if WITH_EDITOR
+	virtual bool CanEditChange(const FProperty* InProperty) const override;
+#endif
+
+private:
+	UPROPERTY(EditAnywhere, Category="Distribution Settings Mode", meta=(PCG_Overridable))
+	EPVDistributionSettingsMode Mode = EPVDistributionSettingsMode::ParametricSettings;
+
+	UPROPERTY(EditAnywhere, Category="Parametric Settings",
+		meta=(PCG_Overridable, ShowOnlyInnerProperties))
+	FPVDistributionParametricSettings ParametricSettings;
+
+	UPROPERTY(EditAnywhere, Category="Hormone Based Settings",
+		meta=(PCG_Overridable, ShowOnlyInnerProperties))
+	FPVDistributionHormoneBasedSettings HormoneBasedSettings;
 
-	UPROPERTY(EditAnywhere, Category="Misc Settings", meta = (PCG_Overridable))
+	UPROPERTY(EditAnywhere, Category="Misc Settings", meta=(PCG_Overridable))
 	int32 RandomSeed = 123456;
 };
 
diff --git a/Engine/Plugins/Experimental/ProceduralVegetationEditor/Source/ProceduralVegetation/Public/Facades/PVBranchFacade.h b/Engine/Plugins/Experimental/ProceduralVegetationEditor/Source/ProceduralVegetation/Public/Facades/PVBranchFacade.h
--- a/Engine/Plugins/Experimental/ProceduralVegetationEditor/Source/ProceduralVegetation/Public/Facades/PVBranchFacade.h
+++ b/Engine/Plugins/Experimental/ProceduralVegetationEditor/Source/ProceduralVegetation/Public/Facades/PVBranchFacade.h
@@ -108,7 +108,12 @@ namespace PV::Facades
 		void FillParents(const TArray<TArray<int32>>& InputArray);
 		
 		void RecomputeBranchChildren();
-		
+
+		static TArray<int32> ComputeParentsForBranchNumber(const int32 BranchNumber, const TMap<int32, int32>& BranchNumbersToBranchParentNumbers,
+		                                                   TMap<int32, TArray<int32>>& Memo, TSet<int32>& VisitedParents);
+
+		void RecomputeBranchParents();
+
 		void FillChildren(const TArray<TArray<int32>>& InputArray);
 		
 		void FillBranchPoints(const TArray<TArray<int32>>& InputArray);
