Fix PVTransformHelper::TransformPoint supplying the wrong index to GetBranchPointFoliageEntryIds

Since CL 51291134 GetBranchPointFoliageEntryIds no longer takes a PointIndex but instead a BranchPointIndex. PVTransformHelper::TransformPoint was not updated to reflect this change.

#Synced-CL 51372802

diff --git a/Engine/Plugins/Experimental/ProceduralVegetationEditor/Source/ProceduralVegetation/Private/Helpers/PVTransformHelper.cpp b/Engine/Plugins/Experimental/ProceduralVegetationEditor/Source/ProceduralVegetation/Private/Helpers/PVTransformHelper.cpp
--- a/Engine/Plugins/Experimental/ProceduralVegetationEditor/Source/ProceduralVegetation/Private/Helpers/PVTransformHelper.cpp
+++ b/Engine/Plugins/Experimental/ProceduralVegetationEditor/Source/ProceduralVegetation/Private/Helpers/PVTransformHelper.cpp
@@ -38,7 +38,7 @@ void PV::Transform::TransformPoint(
 			Attributes.OptionalBranchFoliageIDs,
 			Attributes.OptionalFoliageLengthFromRoot,
 			BranchIndex,
-			PointIndex
+			Attributes.OptionalBranchPoints[BranchIndex].Find(PointIndex)
 		)
 		: TArray<int32>();
 
