Convert checks to ensure in AttributesHelper & PlantTraversalHelper where we can gracefully recover to avoid crashing when being fed bad data

#Synced-CL 51372802

diff --git a/Engine/Plugins/Experimental/ProceduralVegetationEditor/Source/ProceduralVegetation/Private/Helpers/PVAttributesHelper.cpp b/Engine/Plugins/Experimental/ProceduralVegetationEditor/Source/ProceduralVegetation/Private/Helpers/PVAttributesHelper.cpp
--- a/Engine/Plugins/Experimental/ProceduralVegetationEditor/Source/ProceduralVegetation/Private/Helpers/PVAttributesHelper.cpp
+++ b/Engine/Plugins/Experimental/ProceduralVegetationEditor/Source/ProceduralVegetation/Private/Helpers/PVAttributesHelper.cpp
@@ -12,7 +12,11 @@ void PV::AttributesHelper::ComputeBudNumbers(const FComputeBudNumbersAttributes&
 
 	for (int32 BranchIndex = 0; BranchIndex < InAttributes.BranchSourceBudNumber.Num(); ++BranchIndex)
 	{
-		check(InAttributes.BranchPoints[BranchIndex].Num() > 1);
+		if (!ensure(InAttributes.BranchPoints[BranchIndex].Num() > 0))
+		{
+			continue;
+		}
+
 		const int32 BranchRootPointIndex = InAttributes.BranchPoints[BranchIndex][0];
 		InAttributes.BranchSourceBudNumber[BranchIndex] = InAttributes.PointBudNumber[BranchRootPointIndex];
 	}
@@ -57,7 +61,10 @@ void PV::AttributesHelper::ComputeLengthFromRoot(const FComputeLengthFromRootAtt
 	RecursiveWalkBranches(InAttributes.BranchParentNumber, InAttributes.BranchChildren, InAttributes.BranchNumber, [&](int32 BranchIndex)->EForEachResult
 	{
 		const TArray<int32>& CurrentBranchPoints = InAttributes.BranchPoints[BranchIndex];
-		check(CurrentBranchPoints.Num() > 0);
+		if (!ensure(CurrentBranchPoints.Num() > 0))
+		{
+			return EForEachResult::Continue;
+		}
 
 		const int32 FirstPointIndex = CurrentBranchPoints[0];
 		const int32 ParentBranchIndex = GetBranchParentIndex(InAttributes.BranchParentNumber, InAttributes.BranchNumber, BranchIndex);
@@ -70,7 +77,10 @@ void PV::AttributesHelper::ComputeLengthFromRoot(const FComputeLengthFromRootAtt
 		else
 		{
 			const TArray<int32>& ParentBranchPoints = InAttributes.BranchPoints[ParentBranchIndex];
-			check(ParentBranchPoints.Num() > 0);
+			if (!ensure(ParentBranchPoints.Num() > 0))
+			{
+				return EForEachResult::Continue;
+			}
 
 			const FVector3f& FirstPoint = InAttributes.PointPosition[FirstPointIndex];
 
@@ -360,7 +370,10 @@ void PV::AttributesHelper::ComputePointPlantAndBranchGradient(const FComputePoin
 		[&](int32 BranchIndex)->EForEachResult
 		{
 			const TArray<int32>& BranchPoints = InAttributes.BranchPoints[BranchIndex];
-			check(BranchPoints.Num() > 1);
+			if (!ensure(BranchPoints.Num() > 1))
+			{
+				return EForEachResult::Continue;
+			}
 
 			const float BasePointPlantGradient = GetBranchParentIndex(InAttributes.BranchParentNumber, InAttributes.BranchNumber, BranchIndex) == INDEX_NONE
 				? 1.0
@@ -481,7 +494,10 @@ void PV::AttributesHelper::ComputeBudLateralMeristem(const FComputeBudLateralMer
 		[&](int32 BranchIndex)->EForEachResult
 		{
 			const TArray<int32>& BranchPointIndices = InAttributes.BranchPoints[BranchIndex];
-			check(BranchPointIndices.Num() > 1);
+			if (!ensure(BranchPointIndices.Num() > 1))
+			{
+				return EForEachResult::Continue;
+			}
 
 			const bool bIsTrunk = IsTrunk(InAttributes.BranchParentNumber, BranchIndex);
 
diff --git a/Engine/Plugins/Experimental/ProceduralVegetationEditor/Source/ProceduralVegetation/Private/Helpers/PVPlantTraversalHelper.h b/Engine/Plugins/Experimental/ProceduralVegetationEditor/Source/ProceduralVegetation/Private/Helpers/PVPlantTraversalHelper.h
--- a/Engine/Plugins/Experimental/ProceduralVegetationEditor/Source/ProceduralVegetation/Private/Helpers/PVPlantTraversalHelper.h
+++ b/Engine/Plugins/Experimental/ProceduralVegetationEditor/Source/ProceduralVegetation/Private/Helpers/PVPlantTraversalHelper.h
@@ -114,7 +114,10 @@ namespace PV::PlantTraversalHelper
 			for (int32 ChildBranchNumber : BranchChildren)
 			{
 				const int32 ChildBranchIndex = BranchNumberAttribute.Find(ChildBranchNumber);
-				check(ChildBranchIndex != INDEX_NONE);
+				if (!ensure(ChildBranchIndex != INDEX_NONE))
+				{
+					continue;
+				}
 
 				if (BranchParentNumberAttribute[ChildBranchIndex] != BranchNumberAttribute[BranchIndex])
 				{
