Add visualization to PVE StaticMesh Importer node

This CL adds a custom visualization output to the StaticMesh importer node which shows the material-filtered mesh, the identified tips on the mesh, and the branch curves which were successfully extracted.

Key changes:
- Output separate debug, visualization, and growth data collections
- Add DepthPriorityGroup to debug visualization settings
- Move DynamicMesh specific functionality to new helper library (PVDynamicMeshHelpers.h)
- Move all import related attribute names to PVImportCommpon.h so they can more easily be shared
- Fixes bug which caused the node to continuously re-execute if an invalid material was added by the user.

#Synced-CL 51372802

diff --git a/Engine/Plugins/Experimental/ProceduralVegetationEditor/Source/ProceduralVegetation/Private/Helpers/PVDynamicMeshHelpers.cpp b/Engine/Plugins/Experimental/ProceduralVegetationEditor/Source/ProceduralVegetation/Private/Helpers/PVDynamicMeshHelpers.cpp
new file mode 100644
--- /dev/null
+++ b/Engine/Plugins/Experimental/ProceduralVegetationEditor/Source/ProceduralVegetation/Private/Helpers/PVDynamicMeshHelpers.cpp
@@ -0,0 +1,73 @@
+// Copyright Epic Games, Inc. All Rights Reserved.
+
+#include "Helpers/PVDynamicMeshHelpers.h"
+
+#include "DynamicMesh/DynamicMesh3.h"
+#include "DynamicMesh/DynamicMeshAABBTree3.h"
+
+TArray<int32> PV::DynamicMeshHelper::GetUniqueVertexIndices(const UE::Geometry::FDynamicMesh3& DynamicMesh, const TArray<int32>& TriangleIds)
+{
+	TSet<int32> UniqueVertexIndices;
+	UniqueVertexIndices.Reserve(TriangleIds.Num() * 3);
+	for (int32 tid : TriangleIds)
+	{
+		const UE::Geometry::FIndex3i VertexIndices = DynamicMesh.GetTriangle(tid);
+		UniqueVertexIndices.Add(VertexIndices[0]);
+		UniqueVertexIndices.Add(VertexIndices[1]);
+		UniqueVertexIndices.Add(VertexIndices[2]);
+	}
+
+	return UniqueVertexIndices.Array();
+}
+
+TArray<int32> PV::DynamicMeshHelper::GetVertexNeighbours(const UE::Geometry::FDynamicMesh3& DynamicMesh, int32 VertexID)
+{
+	TArray<int32> OutNeighbours;
+	OutNeighbours.Reserve(DynamicMesh.GetVtxEdgeCount(VertexID));
+	for (int32 EdgeID : DynamicMesh.VtxEdgesItr(VertexID))
+	{
+		const UE::Geometry::FIndex2i VertexIDs = DynamicMesh.GetEdgeV(EdgeID);
+		const int32 OtherVertexID = VertexIDs.A == VertexID ? VertexIDs.B : VertexIDs.A;
+		OutNeighbours.Add(OtherVertexID);
+	}
+	return OutNeighbours;
+}
+
+TArray<int32> PV::DynamicMeshHelper::FindVerticesInRadius(const UE::Geometry::FDynamicMeshAABBTree3& DynamicMeshAABBTree, const FVector& Pos, float Radius)
+{
+	TArray<int32> OutVertexIndices;
+
+	const UE::Geometry::FDynamicMesh3* DynamicMesh = DynamicMeshAABBTree.GetMesh();
+	if (DynamicMesh == nullptr)
+	{
+		return OutVertexIndices;
+	}
+
+	OutVertexIndices.Reserve(DynamicMesh->VertexCount());
+
+	const double RadiusSqr = Radius * Radius;
+
+	UE::Geometry::FDynamicMeshAABBTree3::FTreeTraversal TreeTraversal;
+
+	TreeTraversal.NextBoxF = [&](const UE::Geometry::FAxisAlignedBox3d& Box, int32 Depth)->bool
+	{
+		return Box.DistanceSquared(Pos) <= RadiusSqr;
+	};
+
+	TreeTraversal.NextTriangleF = [&](int32 TID)
+	{
+		const UE::Geometry::FIndex3i Triangle = DynamicMesh->GetTriangle(TID);
+		for (int32 j = 0; j < 3; ++j)
+		{
+			const double VertexDistSqr = (DynamicMesh->GetVertex(Triangle[j]) - Pos).SizeSquared();
+			if (VertexDistSqr <= RadiusSqr)
+			{
+				OutVertexIndices.AddUnique(Triangle[j]);
+			}
+		}
+	};
+
+	DynamicMeshAABBTree.DoTraversal(TreeTraversal);
+
+	return OutVertexIndices;
+}
\ No newline at end of file
diff --git a/Engine/Plugins/Experimental/ProceduralVegetationEditor/Source/ProceduralVegetation/Private/Helpers/PVDynamicMeshHelpers.h b/Engine/Plugins/Experimental/ProceduralVegetationEditor/Source/ProceduralVegetation/Private/Helpers/PVDynamicMeshHelpers.h
new file mode 100644
--- /dev/null
+++ b/Engine/Plugins/Experimental/ProceduralVegetationEditor/Source/ProceduralVegetation/Private/Helpers/PVDynamicMeshHelpers.h
@@ -0,0 +1,22 @@
+// Copyright Epic Games, Inc. All Rights Reserved.
+#pragma once
+
+#include "CoreMinimal.h"
+
+struct FManagedArrayCollection;
+
+namespace UE::Geometry
+{
+	class FDynamicMesh3;
+	template<class T> class TMeshAABBTree3;
+	using FDynamicMeshAABBTree3 = TMeshAABBTree3<FDynamicMesh3>;
+}
+
+namespace PV::DynamicMeshHelper
+{
+	TArray<int32> GetUniqueVertexIndices(const UE::Geometry::FDynamicMesh3& DynamicMesh, const TArray<int32>& TriangleIds);
+
+	TArray<int32> GetVertexNeighbours(const UE::Geometry::FDynamicMesh3& DynamicMesh, int32 VertexID);
+
+	TArray<int32> FindVerticesInRadius(const UE::Geometry::FDynamicMeshAABBTree3& DynamicMeshAABBTree, const FVector& Pos, float Radius);
+};
\ No newline at end of file
diff --git a/Engine/Plugins/Experimental/ProceduralVegetationEditor/Source/ProceduralVegetation/Private/Helpers/PVImportHelpers.cpp b/Engine/Plugins/Experimental/ProceduralVegetationEditor/Source/ProceduralVegetation/Private/Helpers/PVImportHelpers.cpp
--- a/Engine/Plugins/Experimental/ProceduralVegetationEditor/Source/ProceduralVegetation/Private/Helpers/PVImportHelpers.cpp
+++ b/Engine/Plugins/Experimental/ProceduralVegetationEditor/Source/ProceduralVegetation/Private/Helpers/PVImportHelpers.cpp
@@ -1,17 +1,17 @@
 // Copyright Epic Games, Inc. All Rights Reserved.
 #include "Helpers/PVImportHelpers.h"
 
-#include "ProceduralVegetationPreset.h"
 #include "DynamicMesh/DynamicMesh3.h"
 #include "DynamicMesh/DynamicMeshAABBTree3.h"
 #include "GeometryCollectionToDynamicMesh.h"
 #include "GeometryCollection/GeometryCollection.h"
 
 #include "Helpers/PVAttributesHelper.h"
+#include "Helpers/PVDynamicMeshHelpers.h"
 
 #include "Engine/StaticMesh.h"
 
-namespace PV::DynamicMeshHelper
+namespace PV::ImportHelper
 {
 	static void ComputeLengthFromRoot(
 		const UE::Geometry::FDynamicMesh3& DynamicMesh,
@@ -46,7 +46,7 @@ namespace PV::DynamicMeshHelper
 			const float PrevLengthFromRoot = LengthFromRootAttribute.GetValue(VertexID);
 			const FVector PrevVertexPosition = DynamicMesh.GetVertex(VertexID);
 
-			const TArray<int32> Neighbours = GetVertexNeighbours(DynamicMesh, VertexID);
+			const TArray<int32> Neighbours = PV::DynamicMeshHelper::GetVertexNeighbours(DynamicMesh, VertexID);
 			for (int32 NeighbourVertexID : Neighbours)
 			{
 				if (!EndPointsAttribute.GetValue(NeighbourVertexID) // Ingore vertex filter for end points as that means we've reached the end of the branch
@@ -306,8 +306,6 @@ namespace PV::DynamicMeshHelper
 		}
 	}
 
-	
-
 	const static TSet<FName> GeometryGroups = {
 		FGeometryCollection::VerticesGroup,
 		FGeometryCollection::FacesGroup,
@@ -379,115 +377,139 @@ namespace PV::DynamicMeshHelper
 	}
 }
 
-const PV::TDynamicMeshVertexAttributeDefinition<float> PV::DynamicMeshHelper::VertexScaleAttributeDefinition(TEXT("VertexScale"));
-const PV::TDynamicMeshVertexAttributeDefinition<bool> PV::DynamicMeshHelper::EndPointAttributeDefinition(TEXT("EndPoint"));
-const PV::TDynamicMeshVertexAttributeDefinition<bool> PV::DynamicMeshHelper::RootPointAttributeDefinition(TEXT("RootPoint"));
-const PV::TDynamicMeshVertexAttributeDefinition<float> PV::DynamicMeshHelper::LengthFromRootAttributeDefinition(TEXT("LengthFromRoot"));
-const PV::TDynamicMeshVertexAttributeDefinition<int32> PV::DynamicMeshHelper::NextBranchPointAttributeDefinition(TEXT("NextBranchPoint"));
+const PV::TDynamicMeshVertexAttributeDefinition<float> PV::ImportHelper::VertexScaleAttributeDefinition(TEXT("VertexScale"));
+const PV::TDynamicMeshVertexAttributeDefinition<bool> PV::ImportHelper::EndPointAttributeDefinition(TEXT("EndPoint"));
+const PV::TDynamicMeshVertexAttributeDefinition<bool> PV::ImportHelper::RootPointAttributeDefinition(TEXT("RootPoint"));
+const PV::TDynamicMeshVertexAttributeDefinition<float> PV::ImportHelper::LengthFromRootAttributeDefinition(TEXT("LengthFromRoot"));
+const PV::TDynamicMeshVertexAttributeDefinition<int32> PV::ImportHelper::NextBranchPointAttributeDefinition(TEXT("NextBranchPoint"));
 
-void PV::DynamicMeshHelper::AddDynamicMeshToCollection(const UE::Geometry::FDynamicMesh3& DynamicMesh, FManagedArrayCollection& InOutCollection, const FTransform& Offset)
+void PV::ImportHelper::AddDynamicMeshToCollection(
+	const UE::Geometry::FDynamicMesh3& DynamicMesh, 
+	FManagedArrayCollection& InOutCollection, 
+	const FTransform& Offset
+)
 {
 	FGeometryCollection MeshGeomCollection;
 
 	UE::Geometry::FGeometryCollectionToDynamicMeshes::FToCollectionOptions Options2;
-	UE::Geometry::FGeometryCollectionToDynamicMeshes::AppendMeshToCollection(MeshGeomCollection, DynamicMesh, Offset, Options2);
+	UE::Geometry::FGeometryCollectionToDynamicMeshes::AppendMeshToCollection(MeshGeomCollection, DynamicMesh, Offset.Inverse(), Options2);
 
 	AppendGeometryFromManagedArrayCollection(MeshGeomCollection, InOutCollection);
 	SetManagedArrayCollectionGeometry(MeshGeomCollection, InOutCollection);
 }
 
-void PV::DynamicMeshHelper::AddBranchHierarchiesToCollection(const TArray<FPVBranchHierarchyDescription>& BranchHierarchyDescriptions, FManagedArrayCollection& OutCollection)
+void PV::ImportHelper::AddDynamicMeshVerticesToCollection(
+	const UE::Geometry::FDynamicMesh3& DynamicMesh,
+	const FBooleanVertexAttribute& Filter,
+	const FVector& VertexOffset,
+	FManagedArrayCollection& InOutCollection,
+	const FName TargetGroupName,
+	const FName TargetAttributeName
+)
 {
-	const static FName CurveAttributeName = TEXT("Curve");
-	const static FName BranchCurvesGroupName = TEXT("BranchCurves");
-	OutCollection.AddGroup(BranchCurvesGroupName);
+	const TArray<int32> VertexIndices = Filter.FindAllNonZero();
+
+	if (!ensure(!InOutCollection.HasAttribute(TargetAttributeName, TargetGroupName)))
+	{
+		return;
+	}
+
+	if (!InOutCollection.HasGroup(TargetGroupName))
+	{
+		InOutCollection.AddGroup(TargetGroupName);
+	}
+	
+	InOutCollection.AddElements(VertexIndices.Num(), TargetGroupName);
+
+	auto& Attribute = InOutCollection.AddAttribute<FVector3f>(TargetAttributeName, TargetGroupName);
+	for (int32 i = 0; i < VertexIndices.Num(); ++i)
+	{
+		Attribute[i] = FVector3f(DynamicMesh.GetVertex(VertexIndices[i]) + VertexOffset);
+	}
+}
+
+void PV::ImportHelper::AddBranchHierarchiesToCollection(
+	const TArray<FPVBranchHierarchyDescription>& BranchHierarchyDescriptions,
+	FManagedArrayCollection& InOutCollection,
+	const FVector& Offset,
+	const FName TargetAttributeName,
+	const FName TargetGroupName
+)
+{
+	InOutCollection.AddGroup(TargetGroupName);
 	int32 AttributeIndex = 0;
 	for (const FPVBranchHierarchyDescription& BranchHierarchy : BranchHierarchyDescriptions)
 	{
-		OutCollection.AddElements(BranchHierarchy.Branches.Num(), BranchCurvesGroupName);
+		InOutCollection.AddElements(BranchHierarchy.Branches.Num(), TargetGroupName);
 
-		TManagedArray<TArray<FVector3f>>& CurveAttribute = *OutCollection.FindOrAddAttributeTyped<TArray<FVector3f>>(CurveAttributeName, BranchCurvesGroupName);
+		TManagedArray<TArray<FVector3f>>& Attribute = *InOutCollection.FindOrAddAttributeTyped<TArray<FVector3f>>(TargetAttributeName, TargetGroupName);
 		for (int32 i = 0; i < BranchHierarchy.Branches.Num(); ++i)
 		{
 			const FPVBranchDescription& Branch = BranchHierarchy.Branches[i];
-			CurveAttribute[AttributeIndex].Reserve(Branch.PointIndices.Num());
+			Attribute[AttributeIndex].Reserve(Branch.PointIndices.Num());
 			for (int32 VertexIndex : Branch.PointIndices)
 			{
-				CurveAttribute[AttributeIndex].Add(BranchHierarchy.Points[VertexIndex]);
+				Attribute[AttributeIndex].Add(BranchHierarchy.Points[VertexIndex] + FVector3f(Offset));
 			}
 			AttributeIndex++;
 		}
 	}
 }
 
-TArray<int32> PV::DynamicMeshHelper::GetUniqueVertexIndices(const UE::Geometry::FDynamicMesh3& DynamicMesh, const TArray<int32>& TriangleIds)
+void PV::ImportHelper::AddLabelsToCollection(const TArray<FLabel>& Labels, FManagedArrayCollection& InOutCollection)
 {
-	TSet<int32> UniqueVertexIndices;
-	UniqueVertexIndices.Reserve(TriangleIds.Num() * 3);
-	for (int32 tid : TriangleIds)
+	if (!InOutCollection.HasGroup(PVImportNames::LabelsGroup))
 	{
-		const UE::Geometry::FIndex3i VertexIndices = DynamicMesh.GetTriangle(tid);
-		UniqueVertexIndices.Add(VertexIndices[0]);
-		UniqueVertexIndices.Add(VertexIndices[1]);
-		UniqueVertexIndices.Add(VertexIndices[2]);
+		InOutCollection.AddGroup(PVImportNames::LabelsGroup);
 	}
 
-	return UniqueVertexIndices.Array();
-}
-
-TArray<int32> PV::DynamicMeshHelper::GetVertexNeighbours(const UE::Geometry::FDynamicMesh3& DynamicMesh, int32 VertexID)
-{
-	TArray<int32> OutNeighbours;
-	OutNeighbours.Reserve(DynamicMesh.GetVtxEdgeCount(VertexID));
-	for (int32 EdgeID : DynamicMesh.VtxEdgesItr(VertexID))
+	if (!InOutCollection.HasAttribute(PVImportNames::LabelTextAttribute, PVImportNames::LabelsGroup))
 	{
-		const UE::Geometry::FIndex2i VertexIDs = DynamicMesh.GetEdgeV(EdgeID);
-		const int32 OtherVertexID = VertexIDs.A == VertexID ? VertexIDs.B : VertexIDs.A;
-		OutNeighbours.Add(OtherVertexID);
+		InOutCollection.AddAttribute<FString>(PVImportNames::LabelTextAttribute, PVImportNames::LabelsGroup);
 	}
-	return OutNeighbours;
-}
-
-TArray<int32> PV::DynamicMeshHelper::FindVerticesInRadius(const UE::Geometry::FDynamicMeshAABBTree3& DynamicMeshAABBTree, const FVector& Pos, float Radius)
-{
-	TArray<int32> OutVertexIndices;
 
-	const UE::Geometry::FDynamicMesh3* DynamicMesh = DynamicMeshAABBTree.GetMesh();
-	if (DynamicMesh == nullptr)
+	if (!InOutCollection.HasAttribute(PVImportNames::LabelPositionAttribute, PVImportNames::LabelsGroup))
 	{
-		return OutVertexIndices;
+		InOutCollection.AddAttribute<FVector3f>(PVImportNames::LabelPositionAttribute, PVImportNames::LabelsGroup);
 	}
 
-	OutVertexIndices.Reserve(DynamicMesh->VertexCount());
-
-	const double RadiusSqr = Radius * Radius;
+	if (!InOutCollection.HasAttribute(PVImportNames::LabelScaleAttribute, PVImportNames::LabelsGroup))
+	{
+		InOutCollection.AddAttribute<float>(PVImportNames::LabelScaleAttribute, PVImportNames::LabelsGroup);
+	}
 
-	UE::Geometry::FDynamicMeshAABBTree3::FTreeTraversal TreeTraversal;
+	if (!ensure(InOutCollection.GetAttributeType(PVImportNames::LabelTextAttribute, PVImportNames::LabelsGroup) == FManagedArrayCollection::EArrayType::FStringType))
+	{
+		return;
+	}
 
-	TreeTraversal.NextBoxF = [&](const UE::Geometry::FAxisAlignedBox3d& Box, int32 Depth)->bool
+	if (!ensure(InOutCollection.GetAttributeType(PVImportNames::LabelPositionAttribute, PVImportNames::LabelsGroup) == FManagedArrayCollection::EArrayType::FVectorType))
 	{
-		return Box.DistanceSquared(Pos) <= RadiusSqr;
-	};
+		return;
+	}
 
-	TreeTraversal.NextTriangleF = [&](int32 TID)
+	if (!ensure(InOutCollection.GetAttributeType(PVImportNames::LabelScaleAttribute, PVImportNames::LabelsGroup) == FManagedArrayCollection::EArrayType::FFloatType))
 	{
-		const UE::Geometry::FIndex3i Triangle = DynamicMesh->GetTriangle(TID);
-		for (int32 j = 0; j < 3; ++j)
-		{
-			const double VertexDistSqr = (DynamicMesh->GetVertex(Triangle[j]) - Pos).SizeSquared();
-			if (VertexDistSqr <= RadiusSqr)
-			{
-				OutVertexIndices.AddUnique(Triangle[j]);
-			}
-		}
-	};
+		return;
+	}
+
+	const int32 StartIndex = InOutCollection.NumElements(PVImportNames::LabelsGroup);
+
+	InOutCollection.AddElements(Labels.Num(), PVImportNames::LabelsGroup);
 
-	DynamicMeshAABBTree.DoTraversal(TreeTraversal);
+	TManagedArray<FString>& LabelTextAttribute = *InOutCollection.FindAttribute<FString>(PVImportNames::LabelTextAttribute, PVImportNames::LabelsGroup);
+	TManagedArray<FVector3f>& LabelPositionAttribute = *InOutCollection.FindAttribute<FVector3f>(PVImportNames::LabelPositionAttribute, PVImportNames::LabelsGroup);
+	TManagedArray<float>& LabelScaleAttribute = *InOutCollection.FindAttribute<float>(PVImportNames::LabelScaleAttribute, PVImportNames::LabelsGroup);
 
-	return OutVertexIndices;
+	for (int32 i = 0; i < Labels.Num(); ++i)
+	{
+		LabelPositionAttribute[StartIndex + i] = Labels[i].Position;
+		LabelTextAttribute[StartIndex + i] = Labels[i].Text;
+		LabelScaleAttribute[StartIndex + i] = Labels[i].Scale;
+	}
 }
 
-void PV::DynamicMeshHelper::ComputeDynamicMeshBranchHierarchies(
+void PV::ImportHelper::ComputeDynamicMeshBranchHierarchies(
 	const UE::Geometry::FDynamicMesh3& DynamicMesh,
 	const FBooleanVertexAttribute& RootPointsAttribute,
 	const FBooleanVertexAttribute& EndPointsAttribute,
@@ -524,7 +546,7 @@ void PV::DynamicMeshHelper::ComputeDynamicMeshBranchHierarchies(
 	);
 }
 
-void PV::DynamicMeshHelper::ConvertDynamicMeshBranchHierarchiesToBranchDescription(
+void PV::ImportHelper::ConvertDynamicMeshBranchHierarchiesToBranchDescription(
 	const UE::Geometry::FDynamicMesh3& DynamicMesh,
 	const FFloatVertexAttribute& VertexScaleAttribute,
 	const FBooleanVertexAttribute& EndPointAttribute,
@@ -909,7 +931,10 @@ void PV::ImportHelper::ConnectDisconnectedBranches(FPVBranchHierarchyDescription
 			continue;
 		}
 
-		check(Branch.PointIndices.Num() > 1);
+		if (!ensure(Branch.PointIndices.Num() > 1))
+		{
+			continue;
+		}
 		
 		const FPVBranchDescription& ParentBranch = InOutBranchHierarchy.Branches[Branch.ParentBranchIndex];
 		
diff --git a/Engine/Plugins/Experimental/ProceduralVegetationEditor/Source/ProceduralVegetation/Private/Helpers/PVImportHelpers.h b/Engine/Plugins/Experimental/ProceduralVegetationEditor/Source/ProceduralVegetation/Private/Helpers/PVImportHelpers.h
--- a/Engine/Plugins/Experimental/ProceduralVegetationEditor/Source/ProceduralVegetation/Private/Helpers/PVImportHelpers.h
+++ b/Engine/Plugins/Experimental/ProceduralVegetationEditor/Source/ProceduralVegetation/Private/Helpers/PVImportHelpers.h
@@ -2,27 +2,11 @@
 #pragma once
 
 #include "CoreMinimal.h"
-
-#include "Utils/PVAttributes.h"
+#include "PVImportCommon.h"
 #include "Utils/PVDynamicMeshVertexAttribute.h"
-
-namespace PV::Facades
-{
-	class FPointFacade;
-	class FBranchFacade;
-	class FPlantFacade;
-};
+#include "Utils/PVAttributes.h"
 
 struct FManagedArrayCollection;
-class UProceduralVegetationPreset;
-class UPackage;
-
-namespace PV
-{
-	template<typename T> class TDynamicMeshVertexAttributeDefinition;
-	template<typename AttribValueType> class TDynamicMeshVertexAttributeExt;
-	template<typename AttribValueType> class TVertexAttributeView;
-}
 
 namespace UE::Geometry
 {
@@ -65,7 +49,7 @@ struct FPVBranchHierarchyDescription
 	EValidateHierarchyResult ValidateHierarchy() const;
 };
 
-namespace PV::DynamicMeshHelper
+namespace PV::ImportHelper
 {
 	using FBooleanVertexAttribute = TDynamicMeshVertexAttributeExt<bool>;
 	using FIntVertexAttribute     = TDynamicMeshVertexAttributeExt<int32>;
@@ -77,16 +61,37 @@ namespace PV::DynamicMeshHelper
 	extern const PV::TDynamicMeshVertexAttributeDefinition<float> LengthFromRootAttributeDefinition;
 	extern const PV::TDynamicMeshVertexAttributeDefinition<int32> NextBranchPointAttributeDefinition;
 
-	void AddDynamicMeshToCollection(const UE::Geometry::FDynamicMesh3& DynamicMesh, FManagedArrayCollection& InOutCollection, const FTransform& Offset = FTransform::Identity);
-	
-	void AddBranchHierarchiesToCollection(const TArray<FPVBranchHierarchyDescription>& BranchHierarchyDescriptions, FManagedArrayCollection& OutCollection);
+	void AddDynamicMeshToCollection(
+		const UE::Geometry::FDynamicMesh3& DynamicMesh, 
+		FManagedArrayCollection& InOutCollection, 
+		const FTransform& Offset = FTransform::Identity
+	);
 
-	TArray<int32> GetUniqueVertexIndices(const UE::Geometry::FDynamicMesh3& DynamicMesh, const TArray<int32>& TriangleIds);
+	void AddDynamicMeshVerticesToCollection(
+		const UE::Geometry::FDynamicMesh3& DynamicMesh,
+		const FBooleanVertexAttribute& Filter,
+		const FVector& VertexOffset,
+		FManagedArrayCollection& InOutCollection,
+		const FName TargetGroupName,
+		const FName TargetAttributeName
+	);
+	
+	void AddBranchHierarchiesToCollection(
+		const TArray<FPVBranchHierarchyDescription>& BranchHierarchyDescriptions,
+		FManagedArrayCollection& InOutCollection,
+		const FVector& Offset = FVector::ZeroVector,
+		const FName TargetAttributeName = PVImportNames::BranchHierarchyAttribute,
+		const FName TargetGroupName = PVImportNames::BranchHierarchyGroup
+	);
 
-	TArray<int32> GetVertexNeighbours(const UE::Geometry::FDynamicMesh3& DynamicMesh, int32 VertexID);
+	struct FLabel
+	{
+		FVector3f Position;
+		FString Text;
+		float Scale;
+	};
+	void AddLabelsToCollection(const TArray<FLabel>& Labels, FManagedArrayCollection& InOutCollection);
 
-	TArray<int32> FindVerticesInRadius(const UE::Geometry::FDynamicMeshAABBTree3& DynamicMeshAABBTree, const FVector& Pos, float Radius);
-	
 	struct FMeshBranchHierarchy
 	{
 		struct FBranch
@@ -116,10 +121,7 @@ namespace PV::DynamicMeshHelper
 		const TArray<FMeshBranchHierarchy>& InMeshBranchHierarchies,
 		TArray<FPVBranchHierarchyDescription>& OutBranchHierarchyDescriptions
 	);
-}
 
-namespace PV::ImportHelper
-{
 	// Creates a collection containing all attributes (of the supplied size) required for creating PV growth data.
 	void CreateEmptyGrowthData(
 		FManagedArrayCollection& OutCollection,
diff --git a/Engine/Plugins/Experimental/ProceduralVegetationEditor/Source/ProceduralVegetation/Private/Implementations/PVImporter_StaticMesh.cpp b/Engine/Plugins/Experimental/ProceduralVegetationEditor/Source/ProceduralVegetation/Private/Implementations/PVImporter_StaticMesh.cpp
--- a/Engine/Plugins/Experimental/ProceduralVegetationEditor/Source/ProceduralVegetation/Private/Implementations/PVImporter_StaticMesh.cpp
+++ b/Engine/Plugins/Experimental/ProceduralVegetationEditor/Source/ProceduralVegetation/Private/Implementations/PVImporter_StaticMesh.cpp
@@ -3,6 +3,7 @@
 #include "PVImporter_StaticMesh.h"
 
 #include "Helpers/PVImportHelpers.h"
+#include "Helpers/PVDynamicMeshHelpers.h"
 #include "Utils/PVDynamicMeshVertexAttribute.h"
 
 #include "Engine/StaticMesh.h"
@@ -12,44 +13,31 @@
 #include "ConversionUtils/SceneComponentToDynamicMesh.h"
 #include "Selections/MeshConnectedComponents.h"
 
-using FBooleanVertexAttribute = PV::DynamicMeshHelper::FBooleanVertexAttribute;
-using FIntVertexAttribute = PV::DynamicMeshHelper::FIntVertexAttribute;
-using FFloatVertexAttribute = PV::DynamicMeshHelper::FFloatVertexAttribute;
+namespace PV::StaticMeshImport
+{
+using FBooleanVertexAttribute = PV::ImportHelper::FBooleanVertexAttribute;
+using FIntVertexAttribute = PV::ImportHelper::FIntVertexAttribute;
+using FFloatVertexAttribute = PV::ImportHelper::FFloatVertexAttribute;
+
+using FIndex2i = UE::Geometry::FIndex2i;
+using FIndex3i = UE::Geometry::FIndex3i;
+using FDynamicMesh3 = UE::Geometry::FDynamicMesh3;
+using FAxisAlignedBox3d = UE::Geometry::FAxisAlignedBox3d;
+using FDynamicMeshAABBTree3 = UE::Geometry::FDynamicMeshAABBTree3;
+using FMeshConnectedComponents = UE::Geometry::FMeshConnectedComponents;
+using FDynamicMeshQueries3 = UE::Geometry::TMeshQueries<FDynamicMesh3>;
+using FDistPoint3Triangle3d = UE::Geometry::FDistPoint3Triangle3d;
 
 const static PV::TDynamicMeshVertexAttributeDefinition<bool> FillerVertexAttributeDefinition(TEXT("FillerVertex"));
 const static PV::TDynamicMeshVertexAttributeDefinition<int32> FuseToAttributeDefinition(TEXT("FuseTo"));
 
-static void AddBranchCurvesToCollection(FManagedArrayCollection& InCollection, const TArray<TArray<int32>>& BranchCurves, const TFunction<FVector3f(int32)>& GetVertexPosFunc)
-{
-	const static FName BranchCurveGroupName = TEXT("BranchCurves");
-	const static FName CurveAttributeName = TEXT("Curve");
-
-	if (!InCollection.HasGroup(BranchCurveGroupName))
-	{
-		InCollection.AddGroup(BranchCurveGroupName);
-	}
-
-	InCollection.AddElements(BranchCurves.Num(), BranchCurveGroupName);
-
-	TManagedArray<TArray<FVector3f>>& CurveAttribute = InCollection.AddAttribute<TArray<FVector3f>>(CurveAttributeName, BranchCurveGroupName);
-	for (int32 i = 0; i < BranchCurves.Num(); ++i)
-	{
-		const TArray<int32>& BranchCurve = BranchCurves[i];
-		CurveAttribute[i].Reserve(BranchCurve.Num());
-		for (int32 VertexIndex : BranchCurve)
-		{
-			CurveAttribute[i].Add(GetVertexPosFunc(VertexIndex));
-		}
-	}
-}
-
 static FVector GenerateFillerVertex(const FVector& Pos1, const FVector& Pos2)
 {
 	// Since we don't have edges for the DynamicMesh we add a new vertex at the center (offset by 0.1mm) whenever we need to represent a line
 	return ((Pos1 + Pos2) * 0.5)  + (Pos1 - Pos2).GetSafeNormal().Cross(FVector::ForwardVector).Cross(FVector::UpVector)/* * 0.01*/;
 }
 
-static void RemoveUnusedVertices(UE::Geometry::FDynamicMesh3& DynamicMesh)
+static void RemoveUnusedVertices(FDynamicMesh3& DynamicMesh)
 {
 	TRACE_CPUPROFILER_EVENT_SCOPE(OrphanRemoval);
 
@@ -57,7 +45,7 @@ static void RemoveUnusedVertices(UE::Geometry::FDynamicMesh3& DynamicMesh)
 	DynamicMesh.CompactInPlace();
 }
 
-static bool FilterTrianglesUsingMaterials(UE::Geometry::FDynamicMesh3& DynamicMesh, const UStaticMesh* StaticMesh, const TArray<FName>& MaterialsToKeep)
+static bool FilterTrianglesUsingMaterials(FDynamicMesh3& DynamicMesh, const UStaticMesh* StaticMesh, const TArray<FName>& MaterialsToKeep)
 {
 	TRACE_CPUPROFILER_EVENT_SCOPE(KeepMaterial);
 
@@ -100,7 +88,7 @@ static bool FilterTrianglesUsingMaterials(UE::Geometry::FDynamicMesh3& DynamicMe
 	return true;
 }
 
-static void ComputeNormals(UE::Geometry::FDynamicMesh3& DynamicMesh)
+static void ComputeNormals(FDynamicMesh3& DynamicMesh)
 {
 	TRACE_CPUPROFILER_EVENT_SCOPE(ComputeNormals);
 
@@ -128,7 +116,7 @@ static void ComputeNormals(UE::Geometry::FDynamicMesh3& DynamicMesh)
 }
 
 static void ComputeVerticesCenterAndRadius_Internal_V1(
-	UE::Geometry::FDynamicMesh3& DynamicMesh, 
+	FDynamicMesh3& DynamicMesh, 
 	const TArray<int32>& TriangleIds,
 	FFloatVertexAttribute& VertexScaleAttribute
 )
@@ -142,7 +130,7 @@ static void ComputeVerticesCenterAndRadius_Internal_V1(
 	
 	const TArray<int32> ComponentVertexIndices = PV::DynamicMeshHelper::GetUniqueVertexIndices(DynamicMesh, TriangleIds);
 
-	UE::Geometry::FDynamicMeshAABBTree3 DynamicMeshAABBTree(&DynamicMesh, false);
+	FDynamicMeshAABBTree3 DynamicMeshAABBTree(&DynamicMesh, false);
 	DynamicMeshAABBTree.Build(TriangleIds);
 
 	const static auto ComputeRadius = [](const FVector& Point, const FVector& Normal, const FVector& QueryPosition)->double
@@ -204,7 +192,7 @@ static void ComputeVerticesCenterAndRadius_Internal_V1(
 			continue;
 		}
 
-		const UE::Geometry::FDistPoint3Triangle3d Query = UE::Geometry::TMeshQueries<UE::Geometry::FDynamicMesh3>::TriangleDistance(DynamicMesh, NearTriID, SphereCenter);
+		const FDistPoint3Triangle3d Query = FDynamicMeshQueries3::TriangleDistance(DynamicMesh, NearTriID, SphereCenter);
 		const FVector& ClosestTrianglePoint = Query.ClosestTrianglePoint;
 		const FVector ClosestTrianglePointNormal = DynamicMesh.GetTriBaryNormal(NearTriID, Query.TriangleBaryCoords[0], Query.TriangleBaryCoords[1], Query.TriangleBaryCoords[2]);
 
@@ -228,7 +216,7 @@ static void ComputeVerticesCenterAndRadius_Internal_V1(
 }
 
 static void ComputeVerticesCenterAndRadius_Internal_V2(
-	UE::Geometry::FDynamicMesh3& DynamicMesh, 
+	FDynamicMesh3& DynamicMesh, 
 	const TArray<int32>& TriangleIds,
 	FFloatVertexAttribute& VertexScaleAttribute
 )
@@ -254,7 +242,7 @@ static void ComputeVerticesCenterAndRadius_Internal_V2(
 
 	const TArray<int32> ComponentVertexIndices = PV::DynamicMeshHelper::GetUniqueVertexIndices(DynamicMesh, TriangleIds);
 
-	UE::Geometry::FDynamicMeshAABBTree3 DynamicMeshAABBTree(&DynamicMesh, false);
+	FDynamicMeshAABBTree3 DynamicMeshAABBTree(&DynamicMesh, false);
 	DynamicMeshAABBTree.Build(TriangleIds);
 
 	TArray<FVector> NewVertexPositions;
@@ -279,7 +267,7 @@ static void ComputeVerticesCenterAndRadius_Internal_V2(
 		double HitTriangleDistance = 0;
 		int HitTriangleID = INDEX_NONE;
 		FVector3d HitBaryCoords = FVector3d::Zero();
-		UE::Geometry::FDynamicMeshAABBTree3::FQueryOptions QueryOptions;
+		FDynamicMeshAABBTree3::FQueryOptions QueryOptions;
 		QueryOptions.TriangleFilterF = [&](int32 TriangleID) { return !ThisVertexTriangleIDs.Contains(TriangleID); }; // Make sure to ignore triangles which includes this vertex, otherise we will instantly get a hit on ourselves
 		DynamicMeshAABBTree.FindNearestHitTriangle(Ray, HitTriangleDistance, HitTriangleID, HitBaryCoords, QueryOptions);
 
@@ -398,7 +386,7 @@ static TAutoConsoleVariable<int32> ComputeVerticesCenterVersion(
 );
 
 static void ComputeVerticesCenterAndRadius_Internal(
-	UE::Geometry::FDynamicMesh3& DynamicMesh, 
+	FDynamicMesh3& DynamicMesh, 
 	const TArray<int32>& TriangleIds,
 	FFloatVertexAttribute& VertexScaleAttribute
 )
@@ -416,13 +404,13 @@ static void ComputeVerticesCenterAndRadius_Internal(
 	}
 }
 
-static void ComputeVerticesCenterAndRadius(UE::Geometry::FDynamicMesh3& DynamicMesh, FFloatVertexAttribute& VertexScaleAttribute)
+static void ComputeVerticesCenterAndRadius(FDynamicMesh3& DynamicMesh, FFloatVertexAttribute& VertexScaleAttribute)
 {
 	TRACE_CPUPROFILER_EVENT_SCOPE(ComputeVerticesCenterAndRadius);
 
 	// We execute this per component (detached mesh) as that often correlates to individual branches which significantly improves
 	// the results as we won't accidentally find an opposing vertex in a different branch from the one we're currently operating on. 
-	UE::Geometry::FMeshConnectedComponents ConnectedComponents(&DynamicMesh);
+	FMeshConnectedComponents ConnectedComponents(&DynamicMesh);
 	ConnectedComponents.FindConnectedTriangles(); 
 
 	ParallelFor(ConnectedComponents.Num(), [&](int32 Index)
@@ -432,7 +420,7 @@ static void ComputeVerticesCenterAndRadius(UE::Geometry::FDynamicMesh3& DynamicM
 }
 
 static void ExtractReductionCloud_Internal(
-	const UE::Geometry::FDynamicMesh3& DynamicMesh, 
+	const FDynamicMesh3& DynamicMesh, 
 	const TArray<int32>& TriangleIds, 
 	float SearchRadiusInflate,
 	const FFloatVertexAttribute& VertexScaleAttribute,
@@ -441,7 +429,7 @@ static void ExtractReductionCloud_Internal(
 {
 	TRACE_CPUPROFILER_EVENT_SCOPE(ExtractReductionCloud_Internal);
 		
-	UE::Geometry::FDynamicMeshAABBTree3 DynamicMeshAABBTree(&DynamicMesh, false);
+	FDynamicMeshAABBTree3 DynamicMeshAABBTree(&DynamicMesh, false);
 	DynamicMeshAABBTree.Build(TriangleIds);
 
 	TArray<int32> VertexIndices = PV::DynamicMeshHelper::GetUniqueVertexIndices(DynamicMesh, TriangleIds);
@@ -523,7 +511,7 @@ static void ExtractReductionCloud_Internal(
 }
 	
 static void FuseVerticesUsingFuseAttribute(
-	UE::Geometry::FDynamicMesh3& DynamicMesh,
+	FDynamicMesh3& DynamicMesh,
 	const FIntVertexAttribute& FuseToAttribute,
 	FBooleanVertexAttribute& FillerVertexAttribute
 )
@@ -540,7 +528,7 @@ static void FuseVerticesUsingFuseAttribute(
 	for (int32 i = 0; i < TriangleIndices.Num(); i++)
 	{
 		const int32 TriangleIndex = TriangleIndices[i];
-		const UE::Geometry::FIndex3i VertexIndices = DynamicMesh.GetTriangle(TriangleIndex);
+		const FIndex3i VertexIndices = DynamicMesh.GetTriangle(TriangleIndex);
 
 		TArray<int32, TInlineAllocator<3>> NewTriangleIndices;
 		for (int32 j = 0; j < 3; j++) 
@@ -578,7 +566,7 @@ static void FuseVerticesUsingFuseAttribute(
 }
 
 static void CollapseBranchesToLine(
-	UE::Geometry::FDynamicMesh3& DynamicMesh,
+	FDynamicMesh3& DynamicMesh,
 	const FFloatVertexAttribute& VertexScaleAttribute,
 	FIntVertexAttribute& FuseToAttribute,
 	FBooleanVertexAttribute& FillerVertexAttribute
@@ -586,7 +574,7 @@ static void CollapseBranchesToLine(
 {
 	TRACE_CPUPROFILER_EVENT_SCOPE(CollapseBranchesToLine);
 
-	UE::Geometry::FMeshConnectedComponents ConnectedComponents(&DynamicMesh);
+	FMeshConnectedComponents ConnectedComponents(&DynamicMesh);
 	ConnectedComponents.FindConnectedTriangles();
 
 	ParallelFor(ConnectedComponents.Num(), [&](int32 Index)
@@ -598,7 +586,7 @@ static void CollapseBranchesToLine(
 }
 
 static void FindEndPoints(
-	const UE::Geometry::FDynamicMesh3& DynamicMesh,
+	const FDynamicMesh3& DynamicMesh,
 	FBooleanVertexAttribute& EndPointAttribute,
 	const FBooleanVertexAttribute& FillerVertexAttribute
 )
@@ -626,7 +614,7 @@ static void FindEndPoints(
 
 		for (int32 EdgeID : DynamicMesh.VtxEdgesItr(VertexIndex))
 		{
-			const UE::Geometry::FIndex2i VertexIDs = DynamicMesh.GetEdgeV(EdgeID);
+			const FIndex2i VertexIDs = DynamicMesh.GetEdgeV(EdgeID);
 			const int32 OtherVertexID = VertexIDs.A == VertexIndex ? VertexIDs.B : VertexIDs.A;
 
 			if (FillerVertexAttribute.GetValue(OtherVertexID))
@@ -791,7 +779,7 @@ static void ConnectBranchesToParent(
 }
 
 static void FindRootPoints(
-	const UE::Geometry::FDynamicMesh3& DynamicMesh,
+	const FDynamicMesh3& DynamicMesh,
 	FBooleanVertexAttribute& RootPointAttribute,
 	FBooleanVertexAttribute& EndPointAttribute
 )
@@ -801,7 +789,7 @@ static void FindRootPoints(
 	// This function currently just finds all the end points that are close to the bottom of the plant bounding box and converts them to root points.
 	// This could result in issues if a plant has multiple points touching the floor. 
 
-	const UE::Geometry::FAxisAlignedBox3d BoxBounds = DynamicMesh.GetBounds();
+	const FAxisAlignedBox3d BoxBounds = DynamicMesh.GetBounds();
 
 	for (int32 VertexIndex : EndPointAttribute.FindAllNonZero())
 	{
@@ -817,11 +805,13 @@ static void FindRootPoints(
 	}
 }
 
-PV::StaticMeshImport::EImportResult PV::StaticMeshImport::ImportGrowthDataFromStaticMesh(const FPVImportStaticMeshParams& InParams, FManagedArrayCollection& OutCollection)
+}
+
+PV::StaticMeshImport::EImportResult PV::StaticMeshImport::ImportGrowthDataFromStaticMesh(const FPVImportStaticMeshParams& InParams, FPVImportStaticMeshOutput& Output)
 {
-	using namespace PV::DynamicMeshHelper;
+	using namespace PV::ImportHelper;
 	
-	TRACE_CPUPROFILER_EVENT_SCOPE(PV::StaticMeshImport::ImportGrowthDataFromStaticMesh);
+	TRACE_CPUPROFILER_EVENT_SCOPE(ImportGrowthDataFromStaticMesh);
 	
 	if (!InParams.StaticMeshAsset)
 	{
@@ -852,26 +842,23 @@ PV::StaticMeshImport::EImportResult PV::StaticMeshImport::ImportGrowthDataFromSt
 
 	if (InParams.DebugState == EPVImportStaticMeshDebugState::InitialMesh)
 	{
-		AddDynamicMeshToCollection(DynamicMesh, OutCollection);
-		return EImportResult::Success;
+		FManagedArrayCollection& DebugCollection = Output.DebugCollection.Emplace();
+		AddDynamicMeshToCollection(DynamicMesh, DebugCollection);
 	}
 	
 	if (!FilterTrianglesUsingMaterials(DynamicMesh, InParams.StaticMeshAsset, InParams.MaterialsToKeep))
 	{
 		return EImportResult::InvalidMaterialFilter;
 	}
-	
-	if (InParams.DebugState == EPVImportStaticMeshDebugState::KeepMaterial)
-	{
-		AddDynamicMeshToCollection(DynamicMesh, OutCollection);
-		return EImportResult::Success;
-	}
+
+	const FVector VisualizationOffset = FVector(0, DynamicMesh.GetBounds().Extents().Y * 2, 0);
+	AddDynamicMeshToCollection(DynamicMesh, Output.VisualizationCollection, FTransform(VisualizationOffset));
 
 	ComputeNormals(DynamicMesh);
 	if (InParams.DebugState == EPVImportStaticMeshDebugState::ComputeNormals)
 	{
-		AddDynamicMeshToCollection(DynamicMesh, OutCollection);
-		return EImportResult::Success;
+		FManagedArrayCollection& DebugCollection = Output.DebugCollection.Emplace();
+		AddDynamicMeshToCollection(DynamicMesh, DebugCollection);
 	}
 
 	ComputeVerticesCenterAndRadius(
@@ -880,9 +867,9 @@ PV::StaticMeshImport::EImportResult PV::StaticMeshImport::ImportGrowthDataFromSt
 	);
 	if (InParams.DebugState == EPVImportStaticMeshDebugState::ComputeVerticesCenterAndRadius)
 	{
-		AddDynamicMeshToCollection(DynamicMesh, OutCollection);
-		VertexScaleAttributeDefinition.CopyAttributeToCollection(DynamicMesh, OutCollection);
-		return EImportResult::Success;
+		FManagedArrayCollection& DebugCollection = Output.DebugCollection.Emplace();
+		AddDynamicMeshToCollection(DynamicMesh, DebugCollection);
+		VertexScaleAttributeDefinition.CopyAttributeToCollection(DynamicMesh, DebugCollection);
 	}
 
 	CollapseBranchesToLine(
@@ -893,10 +880,10 @@ PV::StaticMeshImport::EImportResult PV::StaticMeshImport::ImportGrowthDataFromSt
 	);
 	if (InParams.DebugState == EPVImportStaticMeshDebugState::CollapseBranchesToLine)
 	{
-		AddDynamicMeshToCollection(DynamicMesh, OutCollection);
-		FuseToAttributeDefinition.CopyAttributeToCollection(DynamicMesh, OutCollection);
-		VertexScaleAttributeDefinition.CopyAttributeToCollection(DynamicMesh, OutCollection);
-		return EImportResult::Success;
+		FManagedArrayCollection& DebugCollection = Output.DebugCollection.Emplace();
+		AddDynamicMeshToCollection(DynamicMesh, DebugCollection);
+		FuseToAttributeDefinition.CopyAttributeToCollection(DynamicMesh, DebugCollection);
+		VertexScaleAttributeDefinition.CopyAttributeToCollection(DynamicMesh, DebugCollection);
 	}
 
 	FindEndPoints(
@@ -906,9 +893,9 @@ PV::StaticMeshImport::EImportResult PV::StaticMeshImport::ImportGrowthDataFromSt
 	);
 	if (InParams.DebugState == EPVImportStaticMeshDebugState::FindEndPoints)
 	{
-		AddDynamicMeshToCollection(DynamicMesh, OutCollection);
-		EndPointAttributeDefinition.CopyAttributeToCollection(DynamicMesh, OutCollection);
-		return EImportResult::Success;
+		FManagedArrayCollection& DebugCollection = Output.DebugCollection.Emplace();
+		AddDynamicMeshToCollection(DynamicMesh, DebugCollection);
+		EndPointAttributeDefinition.CopyAttributeToCollection(DynamicMesh, DebugCollection);
 	}
 
 	ConnectBranchesToParent(
@@ -919,9 +906,9 @@ PV::StaticMeshImport::EImportResult PV::StaticMeshImport::ImportGrowthDataFromSt
 	);
 	if (InParams.DebugState == EPVImportStaticMeshDebugState::ConnectBranchesToParent)
 	{
-		AddDynamicMeshToCollection(DynamicMesh, OutCollection);
-		EndPointAttributeDefinition.CopyAttributeToCollection(DynamicMesh, OutCollection);
-		return EImportResult::Success;
+		FManagedArrayCollection& DebugCollection = Output.DebugCollection.Emplace();
+		AddDynamicMeshToCollection(DynamicMesh, DebugCollection);
+		EndPointAttributeDefinition.CopyAttributeToCollection(DynamicMesh, DebugCollection);
 	}
 
 	FindRootPoints(
@@ -931,10 +918,10 @@ PV::StaticMeshImport::EImportResult PV::StaticMeshImport::ImportGrowthDataFromSt
 	);
 	if (InParams.DebugState == EPVImportStaticMeshDebugState::FindRootPoints)
 	{
-		AddDynamicMeshToCollection(DynamicMesh, OutCollection);
-		RootPointAttributeDefinition.CopyAttributeToCollection(DynamicMesh, OutCollection);
-		EndPointAttributeDefinition.CopyAttributeToCollection(DynamicMesh, OutCollection);
-		return EImportResult::Success;
+		FManagedArrayCollection& DebugCollection = Output.DebugCollection.Emplace();
+		AddDynamicMeshToCollection(DynamicMesh, DebugCollection);
+		RootPointAttributeDefinition.CopyAttributeToCollection(DynamicMesh, DebugCollection);
+		EndPointAttributeDefinition.CopyAttributeToCollection(DynamicMesh, DebugCollection);
 	}
 
 	auto& FillerVertexAttribute = FillerVertexAttributeDefinition.GetOrAttachAttribute(DynamicMesh);
@@ -961,21 +948,35 @@ PV::StaticMeshImport::EImportResult PV::StaticMeshImport::ImportGrowthDataFromSt
 
 	if (InParams.DebugState == EPVImportStaticMeshDebugState::ComputeBranchHierarchies)
 	{
-		AddDynamicMeshToCollection(DynamicMesh, OutCollection);
-		AddBranchHierarchiesToCollection(BranchHierarchyDescriptions, OutCollection);
-		RootPointAttributeDefinition.CopyAttributeToCollection(DynamicMesh, OutCollection);
-		EndPointAttributeDefinition.CopyAttributeToCollection(DynamicMesh, OutCollection);
-		VertexScaleAttributeDefinition.CopyAttributeToCollection(DynamicMesh, OutCollection);
-		LengthFromRootAttributeDefinition.CopyAttributeToCollection(DynamicMesh, OutCollection);
-		NextBranchPointAttributeDefinition.CopyAttributeToCollection(DynamicMesh, OutCollection);
-
-		return EImportResult::Success;
+		FManagedArrayCollection& DebugCollection = Output.DebugCollection.Emplace();
+		AddDynamicMeshToCollection(DynamicMesh, DebugCollection);
+		AddBranchHierarchiesToCollection(BranchHierarchyDescriptions, DebugCollection);
+		RootPointAttributeDefinition.CopyAttributeToCollection(DynamicMesh, DebugCollection);
+		EndPointAttributeDefinition.CopyAttributeToCollection(DynamicMesh, DebugCollection);
+		VertexScaleAttributeDefinition.CopyAttributeToCollection(DynamicMesh, DebugCollection);
+		LengthFromRootAttributeDefinition.CopyAttributeToCollection(DynamicMesh, DebugCollection);
+		NextBranchPointAttributeDefinition.CopyAttributeToCollection(DynamicMesh, DebugCollection);
 	}
 
-	if (!PV::ImportHelper::GenerateGrowthDataFromBranchHierarchies(OutCollection, BranchHierarchyDescriptions))
+	if (!PV::ImportHelper::GenerateGrowthDataFromBranchHierarchies(Output.GrowthDataCollection, BranchHierarchyDescriptions))
 	{
 		return EImportResult::FailedToGenerateGrowthData;
 	}
 
+	// Build visualization output
+	{
+		AddDynamicMeshVerticesToCollection(
+			DynamicMesh,
+			EndPointAttributeDefinition.GetOrAttachAttribute(DynamicMesh),
+			VisualizationOffset,
+			Output.VisualizationCollection,
+			PVImportNames::TipVisualizationGroup,
+			PVImportNames::TipVisualizationAttribute
+		);
+
+		AddBranchHierarchiesToCollection(BranchHierarchyDescriptions, Output.VisualizationCollection, VisualizationOffset);
+	}
+
 	return EImportResult::Success;
-}
\ No newline at end of file
+}
+
diff --git a/Engine/Plugins/Experimental/ProceduralVegetationEditor/Source/ProceduralVegetation/Private/Implementations/PVImporter_StaticMesh.h b/Engine/Plugins/Experimental/ProceduralVegetationEditor/Source/ProceduralVegetation/Private/Implementations/PVImporter_StaticMesh.h
--- a/Engine/Plugins/Experimental/ProceduralVegetationEditor/Source/ProceduralVegetation/Private/Implementations/PVImporter_StaticMesh.h
+++ b/Engine/Plugins/Experimental/ProceduralVegetationEditor/Source/ProceduralVegetation/Private/Implementations/PVImporter_StaticMesh.h
@@ -21,5 +21,5 @@ namespace PV::StaticMeshImport
 		FailedToGenerateGrowthData
 	};
 
-	EImportResult ImportGrowthDataFromStaticMesh(const FPVImportStaticMeshParams& InParams, FManagedArrayCollection& OutCollection);
+	EImportResult ImportGrowthDataFromStaticMesh(const FPVImportStaticMeshParams& InParams, FPVImportStaticMeshOutput& Output);
 };
diff --git a/Engine/Plugins/Experimental/ProceduralVegetationEditor/Source/ProceduralVegetation/Private/Implementations/PVImporter_Texture2D.cpp b/Engine/Plugins/Experimental/ProceduralVegetationEditor/Source/ProceduralVegetation/Private/Implementations/PVImporter_Texture2D.cpp
--- a/Engine/Plugins/Experimental/ProceduralVegetationEditor/Source/ProceduralVegetation/Private/Implementations/PVImporter_Texture2D.cpp
+++ b/Engine/Plugins/Experimental/ProceduralVegetationEditor/Source/ProceduralVegetation/Private/Implementations/PVImporter_Texture2D.cpp
@@ -1,7 +1,9 @@
 // Copyright Epic Games, Inc. All Rights Reserved.
 
 #include "Implementations/PVImporter_Texture2D.h"
+#include "PVImportCommon.h"
 #include "Helpers/PVImportHelpers.h"
+#include "Helpers/PVDynamicMeshHelpers.h"
 #include "Utils/PVDynamicMeshVertexAttribute.h"
 
 #include "IntVectorTypes.h"
@@ -25,10 +27,10 @@
 
 namespace Internal
 {
-	using FBooleanVertexAttribute = PV::DynamicMeshHelper::FBooleanVertexAttribute;
-	using FFloatVertexAttribute = PV::DynamicMeshHelper::FFloatVertexAttribute; 
+	using FBooleanVertexAttribute = PV::ImportHelper::FBooleanVertexAttribute;
+	using FFloatVertexAttribute = PV::ImportHelper::FFloatVertexAttribute; 
+	
 	using FIndex2i = UE::Geometry::FIndex2i;
-
 	using FIndex3i = UE::Geometry::FIndex3i;
 	using FVector2i = UE::Geometry::FVector2i;
 	using FMeshConnectedComponents = UE::Geometry::FMeshConnectedComponents;
@@ -38,18 +40,6 @@ namespace Internal
 	using FPolygon2f = UE::Geometry::FPolygon2f;
 	using FDelaunay2 = UE::Geometry::FDelaunay2;
 
-	const static FName PixelsDebugGroupName = TEXT("PixelDebug");
-	const static FName PixelPositionsAttributeName = TEXT("Position");
-	const static FName PixelValueAttributeName = TEXT("Value");
-	const static FName PerimeterCurveGroupName = TEXT("PerimeterCurves");
-	const static FName PerimeterCurveAttributeName = TEXT("PerimeterCurve");
-	const static FName InputRootPointsGroupName = TEXT("InputRootPoints");
-	const static FName InputRootPointsAttributeName = TEXT("Point");
-	const static FName LabelsGroupName = TEXT("Labels");
-	const static FName LabelPositionAttributeName = TEXT("Position");
-	const static FName LabelScaleAttributeName = TEXT("Scale");
-	const static FName LabelTextAttributeName = TEXT("Text");
-
 	const static PV::TDynamicMeshVertexAttributeDefinition<bool> PerimeterAttributeDefinition(TEXT("Perimeter"));
 	
 	static int32 GetPixelIndex(const FVector2i& PixelCoords, int32 ImageWidth)
@@ -69,69 +59,16 @@ namespace Internal
 	{
 		return UE::Math::TVector<Vector3dType>(0, Vec2D.X * Scale, (1.0 - Vec2D.Y) * Scale);
 	}
-
-	static void AddLabelsToCollection(FManagedArrayCollection& InOutCollection, const TArray<TTuple<FVector3f, FString, float>>& Labels)
-	{
-		if (!InOutCollection.HasGroup(LabelsGroupName))
-		{
-			InOutCollection.AddGroup(LabelsGroupName);
-		}
-
-		if (!InOutCollection.HasAttribute(LabelTextAttributeName, LabelsGroupName))
-		{
-			InOutCollection.AddAttribute<FString>(LabelTextAttributeName, LabelsGroupName);
-		}
-
-		if (!InOutCollection.HasAttribute(LabelPositionAttributeName, LabelsGroupName))
-		{
-			InOutCollection.AddAttribute<FVector3f>(LabelPositionAttributeName, LabelsGroupName);
-		}
-
-		if (!InOutCollection.HasAttribute(LabelScaleAttributeName, LabelsGroupName))
-		{
-			InOutCollection.AddAttribute<float>(LabelScaleAttributeName, LabelsGroupName);
-		}
-
-		if (!ensure(InOutCollection.GetAttributeType(LabelTextAttributeName, LabelsGroupName) == FManagedArrayCollection::EArrayType::FStringType))
-		{
-			return;
-		}
-
-		if (!ensure(InOutCollection.GetAttributeType(LabelPositionAttributeName, LabelsGroupName) == FManagedArrayCollection::EArrayType::FVectorType))
-		{
-			return;
-		}
-
-		if (!ensure(InOutCollection.GetAttributeType(LabelScaleAttributeName, LabelsGroupName) == FManagedArrayCollection::EArrayType::FFloatType))
-		{
-			return;
-		}
-
-		const int32 StartIndex = InOutCollection.NumElements(LabelsGroupName);
-
-		InOutCollection.AddElements(Labels.Num(), LabelsGroupName);
-
-		TManagedArray<FString>& LabelTextAttribute = *InOutCollection.FindAttribute<FString>(LabelTextAttributeName, LabelsGroupName);
-		TManagedArray<FVector3f>& LabelPositionAttribute = *InOutCollection.FindAttribute<FVector3f>(LabelPositionAttributeName, LabelsGroupName);
-		TManagedArray<float>& LabelScaleAttribute = *InOutCollection.FindAttribute<float>(LabelScaleAttributeName, LabelsGroupName);
-
-		for (int32 i = 0; i < Labels.Num(); ++i)
-		{
-			LabelPositionAttribute[StartIndex + i] = Labels[i].Get<FVector3f>();
-			LabelTextAttribute[StartIndex + i] = Labels[i].Get<FString>();
-			LabelScaleAttribute[StartIndex + i] = Labels[i].Get<float>();
-		}
-	}
-
+	
 	static void AddPlantSettingsToCollection(FManagedArrayCollection& InOutCollection, const TArray<FPVImportTexture2DPlantSettings>& InPlantSettings, float Scale, const FVector& Offset = FVector::ZeroVector)
 	{
-		InOutCollection.AddGroup(InputRootPointsGroupName);
-		InOutCollection.AddElements(InPlantSettings.Num(), InputRootPointsGroupName);
+		InOutCollection.AddGroup(PVImportNames::UserRootPositionsGroup);
+		InOutCollection.AddElements(InPlantSettings.Num(), PVImportNames::UserRootPositionsGroup);
 		
-		TArray<TTuple<FVector3f, FString, float>> Labels;
+		TArray<PV::ImportHelper::FLabel> Labels;
 		Labels.Reserve(InPlantSettings.Num());
 
-		auto& InputRootPointAttribute = InOutCollection.AddAttribute<FVector3f>(InputRootPointsAttributeName, InputRootPointsGroupName);
+		auto& InputRootPointAttribute = InOutCollection.AddAttribute<FVector3f>(PVImportNames::UserRootPositionsAttribute, PVImportNames::UserRootPositionsGroup);
 		for (int32 i = 0; i < InPlantSettings.Num(); ++i)
 		{
 			const FVector3f PointPosition = ToVector3(InPlantSettings[i].RootPosition, Scale) + FVector3f(Offset);
@@ -139,17 +76,17 @@ namespace Internal
 			Labels.Emplace(PointPosition, FString::Printf(TEXT("%d"), i), 1.f);
 		}
 
-		AddLabelsToCollection(InOutCollection, Labels);
+		PV::ImportHelper::AddLabelsToCollection(Labels, InOutCollection);
 	}
 	
 	static void AddDebugPixelValuesToCollection(FManagedArrayCollection& InOutCollection, const TArray<bool>& ImageData, const FVector2i& ImageSize, float Scale)
 	{
 		check(ImageData.Num() == ImageSize.X * ImageSize.Y);
 		
-		InOutCollection.AddGroup(PixelsDebugGroupName);
-		InOutCollection.AddElements(ImageData.Num(), PixelsDebugGroupName);
-		auto& PixelPositionsAttribute = InOutCollection.AddAttribute<FVector3f>(PixelPositionsAttributeName, PixelsDebugGroupName);
-		auto& PixelValueAttribute = InOutCollection.AddAttribute<bool>(PixelValueAttributeName, PixelsDebugGroupName);
+		InOutCollection.AddGroup(PVImportNames::PixelsDebugGroup);
+		InOutCollection.AddElements(ImageData.Num(), PVImportNames::PixelsDebugGroup);
+		auto& PixelPositionsAttribute = InOutCollection.AddAttribute<FVector3f>(PVImportNames::PixelPositionsAttribute, PVImportNames::PixelsDebugGroup);
+		auto& PixelValueAttribute = InOutCollection.AddAttribute<bool>(PVImportNames::PixelValueAttribute, PVImportNames::PixelsDebugGroup);
 
 		for (int32 x = 0; x < ImageSize.X; x++)
 		{
@@ -172,10 +109,10 @@ namespace Internal
 		const FVector& Offset = FVector::ZeroVector
 	)
 	{
-		InOutCollection.AddGroup(PerimeterCurveGroupName);
-		InOutCollection.AddElements(InPerimeterCurves.Num(), PerimeterCurveGroupName);
+		InOutCollection.AddGroup(PVImportNames::PerimeterCurveGroup);
+		InOutCollection.AddElements(InPerimeterCurves.Num(), PVImportNames::PerimeterCurveGroup);
     	
-		auto& PerimeterCurveAttribute = InOutCollection.AddAttribute<TArray<FVector3f>>(PerimeterCurveAttributeName, PerimeterCurveGroupName);
+		auto& PerimeterCurveAttribute = InOutCollection.AddAttribute<TArray<FVector3f>>(PVImportNames::PerimeterCurveAttribute, PVImportNames::PerimeterCurveGroup);
 		for (int32 i = 0; i < InPerimeterCurves.Num(); i++)
 		{
 			const TArray<FVector2f>& Curve = InPerimeterCurves[i].GetVertices();
@@ -186,27 +123,6 @@ namespace Internal
 		}
 	}
 
-	static void AddDynamicMeshVerticesToCollection(
-		FManagedArrayCollection& InOutCollection, 
-		const FName GroupName,
-		const FName AttributeName,
-		const UE::Geometry::FDynamicMesh3& DynamicMesh, 
-		const FBooleanVertexAttribute& Filter,
-		const FVector& Offset
-	)
-	{
-		const TArray<int32> VertexIndices = Filter.FindAllNonZero();
-		
-		InOutCollection.AddGroup(GroupName);
-		InOutCollection.AddElements(VertexIndices.Num(), GroupName);
-
-		auto& Attribute = InOutCollection.AddAttribute<FVector3f>(AttributeName, GroupName);
-		for (int32 i = 0; i < VertexIndices.Num(); ++i)
-		{
-			Attribute[i] = FVector3f(DynamicMesh.GetVertex(VertexIndices[i]) + Offset);
-		}
-	}
-
 	static void AddTextureQuadToCollection(
 		FManagedArrayCollection& InOutCollection,
 		UTexture2D& Texture2D,
@@ -262,7 +178,7 @@ namespace Internal
 	{
 		using namespace PV::PlantTraversalHelper;
 
-		TArray<TTuple<FVector3f, FString, float>> Labels;
+		TArray<PV::ImportHelper::FLabel> Labels;
 		Labels.Reserve(PlantNumberToPlantIndices.Num());
 
 		struct FPlantOffsetInfo
@@ -312,7 +228,7 @@ namespace Internal
 			Offset.Y += PlantsOffsetInfo[i].Bounds.Min.Y;
 		}
 
-		AddLabelsToCollection(InCollection, Labels);
+		PV::ImportHelper::AddLabelsToCollection(Labels, InCollection);
 	}
 
 	static bool SampleImage(UTexture2D& InTexture2D, FVector2i& InOutSampleSize, TArray<bool>& OutBlackAndWhiteImageData, bool bInvertImage, float WhiteLevel)
@@ -667,8 +583,6 @@ namespace Internal
 		*       |   |     |  |
 		*/
 
-		using namespace PV::DynamicMeshHelper;
-
 		const static auto ComputeAngleDeg = [](
 			const UE::Geometry::FDynamicMesh3& Mesh,
 			const TArray<int32>& PerimeterVertices, 
@@ -954,7 +868,7 @@ namespace Internal
 
 	static void MapPlantIndicesToBranchHierarchies(
 		const TArray<FPlantMesh>& PlantMapping,
-		const TArray<PV::DynamicMeshHelper::FMeshBranchHierarchy>& DynamicMeshBranchHierarchies,
+		const TArray<PV::ImportHelper::FMeshBranchHierarchy>& DynamicMeshBranchHierarchies,
 		TArray<int32>& OutBranchHierarchiesToPlantIndices
 	)
 	{
@@ -1025,7 +939,7 @@ namespace Internal
 
 PV::Texture2DImport::EImportResult PV::Texture2DImport::ImportGrowthDataFromTexture2D(const FPVImportTexture2DParams& InParams, const UMaterialInterface* TextureVisualizationMaterial, FPVImportTexture2DOutput& Output)
 {
-	using namespace PV::DynamicMeshHelper;
+	using namespace PV::ImportHelper;
 	using namespace Internal;
 	
 	if (!InParams.Texture2DAsset)
@@ -1105,12 +1019,12 @@ PV::Texture2DImport::EImportResult PV::Texture2DImport::ImportGrowthDataFromText
 
 	// Add end points visualization before applying root points as that transforms the DynamicMesh
 	AddDynamicMeshVerticesToCollection(
-		Output.VisualizationCollection,
-		TEXT("EndPoints"), 
-		TEXT("Vertex"), 
 		DynamicMesh, 
 		EndPointAttributeDefinition.GetOrAttachAttribute(DynamicMesh), 
-		VisualizationOffset
+		VisualizationOffset,
+		Output.VisualizationCollection,
+		PVImportNames::TipVisualizationGroup,
+		PVImportNames::TipVisualizationAttribute
 	);
 
 	TArray<FPlantMesh> PlantMapping;
diff --git a/Engine/Plugins/Experimental/ProceduralVegetationEditor/Source/ProceduralVegetation/Private/Nodes/PVStaticMeshImporterSettings.cpp b/Engine/Plugins/Experimental/ProceduralVegetationEditor/Source/ProceduralVegetation/Private/Nodes/PVStaticMeshImporterSettings.cpp
--- a/Engine/Plugins/Experimental/ProceduralVegetationEditor/Source/ProceduralVegetation/Private/Nodes/PVStaticMeshImporterSettings.cpp
+++ b/Engine/Plugins/Experimental/ProceduralVegetationEditor/Source/ProceduralVegetation/Private/Nodes/PVStaticMeshImporterSettings.cpp
@@ -1,6 +1,7 @@
 // Copyright Epic Games, Inc. All Rights Reserved.
 
 #include "PVStaticMeshImporterSettings.h"
+#include "PVImportCommon.h"
 #include "Implementations/PVImporter_StaticMesh.h"
 #include "PCGContext.h"
 #include "ProceduralVegetationModule.h"
@@ -66,15 +67,61 @@ bool FPVStaticMeshImporterElement::ExecuteInternal(FPCGContext* InContext) const
 		return true;
 	}
 	
-	FManagedArrayCollection Collection;
-	const auto Result = PV::StaticMeshImport::ImportGrowthDataFromStaticMesh(Settings->Params, Collection);
+	FPVImportStaticMeshOutput Output;
+	const auto Result = PV::StaticMeshImport::ImportGrowthDataFromStaticMesh(Settings->Params, Output);
 	if (Result != PV::StaticMeshImport::EImportResult::Success)
 	{
-		return false;
+		return Result == PV::StaticMeshImport::EImportResult::MeshNotReady ? false : true;
 	}
 
 	UPVGrowthData* OutVariantData = FPCGContext::NewObject_AnyThread<UPVGrowthData>(InContext);
-	OutVariantData->Initialize(MoveTemp(Collection));
+
+	if (Output.DebugCollection.IsSet())
+	{
+		OutVariantData->Initialize(MoveTemp(Output.DebugCollection.GetValue()));
+	}
+	else
+	{
+		OutVariantData->Initialize(MoveTemp(Output.GrowthDataCollection));
+	}
+
+#if WITH_EDITOR
+	const TArray<FPVVisualizationSettings> VisualizationSettings = Invoke([&]()
+	{
+		TArray<FPVVisualizationSettings> OutVisualizationSettings;
+
+		if (Settings->Params.bShowFoundTips)
+		{
+			FPVVisualizationSettings& TipVisualization = OutVisualizationSettings.AddDefaulted_GetRef();
+			TipVisualization.DebugType = EPVDebugType::Custom;
+			TipVisualization.bShowAnchorPoints = false;
+			TipVisualization.bDrawPointAsMesh = false;
+			TipVisualization.bUsePivotAsPosition = false;
+			TipVisualization.AttributeToFilter = PVImportNames::TipVisualizationAttribute;
+			TipVisualization.CustomGroupToFilter = PVImportNames::TipVisualizationGroup;
+			TipVisualization.VisualizationMode = EPVDebugValueVisualizationMode::Point;
+			TipVisualization.Color = FColor::Red;
+			TipVisualization.GizmoScale = 6.5f;
+		}
+
+		if (Settings->Params.bShowFoundBranchCurves)
+		{
+			FPVVisualizationSettings& BranchHierarchyVisualization = OutVisualizationSettings.AddDefaulted_GetRef();
+			BranchHierarchyVisualization.DebugType = EPVDebugType::Custom;
+			BranchHierarchyVisualization.bShowAnchorPoints = false;
+			BranchHierarchyVisualization.bDrawPointAsMesh = false;
+			BranchHierarchyVisualization.AttributeToFilter = PVImportNames::BranchHierarchyAttribute;
+			BranchHierarchyVisualization.CustomGroupToFilter = PVImportNames::BranchHierarchyGroup;
+			BranchHierarchyVisualization.VisualizationMode = EPVDebugValueVisualizationMode::Curve;
+			BranchHierarchyVisualization.Color = FColor::White;
+			BranchHierarchyVisualization.DepthPriorityGroup = ESceneDepthPriorityGroup::SDPG_Foreground;
+		}
+
+		return OutVisualizationSettings;
+	});
+
+	OutVariantData->AddVisualizationCollection(Output.VisualizationCollection, VisualizationSettings);
+#endif
 
 	FPCGTaggedData& CollectionOutput = InContext->OutputData.TaggedData.Emplace_GetRef();
 	CollectionOutput.Data = OutVariantData;
diff --git a/Engine/Plugins/Experimental/ProceduralVegetationEditor/Source/ProceduralVegetation/Private/Nodes/PVTexture2DImporterSettings.cpp b/Engine/Plugins/Experimental/ProceduralVegetationEditor/Source/ProceduralVegetation/Private/Nodes/PVTexture2DImporterSettings.cpp
--- a/Engine/Plugins/Experimental/ProceduralVegetationEditor/Source/ProceduralVegetation/Private/Nodes/PVTexture2DImporterSettings.cpp
+++ b/Engine/Plugins/Experimental/ProceduralVegetationEditor/Source/ProceduralVegetation/Private/Nodes/PVTexture2DImporterSettings.cpp
@@ -1,7 +1,7 @@
 // Copyright Epic Games, Inc. All Rights Reserved.
 
 #include "PVTexture2DImporterSettings.h"
-
+#include "PVImportCommon.h"
 #include "ProceduralVegetationModule.h"
 #include "DataTypes/PVGrowthData.h"
 #include "ProceduralVegetationPreset.h"
@@ -32,21 +32,21 @@ UPVTexture2DImporterSettings::UPVTexture2DImporterSettings()
 	static ConstructorHelpers::FObjectFinder<UMaterialInterface> TextureVisualizationMaterialRef(TEXT("/Script/Engine.Material'/Engine/Engine_MI_Shaders/M_Shader_SimpleTranslucent.M_Shader_SimpleTranslucent'"));
 	TextureVisualizationMaterial = TextureVisualizationMaterialRef.Object;
 
-	FPVVisualizationSettings& BranchCurvesVisualization = DebugVisualizationSettings.VisualizationSettings.AddDefaulted_GetRef();
-	BranchCurvesVisualization.DebugType = EPVDebugType::Custom;
-	BranchCurvesVisualization.bShowAnchorPoints = false;
-	BranchCurvesVisualization.bDrawPointAsMesh = false;
-	BranchCurvesVisualization.AttributeToFilter = TEXT("Curve");
-	BranchCurvesVisualization.CustomGroupToFilter = TEXT("BranchCurves");
-	BranchCurvesVisualization.VisualizationMode = EPVDebugValueVisualizationMode::Curve;
+	FPVVisualizationSettings& BranchHierarchyVisualization = DebugVisualizationSettings.VisualizationSettings.AddDefaulted_GetRef();
+	BranchHierarchyVisualization.DebugType = EPVDebugType::Custom;
+	BranchHierarchyVisualization.bShowAnchorPoints = false;
+	BranchHierarchyVisualization.bDrawPointAsMesh = false;
+	BranchHierarchyVisualization.AttributeToFilter = PVImportNames::BranchHierarchyAttribute;
+	BranchHierarchyVisualization.CustomGroupToFilter = PVImportNames::BranchHierarchyGroup;
+	BranchHierarchyVisualization.VisualizationMode = EPVDebugValueVisualizationMode::Curve;
 
 	FPVVisualizationSettings& PixelSamplingVisualization = DebugVisualizationSettings.VisualizationSettings.AddDefaulted_GetRef();
 	PixelSamplingVisualization.DebugType = EPVDebugType::Custom;
 	PixelSamplingVisualization.bShowAnchorPoints = false;
 	PixelSamplingVisualization.bDrawPointAsMesh = false;
 	PixelSamplingVisualization.bUsePivotAsPosition = true;
-	PixelSamplingVisualization.AttributeToFilter = TEXT("Value");
-	PixelSamplingVisualization.CustomGroupToFilter = TEXT("PixelDebug");
+	PixelSamplingVisualization.AttributeToFilter = PVImportNames::PixelValueAttribute;
+	PixelSamplingVisualization.CustomGroupToFilter = PVImportNames::PixelsDebugGroup;
 	PixelSamplingVisualization.VisualizationMode = EPVDebugValueVisualizationMode::Point;
 	PixelSamplingVisualization.Color = FColor(0, 0, 255);
 	PixelSamplingVisualization.GizmoScale = 3.000000;
@@ -123,38 +123,38 @@ bool FPVTexture2DImporterElement::ExecuteInternal(FPCGContext* InContext) const
 		FPVVisualizationSettings& PerimeterCurveVisualization = OutVisualizationSettings.AddDefaulted_GetRef();
 		PerimeterCurveVisualization.DebugType = EPVDebugType::Custom;
 		PerimeterCurveVisualization.bShowAnchorPoints = false;
-		PerimeterCurveVisualization.AttributeToFilter = TEXT("PerimeterCurve");
-		PerimeterCurveVisualization.CustomGroupToFilter = TEXT("PerimeterCurves");
+		PerimeterCurveVisualization.AttributeToFilter = PVImportNames::PerimeterCurveAttribute;
+		PerimeterCurveVisualization.CustomGroupToFilter = PVImportNames::PerimeterCurveGroup;
 		PerimeterCurveVisualization.VisualizationMode = EPVDebugValueVisualizationMode::Curve;
 
-		FPVVisualizationSettings& EndPointVisualization = OutVisualizationSettings.AddDefaulted_GetRef();
-		EndPointVisualization.DebugType = EPVDebugType::Custom;
-		EndPointVisualization.bShowAnchorPoints = false;
-		EndPointVisualization.bDrawPointAsMesh = false;
-		EndPointVisualization.bUsePivotAsPosition = false;
-		EndPointVisualization.AttributeToFilter = TEXT("Vertex");
-		EndPointVisualization.CustomGroupToFilter = TEXT("EndPoints");
-		EndPointVisualization.VisualizationMode = EPVDebugValueVisualizationMode::Point;
-		EndPointVisualization.Color = FColor::Red;
-		EndPointVisualization.GizmoScale = 6.5f;
-
-		FPVVisualizationSettings& RootPointsVisualization = OutVisualizationSettings.AddDefaulted_GetRef();
-		RootPointsVisualization.DebugType = EPVDebugType::Custom;
-		RootPointsVisualization.bShowAnchorPoints = false;
-		RootPointsVisualization.bDrawPointAsMesh = false;
-		RootPointsVisualization.AttributeToFilter = TEXT("Point");
-		RootPointsVisualization.CustomGroupToFilter = TEXT("InputRootPoints");
-		RootPointsVisualization.VisualizationMode = EPVDebugValueVisualizationMode::Point;
-		RootPointsVisualization.Color = FColor::Blue;
-		RootPointsVisualization.GizmoScale = 10;
+		FPVVisualizationSettings& TipVisualization = OutVisualizationSettings.AddDefaulted_GetRef();
+		TipVisualization.DebugType = EPVDebugType::Custom;
+		TipVisualization.bShowAnchorPoints = false;
+		TipVisualization.bDrawPointAsMesh = false;
+		TipVisualization.bUsePivotAsPosition = false;
+		TipVisualization.AttributeToFilter = PVImportNames::TipVisualizationAttribute;
+		TipVisualization.CustomGroupToFilter = PVImportNames::TipVisualizationGroup;
+		TipVisualization.VisualizationMode = EPVDebugValueVisualizationMode::Point;
+		TipVisualization.Color = FColor::Red;
+		TipVisualization.GizmoScale = 6.5f;
+
+		FPVVisualizationSettings& UserRootPositionsVisualization = OutVisualizationSettings.AddDefaulted_GetRef();
+		UserRootPositionsVisualization.DebugType = EPVDebugType::Custom;
+		UserRootPositionsVisualization.bShowAnchorPoints = false;
+		UserRootPositionsVisualization.bDrawPointAsMesh = false;
+		UserRootPositionsVisualization.AttributeToFilter = PVImportNames::UserRootPositionsGroup;
+		UserRootPositionsVisualization.CustomGroupToFilter = PVImportNames::UserRootPositionsAttribute;
+		UserRootPositionsVisualization.VisualizationMode = EPVDebugValueVisualizationMode::Point;
+		UserRootPositionsVisualization.Color = FColor::Blue;
+		UserRootPositionsVisualization.GizmoScale = 10;
 
 		FPVVisualizationSettings& LabelVisualization = OutVisualizationSettings.AddDefaulted_GetRef();
 		LabelVisualization.DebugType = EPVDebugType::Custom;
 		LabelVisualization.bShowAnchorPoints = false;
-		LabelVisualization.CustomGroupToFilter = TEXT("Labels");
-		LabelVisualization.AttributeToFilter = TEXT("Text");
-		LabelVisualization.CustomPivotPositionAttributeName = TEXT("Position");
-		LabelVisualization.CustomPivotScaleAttributeName = TEXT("Scale");
+		LabelVisualization.CustomGroupToFilter = PVImportNames::LabelsGroup;
+		LabelVisualization.AttributeToFilter = PVImportNames::LabelTextAttribute;
+		LabelVisualization.CustomPivotPositionAttributeName = PVImportNames::LabelPositionAttribute;
+		LabelVisualization.CustomPivotScaleAttributeName = PVImportNames::LabelScaleAttribute;
 		LabelVisualization.VisualizationMode = EPVDebugValueVisualizationMode::Text;
 		LabelVisualization.Color = FColor::Black;
 		LabelVisualization.TextScale = 10.f;
diff --git a/Engine/Plugins/Experimental/ProceduralVegetationEditor/Source/ProceduralVegetation/Private/PVImportCommon.h b/Engine/Plugins/Experimental/ProceduralVegetationEditor/Source/ProceduralVegetation/Private/PVImportCommon.h
new file mode 100644
--- /dev/null
+++ b/Engine/Plugins/Experimental/ProceduralVegetationEditor/Source/ProceduralVegetation/Private/PVImportCommon.h
@@ -0,0 +1,26 @@
+// Copyright Epic Games, Inc. All Rights Reserved.
+#pragma once
+
+#include "CoreMinimal.h"
+
+namespace PVImportNames
+{
+	// Common
+	inline static const FName BranchHierarchyGroup = TEXT("BranchCurves");
+	inline static const FName BranchHierarchyAttribute = TEXT("Curve");
+	inline static const FName LabelsGroup = TEXT("Labels");
+	inline static const FName LabelPositionAttribute = TEXT("Position");
+	inline static const FName LabelScaleAttribute = TEXT("Scale");
+	inline static const FName LabelTextAttribute = TEXT("Text");
+	inline static const FName TipVisualizationGroup = TEXT("EndPoints");
+	inline static const FName TipVisualizationAttribute = TEXT("Vertex");
+
+	// 2D specific
+	inline static const FName PixelsDebugGroup = TEXT("PixelDebug");
+	inline static const FName PixelPositionsAttribute = TEXT("Position");
+	inline static const FName PixelValueAttribute = TEXT("Value");
+	inline static const FName PerimeterCurveGroup = TEXT("PerimeterCurves");
+	inline static const FName PerimeterCurveAttribute = TEXT("PerimeterCurve");
+	inline static const FName UserRootPositionsGroup = TEXT("RootPositions");
+	inline static const FName UserRootPositionsAttribute = TEXT("Vertex");
+};
\ No newline at end of file
diff --git a/Engine/Plugins/Experimental/ProceduralVegetationEditor/Source/ProceduralVegetation/Private/Params/PVImportStaticMeshParams.h b/Engine/Plugins/Experimental/ProceduralVegetationEditor/Source/ProceduralVegetation/Private/Params/PVImportStaticMeshParams.h
--- a/Engine/Plugins/Experimental/ProceduralVegetationEditor/Source/ProceduralVegetation/Private/Params/PVImportStaticMeshParams.h
+++ b/Engine/Plugins/Experimental/ProceduralVegetationEditor/Source/ProceduralVegetation/Private/Params/PVImportStaticMeshParams.h
@@ -2,7 +2,7 @@
 #pragma once
 
 #include "Engine/EngineTypes.h"
-
+#include "GeometryCollection/ManagedArrayCollection.h"
 #include "PVImportStaticMeshParams.generated.h"
 
 class UStaticMesh;
@@ -12,7 +12,6 @@ enum class EPVImportStaticMeshDebugState : uint8
 {
 	None,
 	InitialMesh,
-	KeepMaterial,
 	ComputeNormals,
 	ComputeVerticesCenterAndRadius,
 	CollapseBranchesToLine,
@@ -34,6 +33,19 @@ struct FPVImportStaticMeshParams
 	UPROPERTY(EditAnywhere, Category = "Asset")
 	TArray<FName> MaterialsToKeep;
 
-	UPROPERTY(EditAnywhere, Category = "Asset")
+	UPROPERTY(EditAnywhere, Category = "Visualization")
+	bool bShowFoundTips = true;
+
+	UPROPERTY(EditAnywhere, Category = "Visualization")
+	bool bShowFoundBranchCurves = true;
+
+	UPROPERTY(EditAnywhere, Category = "Debug State")
 	EPVImportStaticMeshDebugState DebugState = EPVImportStaticMeshDebugState::None;
 };
+
+struct FPVImportStaticMeshOutput
+{
+	TOptional<FManagedArrayCollection> DebugCollection;
+	FManagedArrayCollection VisualizationCollection;
+	FManagedArrayCollection GrowthDataCollection;
+};
diff --git a/Engine/Plugins/Experimental/ProceduralVegetationEditor/Source/ProceduralVegetation/Public/DataTypes/PVData.h b/Engine/Plugins/Experimental/ProceduralVegetationEditor/Source/ProceduralVegetation/Public/DataTypes/PVData.h
--- a/Engine/Plugins/Experimental/ProceduralVegetationEditor/Source/ProceduralVegetation/Public/DataTypes/PVData.h
+++ b/Engine/Plugins/Experimental/ProceduralVegetationEditor/Source/ProceduralVegetation/Public/DataTypes/PVData.h
@@ -5,7 +5,7 @@
 #include "PCGContext.h"
 
 #include "Data/PCGSpatialData.h"
-
+#include "Engine/EngineTypes.h"
 #include "GeometryCollection/ManagedArrayCollection.h"
 
 #include "PVData.generated.h"
@@ -104,6 +104,9 @@ struct FPVVisualizationSettings
 
 	UPROPERTY(EditAnywhere, Category = DebugSettings)
 	float GizmoScale = 1.f;
+
+	UPROPERTY(EditAnywhere, Category = DebugSettings)
+	TEnumAsByte<ESceneDepthPriorityGroup> DepthPriorityGroup = ESceneDepthPriorityGroup::SDPG_World;
 	
 	FName GetGroupToFilter() const
 	{
diff --git a/Engine/Plugins/Experimental/ProceduralVegetationEditor/Source/ProceduralVegetationEditor/Private/Visualizations/DebugVisualization/PVDebugVisualizationBase.cpp b/Engine/Plugins/Experimental/ProceduralVegetationEditor/Source/ProceduralVegetationEditor/Private/Visualizations/DebugVisualization/PVDebugVisualizationBase.cpp
--- a/Engine/Plugins/Experimental/ProceduralVegetationEditor/Source/ProceduralVegetationEditor/Private/Visualizations/DebugVisualization/PVDebugVisualizationBase.cpp
+++ b/Engine/Plugins/Experimental/ProceduralVegetationEditor/Source/ProceduralVegetationEditor/Private/Visualizations/DebugVisualization/PVDebugVisualizationBase.cpp
@@ -521,7 +521,7 @@ void FPVDebugVisualizationBase::DrawAttributes(FVisualizerDrawContext& InContext
 					const FVector EndPos = StartPos + DirectionVector * BaseDirectionStrength * VisualizationSettings.GizmoScale;
 
 					const FColor LineColor = FLinearColor::LerpUsingHSV(DrawColor.ReinterpretAsLinear(), FLinearColor::Black, j / (float)AttributeAsVector[i].Num()).ToFColor(true);
-					DrawAsLine(LineBatchComponent, StartPos, EndPos, LineColor);
+					DrawAsLine(LineBatchComponent, StartPos, EndPos, LineColor, VisualizationSettings.DepthPriorityGroup);
 
 					if (AttributeAsVector[i].Num() > 1)
 					{
@@ -589,7 +589,7 @@ void FPVDebugVisualizationBase::DrawAttributes(FVisualizerDrawContext& InContext
 				GetPivot(Collection, i, OutPos, OutScale);
 
 				const FColor DrawColor = VisualizationSettings.bRandomizeColors ? FColor::MakeRandomColor() : VisualizationSettings.Color;
-				DrawAsSphere(LineBatchComponent, FVector(OutPos), AttributeAsScalar[i] * VisualizationSettings.GizmoScale * OutScale, DrawColor);
+				DrawAsSphere(LineBatchComponent, FVector(OutPos), AttributeAsScalar[i] * VisualizationSettings.GizmoScale * OutScale, DrawColor, VisualizationSettings.DepthPriorityGroup);
 			}
 			break;
 		}
@@ -612,7 +612,7 @@ void FPVDebugVisualizationBase::DrawAttributes(FVisualizerDrawContext& InContext
 				{
 					const FVector StartPos = AttributeAsVector[i][j];
 					const FVector EndPos = AttributeAsVector[i][j + 1];
-					DrawAsLine(LineBatchComponent, StartPos, EndPos, DrawColor);
+					DrawAsLine(LineBatchComponent, StartPos, EndPos, DrawColor, VisualizationSettings.DepthPriorityGroup);
 				}
 			}
 			break;
@@ -658,7 +658,7 @@ void FPVDebugVisualizationBase::DrawAttributes(FVisualizerDrawContext& InContext
 		}
 		else
 		{
-			DrawAsPoints(GetOrCreateLineComponent(InContext.SceneSetupParams), PointsToDraw);
+			DrawAsPoints(GetOrCreateLineComponent(InContext.SceneSetupParams), PointsToDraw, VisualizationSettings.DepthPriorityGroup);
 		}
 
 		PointsToDraw.Reset();
@@ -688,23 +688,30 @@ void FPVDebugVisualizationBase::DrawAsText(FPCGSceneSetupParams& InOutParams, co
 	InOutParams.Scene->AddComponent(Component, FTransform(FVector(Pos.X + TextOffset.X, Pos.Y + TextOffset.Y, Pos.Z + TextOffset.Z)));
 }
 
-void FPVDebugVisualizationBase::DrawAsLine(UPVLineBatchComponent* LineBatchComponent, const FVector& InStartPos, const FVector& InEndPos,
-	const FLinearColor& InColor)
+void FPVDebugVisualizationBase::DrawAsLine(UPVLineBatchComponent* LineBatchComponent, const FVector& InStartPos, const FVector& InEndPos, const FLinearColor& InColor, ESceneDepthPriorityGroup InDepthPriorityGroup)
 {
 	check(LineBatchComponent);
-	LineBatchComponent->AddLine(InStartPos, InEndPos, InColor);
+	LineBatchComponent->AddLine(InStartPos, InEndPos, InColor, InDepthPriorityGroup);
 }
 
-void FPVDebugVisualizationBase::DrawAsSphere(UPVLineBatchComponent* LineBatchComponent, const FVector& InPos, float Radius, const FLinearColor& InColor)
+void FPVDebugVisualizationBase::DrawAsSphere(UPVLineBatchComponent* LineBatchComponent, const FVector& InPos, float Radius, const FLinearColor& InColor, ESceneDepthPriorityGroup InDepthPriorityGroup)
 {
 	check(LineBatchComponent);
-	LineBatchComponent->AddSphere(InPos, Radius, InColor);
+	LineBatchComponent->AddSphere(InPos, Radius, InColor, InDepthPriorityGroup);
 }
 
-void FPVDebugVisualizationBase::DrawAsPoints(UPVLineBatchComponent* LineBatchComponent, const TArray<TTuple<FVector, float, FLinearColor>>& InPoints)
+void FPVDebugVisualizationBase::DrawAsPoints(UPVLineBatchComponent* LineBatchComponent, const TArray<FPoint>& InPoints, ESceneDepthPriorityGroup InDepthPriorityGroup)
 {
 	check(LineBatchComponent);
-	LineBatchComponent->AddPoints(InPoints, SDPG_Foreground);
+	TArray<FPVPointInfo> PointsToDraw;
+	PointsToDraw.Reserve(InPoints.Num());
+
+	for (const FPoint& Point : InPoints)
+	{
+		PointsToDraw.Emplace(Point.Position, Point.Scale, Point.Color, InDepthPriorityGroup);
+	}
+
+	LineBatchComponent->AddPoints(PointsToDraw);
 }
 
 void FPVDebugVisualizationBase::DrawAsMeshPoints(UInstancedStaticMeshComponent* InstancedStaticMeshComponent, const TArray<FTransform>& InTransforms, const FLinearColor& InColor)
@@ -718,8 +725,7 @@ void FPVDebugVisualizationBase::DrawAsMeshPoints(UInstancedStaticMeshComponent*
 	InstancedStaticMeshComponent->AddInstances(InTransforms, false);
 }
 
-void FPVDebugVisualizationBase::DrawAsMeshPoints(UInstancedStaticMeshComponent* InstancedStaticMeshComponent,
-	TArray<TTuple<FVector, float, FLinearColor>> InPoints)
+void FPVDebugVisualizationBase::DrawAsMeshPoints(UInstancedStaticMeshComponent* InstancedStaticMeshComponent, const TArray<FPoint>& InPoints)
 {
 	check(InstancedStaticMeshComponent != nullptr);
 	for (const auto& [Position, Size, Color] : InPoints)
diff --git a/Engine/Plugins/Experimental/ProceduralVegetationEditor/Source/ProceduralVegetationEditor/Private/Visualizations/DebugVisualization/PVDebugVisualizationBase.h b/Engine/Plugins/Experimental/ProceduralVegetationEditor/Source/ProceduralVegetationEditor/Private/Visualizations/DebugVisualization/PVDebugVisualizationBase.h
--- a/Engine/Plugins/Experimental/ProceduralVegetationEditor/Source/ProceduralVegetationEditor/Private/Visualizations/DebugVisualization/PVDebugVisualizationBase.h
+++ b/Engine/Plugins/Experimental/ProceduralVegetationEditor/Source/ProceduralVegetationEditor/Private/Visualizations/DebugVisualization/PVDebugVisualizationBase.h
@@ -28,12 +28,19 @@ protected:
 	virtual TArray<FVector3f> GetPivotPositions(const FManagedArrayCollection& InCollection) = 0;
 	virtual void GetPivot(const FManagedArrayCollection& InCollection, const int InIndex, FVector3f& OutPos, float& OutScale) = 0;
 	
+	struct FPoint
+	{
+		FVector Position;
+		float Scale;
+		FLinearColor Color;
+	};
+
 	static void DrawAsText(FPCGSceneSetupParams& InOutParams, const FText& InTextToDraw, const FVector3f& InPos, float InScale, float TextSize, const FLinearColor& InColor = FLinearColor::White, const FVector3f TextOffset = FVector3f::Zero());
-	static void DrawAsLine(UPVLineBatchComponent* LineBatchComponent, const FVector& InStartPos, const FVector& InEndPos, const FLinearColor& InColor = FLinearColor::White);
-	static void DrawAsSphere(UPVLineBatchComponent* LineBatchComponent, const FVector& InPos, float Radius, const FLinearColor& InColor = FLinearColor::White);
-	static void DrawAsPoints(UPVLineBatchComponent* LineBatchComponent, const TArray<TTuple<FVector, float, FLinearColor>>& InPoints);
+	static void DrawAsLine(UPVLineBatchComponent* LineBatchComponent, const FVector& InStartPos, const FVector& InEndPos, const FLinearColor& InColor, ESceneDepthPriorityGroup InDepthPriorityGroup);
+	static void DrawAsSphere(UPVLineBatchComponent* LineBatchComponent, const FVector& InPos, float Radius, const FLinearColor& InColor, ESceneDepthPriorityGroup InDepthPriorityGroup);
+	static void DrawAsPoints(UPVLineBatchComponent* LineBatchComponent, const TArray<FPoint>& InPoints, ESceneDepthPriorityGroup InDepthPriorityGroup);
 	static void DrawAsMeshPoints(UInstancedStaticMeshComponent* InstancedStaticMeshComponent, const TArray<FTransform>& InTransforms, const FLinearColor& InColor = FLinearColor::White);
-	static void DrawAsMeshPoints(UInstancedStaticMeshComponent* InstancedStaticMeshComponent, TArray<TTuple<FVector, float, FLinearColor>> InPoints);
+	static void DrawAsMeshPoints(UInstancedStaticMeshComponent* InstancedStaticMeshComponent, const TArray<FPoint>& InPoints);
 	
 	static UPVLineBatchComponent* GetOrCreateLineComponent(FPCGSceneSetupParams& InOutParams);
 	static UInstancedStaticMeshComponent* GetOrCreatePointComponent(FPCGSceneSetupParams& InOutParams);
@@ -41,7 +48,7 @@ protected:
 	static void AddToScene(FPCGSceneSetupParams& InOutParams, UPrimitiveComponent* InComponent);
 	
 private:
-	TArray<TTuple<FVector, float, FLinearColor>> PointsToDraw;
+	TArray<FPoint> PointsToDraw;
 };
 typedef TSharedPtr<FPVDebugVisualizationBase> FPVDebugVisualizationPtr;
 
diff --git a/Engine/Plugins/Experimental/ProceduralVegetationEditor/Source/ProceduralVegetationEditor/Private/Visualizations/PVLineBatchComponent.cpp b/Engine/Plugins/Experimental/ProceduralVegetationEditor/Source/ProceduralVegetationEditor/Private/Visualizations/PVLineBatchComponent.cpp
--- a/Engine/Plugins/Experimental/ProceduralVegetationEditor/Source/ProceduralVegetationEditor/Private/Visualizations/PVLineBatchComponent.cpp
+++ b/Engine/Plugins/Experimental/ProceduralVegetationEditor/Source/ProceduralVegetationEditor/Private/Visualizations/PVLineBatchComponent.cpp
@@ -184,13 +184,15 @@ void UPVLineBatchComponent::AddPoint(const FVector& PointLocation, float InPoint
 	MarkRenderStateDirty();
 }
 
-void UPVLineBatchComponent::AddPoints(const TArray<TTuple<FVector, float, FLinearColor>>& InPoints, const ESceneDepthPriorityGroup InDepthPriorityGroup)
+void UPVLineBatchComponent::AddPoints(const TArray<FPVPointInfo>& InPoints)
 {
-	for (const auto& [Position, Scale, Color] : InPoints)
+	Points.Reserve(Points.Num() + InPoints.Num());
+	for (const FPVPointInfo& Point : InPoints)
 	{
-		Points.Emplace(Position, Scale, Color, InDepthPriorityGroup);
-		BBox += Position;
+		Points.Add(Point);
+		BBox += Point.PointLocation;
 	}
+
 	MarkRenderStateDirty();
 }
 
diff --git a/Engine/Plugins/Experimental/ProceduralVegetationEditor/Source/ProceduralVegetationEditor/Private/Visualizations/PVLineBatchComponent.h b/Engine/Plugins/Experimental/ProceduralVegetationEditor/Source/ProceduralVegetationEditor/Private/Visualizations/PVLineBatchComponent.h
--- a/Engine/Plugins/Experimental/ProceduralVegetationEditor/Source/ProceduralVegetationEditor/Private/Visualizations/PVLineBatchComponent.h
+++ b/Engine/Plugins/Experimental/ProceduralVegetationEditor/Source/ProceduralVegetationEditor/Private/Visualizations/PVLineBatchComponent.h
@@ -73,18 +73,13 @@ public:
 
 	void InitBounds();
 	
-	void AddLine(const FVector& InStartPos, const FVector& InEndPos, const FLinearColor& InColor = FLinearColor::White,
-				const ESceneDepthPriorityGroup InDepthPriorityGroup = ESceneDepthPriorityGroup::SDPG_World,
-				const EPointDrawSettings InPointDrawSettings = EPointDrawSettings::None);
+	void AddLine(const FVector& InStartPos, const FVector& InEndPos, const FLinearColor& InColor, const ESceneDepthPriorityGroup InDepthPriorityGroup, const EPointDrawSettings InPointDrawSettings = EPointDrawSettings::None);
 
-	void AddSphere(const FVector& InPos, float Radius, const FLinearColor& InColor = FLinearColor::White,
-		const ESceneDepthPriorityGroup InDepthPriorityGroup = ESceneDepthPriorityGroup::SDPG_World);
+	void AddSphere(const FVector& InPos, float Radius, const FLinearColor& InColor, const ESceneDepthPriorityGroup InDepthPriorityGroup);
 
-	void AddPoint(const FVector& PointLocation, float PointSize, const FLinearColor& InColor = FLinearColor::White,
-		const ESceneDepthPriorityGroup InDepthPriorityGroup = ESceneDepthPriorityGroup::SDPG_World);
+	void AddPoint(const FVector& PointLocation, float PointSize, const FLinearColor& InColor, const ESceneDepthPriorityGroup InDepthPriorityGroup);
 
-	void AddPoints(const TArray<TTuple<FVector, float, FLinearColor>>& InPoints,
-		const ESceneDepthPriorityGroup InDepthPriorityGroup = ESceneDepthPriorityGroup::SDPG_World);
+	void AddPoints(const TArray<FPVPointInfo>& InPoints);
 
 	void Flush();
 	
diff --git a/Engine/Plugins/Experimental/ProceduralVegetationEditor/Source/ProceduralVegetationEditor/Private/Visualizations/PVScaleVisualizationComponent.cpp b/Engine/Plugins/Experimental/ProceduralVegetationEditor/Source/ProceduralVegetationEditor/Private/Visualizations/PVScaleVisualizationComponent.cpp
--- a/Engine/Plugins/Experimental/ProceduralVegetationEditor/Source/ProceduralVegetationEditor/Private/Visualizations/PVScaleVisualizationComponent.cpp
+++ b/Engine/Plugins/Experimental/ProceduralVegetationEditor/Source/ProceduralVegetationEditor/Private/Visualizations/PVScaleVisualizationComponent.cpp
@@ -97,7 +97,7 @@ void UPVScaleVisualizationComponent::UpdateScaleVisualizations()
 		TextRenderComponent->SetMaterial(0, LoadObject<UMaterialInterface>(nullptr, TEXT("/ProceduralVegetationEditor/Materials/ScaleTextMaterial.ScaleTextMaterial")));
 		TextRenderComponent->SetText(FText::FromString(FString::Printf(TEXT("%.2fm"), Length)));
 
-		AddLine(StartPosition, EndPosition, FLinearColor::Red);
+		AddLine(StartPosition, EndPosition, FLinearColor::Red, ESceneDepthPriorityGroup::SDPG_World);
 	}
 
 	{
@@ -121,7 +121,7 @@ void UPVScaleVisualizationComponent::UpdateScaleVisualizations()
 		TextRenderComponent->SetMaterial(0, LoadObject<UMaterialInterface>(nullptr, TEXT("/ProceduralVegetationEditor/Materials/ScaleTextMaterial.ScaleTextMaterial")));
 		TextRenderComponent->SetText(FText::FromString(FString::Printf(TEXT("%.2fm"), Length)));
 
-		AddLine(StartPosition, EndPosition, FLinearColor::Green);
+		AddLine(StartPosition, EndPosition, FLinearColor::Green, ESceneDepthPriorityGroup::SDPG_World);
 	}
 
 	{
@@ -146,6 +146,6 @@ void UPVScaleVisualizationComponent::UpdateScaleVisualizations()
 		TextRenderComponent->SetMaterial(0, LoadObject<UMaterialInterface>(nullptr, TEXT("/ProceduralVegetationEditor/Materials/ScaleTextMaterial.ScaleTextMaterial")));
 		TextRenderComponent->SetText(FText::FromString(FString::Printf(TEXT("%.2fm"), Length)));
 		
-		AddLine(StartPosition, EndPosition, FLinearColor::Blue);
+		AddLine(StartPosition, EndPosition, FLinearColor::Blue, ESceneDepthPriorityGroup::SDPG_World);
 	}
 }
