[PVE] Interactive Tools

#Synced-CL 51373704

diff --git a/Engine/Plugins/Experimental/ProceduralVegetationEditor/Source/ProceduralVegetation/Private/Facades/PVBudVectorsFacade.cpp b/Engine/Plugins/Experimental/ProceduralVegetationEditor/Source/ProceduralVegetation/Private/Facades/PVBudVectorsFacade.cpp
new file mode 100644
--- /dev/null
+++ b/Engine/Plugins/Experimental/ProceduralVegetationEditor/Source/ProceduralVegetation/Private/Facades/PVBudVectorsFacade.cpp
@@ -0,0 +1,42 @@
+﻿// Copyright Epic Games, Inc. All Rights Reserved.
+
+#include "Facades/PVBudVectorsFacade.h"
+
+#include "Facades/PVAttributesNames.h"
+
+namespace PV::Facades
+{
+	FBudVectorsFacade::FBudVectorsFacade(FManagedArrayCollection& InCollection)
+			: BudDirection(InCollection, AttributeNames::BudDirection, GroupNames::PointGroup)
+	{}
+
+	FBudVectorsFacade::FBudVectorsFacade(const FManagedArrayCollection& InCollection)
+		: BudDirection(InCollection, AttributeNames::BudDirection, GroupNames::PointGroup)
+	{}
+
+	bool FBudVectorsFacade::IsValid() const
+	{
+		return BudDirection.IsValid();
+	}
+
+	const TArray<FVector3f>& FBudVectorsFacade::GetBudDirection(int32 Index) const
+	{
+		if(BudDirection.IsValid() && BudDirection.IsValidIndex(Index))
+		{
+			return BudDirection[Index];
+		}
+
+		static const TArray<FVector3f> EmptyBudDirections;
+		return EmptyBudDirections;
+	}
+
+	TManagedArray<TArray<FVector3f>>& FBudVectorsFacade::ModifyBudDirections()
+	{
+		return BudDirection.Modify();
+	}
+
+	const TManagedArray<TArray<FVector3f>>& FBudVectorsFacade::GetBudDirections() const
+	{
+		return BudDirection.Get();
+	}
+}
diff --git a/Engine/Plugins/Experimental/ProceduralVegetationEditor/Source/ProceduralVegetation/Private/Facades/PVBudVectorsFacade.h b/Engine/Plugins/Experimental/ProceduralVegetationEditor/Source/ProceduralVegetation/Private/Facades/PVBudVectorsFacade.h
deleted file mode 100644
diff --git a/Engine/Plugins/Experimental/ProceduralVegetationEditor/Source/ProceduralVegetation/Private/Facades/PVVectorsFacade.cpp b/Engine/Plugins/Experimental/ProceduralVegetationEditor/Source/ProceduralVegetation/Private/Facades/PVVectorsFacade.cpp
deleted file mode 100644
diff --git a/Engine/Plugins/Experimental/ProceduralVegetationEditor/Source/ProceduralVegetation/Private/Implementations/PVObjectInteraction.cpp b/Engine/Plugins/Experimental/ProceduralVegetationEditor/Source/ProceduralVegetation/Private/Implementations/PVObjectInteraction.cpp
--- a/Engine/Plugins/Experimental/ProceduralVegetationEditor/Source/ProceduralVegetation/Private/Implementations/PVObjectInteraction.cpp
+++ b/Engine/Plugins/Experimental/ProceduralVegetationEditor/Source/ProceduralVegetation/Private/Implementations/PVObjectInteraction.cpp
@@ -258,10 +258,12 @@ void FPVObjectInteraction::ObjectInteraction(const FManagedArrayCollection& InCo
 									);
 								}
 							}
-							for (FVector3f& BudDirection : BudDirections[SiblingPoint])
-							{
-								BudDirection = Rotation.RotateVector(BudDirection);
-							}
+							
+							TArray<FVector3f>& BudDirection = BudDirections[SiblingPoint];
+							BudDirection[PV::Facades::BudDirectionsApical] = Rotation.RotateVector(BudDirection[PV::Facades::BudDirectionsApical]);
+							BudDirection[PV::Facades::BudDirectionsAxillary] = Rotation.RotateVector(BudDirection[PV::Facades::BudDirectionsAxillary]);
+							BudDirection[PV::Facades::BudDirectionsGuideCurve] = Rotation.RotateVector(BudDirection[PV::Facades::BudDirectionsGuideCurve]);
+							BudDirection[PV::Facades::BudDirectionsUpVector] = Rotation.RotateVector(BudDirection[PV::Facades::BudDirectionsUpVector]);
 						}
 
 						PreviousSiblingPosition = PointPositions[SiblingPoint];
@@ -308,10 +310,12 @@ void FPVObjectInteraction::ObjectInteraction(const FManagedArrayCollection& InCo
 							}
 						}
 					}
-					for (FVector3f& BudDirection : BudDirections[NextPointIndex])
-					{
-						BudDirection = Rotation.RotateVector(BudDirection);
-					}
+
+					TArray<FVector3f>& BudDirection = BudDirections[NextPointIndex];
+					BudDirection[PV::Facades::BudDirectionsApical] = Rotation.RotateVector(BudDirection[PV::Facades::BudDirectionsApical]);
+					BudDirection[PV::Facades::BudDirectionsAxillary] = Rotation.RotateVector(BudDirection[PV::Facades::BudDirectionsAxillary]);
+					BudDirection[PV::Facades::BudDirectionsGuideCurve] = Rotation.RotateVector(BudDirection[PV::Facades::BudDirectionsGuideCurve]);
+					BudDirection[PV::Facades::BudDirectionsUpVector] = Rotation.RotateVector(BudDirection[PV::Facades::BudDirectionsUpVector]);
 
 					const TArray<int32> BranchChildren = BranchFacade.GetImmediateChildren(BranchIndex);
 					for (int32 ChildBranch : BranchChildren)
@@ -337,6 +341,11 @@ void FPVObjectInteraction::ObjectInteraction(const FManagedArrayCollection& InCo
 						{
 							PointsToRemove[BranchPoints[SiblingPointIndex]] = true;
 						}
+						if (BranchPointIndex <= 1)
+						{
+							PointsToRemove[BranchPoints[0]] = true;
+							BranchesToRemove[BranchIndex] = true;
+						}
 						for (const int32 FoliageEntryId : FoliageFacade.GetFoliageEntryIdsForBranch(BranchIndex))
 						{
 							const float FoliageLFR = FoliageFacade.GetLengthFromRoot(FoliageEntryId);
@@ -347,13 +356,13 @@ void FPVObjectInteraction::ObjectInteraction(const FManagedArrayCollection& InCo
 						}
 
 						const float Ratio = static_cast<float>(BranchPointIndex) / static_cast<float>(BranchPoints.Num());
-						
+
 						// Rescale the branches based on the cut point
 						for (int32 Index = 0; Index <= BranchPointIndex; ++Index)
 						{
 							PointFacade.SetPointScale(BranchPoints[Index], PointFacade.GetPointScale(BranchPoints[Index]) * Ratio);
 						}
-						
+
 						for (int32 ChildBranch : BranchFacade.GetChildren(BranchIndex))
 						{
 							if (const int32 ChildBranchIndex = BranchFacade.GetBranchNumbers().Find(ChildBranch); ChildBranchIndex != INDEX_NONE)
diff --git a/Engine/Plugins/Experimental/ProceduralVegetationEditor/Source/ProceduralVegetation/Private/Nodes/PVObjectInteractionSettings.cpp b/Engine/Plugins/Experimental/ProceduralVegetationEditor/Source/ProceduralVegetation/Private/Nodes/PVObjectInteractionSettings.cpp
--- a/Engine/Plugins/Experimental/ProceduralVegetationEditor/Source/ProceduralVegetation/Private/Nodes/PVObjectInteractionSettings.cpp
+++ b/Engine/Plugins/Experimental/ProceduralVegetationEditor/Source/ProceduralVegetation/Private/Nodes/PVObjectInteractionSettings.cpp
@@ -9,19 +9,6 @@
 
 #include "Implementations/PVObjectInteraction.h"
 
-#if WITH_EDITOR
-TMap<UObject*, FTransform> UPVObjectInteractionSettings::GetViewportObjects() const
-{
-	TMap<UObject*, FTransform> ViewportObjects;
-	for (auto& Item : ObjectInteractionSettings.Colliders)
-	{
-		ViewportObjects.Add(Item.Mesh.Get(), Item.Transform);
-	}
-
-	return ViewportObjects;
-}
-#endif
-
 FPCGDataTypeIdentifier UPVObjectInteractionSettings::GetInputPinTypeIdentifier() const
 {
 	return FPCGDataTypeIdentifier{FPVDataTypeInfoGrowth::AsId()};
diff --git a/Engine/Plugins/Experimental/ProceduralVegetationEditor/Source/ProceduralVegetation/Private/Nodes/PVSkeletonEditSettings.cpp b/Engine/Plugins/Experimental/ProceduralVegetationEditor/Source/ProceduralVegetation/Private/Nodes/PVSkeletonEditSettings.cpp
new file mode 100644
--- /dev/null
+++ b/Engine/Plugins/Experimental/ProceduralVegetationEditor/Source/ProceduralVegetation/Private/Nodes/PVSkeletonEditSettings.cpp
@@ -0,0 +1,232 @@
+// Copyright Epic Games, Inc. All Rights Reserved.
+
+#include "Nodes/PVSkeletonEditSettings.h"
+
+#include "DataTypes/PVGrowthData.h"
+
+#include "Facades/PVBranchFacade.h"
+#include "Facades/PVBudVectorsFacade.h"
+#include "Facades/PVPointFacade.h"
+
+#include "Implementations/PVCarve.h"
+
+#if WITH_EDITOR
+
+FName UPVSkeletonEditSettings::GetDefaultNodeName() const
+{
+	return FName(TEXT("Skeleton Edit"));
+}
+
+FText UPVSkeletonEditSettings::GetNodeTooltipText() const
+{
+	return FText::FromString("Skeleton Edit");
+}
+
+FLinearColor UPVSkeletonEditSettings::GetNodeTitleColor() const
+{
+	return FLinearColor(FColor::Emerald);
+}
+
+#endif
+
+FPCGDataTypeIdentifier UPVSkeletonEditSettings::GetInputPinTypeIdentifier() const
+{
+	return FPCGDataTypeIdentifier{FPVDataTypeInfoGrowth::AsId()};
+}
+
+FPCGDataTypeIdentifier UPVSkeletonEditSettings::GetOutputPinTypeIdentifier() const
+{
+	return FPCGDataTypeIdentifier{FPVDataTypeInfoGrowth::AsId()};
+}
+
+FPCGElementPtr UPVSkeletonEditSettings::CreateElement() const
+{
+	return MakeShared<FPVSkeletonEditElement>();
+}
+
+bool FPVSkeletonEditElement::ExecuteInternal(FPCGContext* InContext) const
+{
+	TRACE_CPUPROFILER_EVENT_SCOPE(FPVObjectInteractionElement::ExecuteInternal);
+
+	check(InContext);
+
+	const UPVSkeletonEditSettings* InputSettings = InContext->GetInputSettings<UPVSkeletonEditSettings>();
+	check(InputSettings);
+
+	const TArray<FPCGTaggedData>& Inputs = InContext->InputData.GetInputsByPin(PCGPinConstants::DefaultInputLabel);
+	if (!Inputs.IsEmpty())
+	{
+		if (const UPVGrowthData* InputData = Cast<UPVGrowthData>(Inputs[0].Data))
+		{
+			FManagedArrayCollection OutCollection = InputData->GetCollection();
+
+			PV::Facades::FPointFacade PointFacade(OutCollection);
+			PV::Facades::FBranchFacade BranchFacade(OutCollection);
+			PV::Facades::FFoliageFacade FoliageFacade(OutCollection);
+			PV::Facades::FBudVectorsFacade BudVectorsFacade(OutCollection);
+
+			const auto GetParentPoint = [&](const int32 BranchIndex) -> int32
+				{
+					const int32 RootPointIndex = BranchFacade.GetRootPoint(BranchIndex);
+					if (RootPointIndex != INDEX_NONE)
+					{
+						const float PointLFR = PointFacade.GetLengthFromRoot(RootPointIndex);
+						const int32 ParentBranchIndex = BranchFacade.GetParentBranchIndex(BranchIndex);
+						if (ParentBranchIndex != INDEX_NONE)
+						{
+							const TArray<int32>& ParentBranchPoints = BranchFacade.GetPoints(ParentBranchIndex);
+							for (int32 BranchPointIndex = 0; BranchPointIndex < ParentBranchPoints.Num(); BranchPointIndex++)
+							{
+								if (PointFacade.GetLengthFromRoot(ParentBranchPoints[BranchPointIndex]) >= PointLFR)
+								{
+									return ParentBranchPoints[BranchPointIndex];
+								}
+							}
+						}
+					}
+					return INDEX_NONE;
+				};
+
+			const TArray<FVector3f>& Offsets = InputSettings->SkeletonEditSettings.RelativeOffsets;
+
+			const int32* MaxBudNumber = Algo::MaxElement(PointFacade.GetBudNumbersAttribute());
+			if (MaxBudNumber && Offsets.Num() <= *MaxBudNumber + 1)
+			{
+				TManagedArray<FVector3f>& Positions = PointFacade.ModifyPositions();
+				TManagedArray<FVector3f>& FoliagePositions = FoliageFacade.ModifyPivotPositions();
+				TManagedArray<TArray<FVector3f>>& BudDirections = BudVectorsFacade.ModifyBudDirections();
+
+				const TArray<FVector3f> OriginalPositions = Positions.GetConstArray();
+
+				for (int32 BranchIndex = 0; BranchIndex < BranchFacade.GetElementCount(); BranchIndex++)
+				{
+					const TArray<int32>& BranchPoints = BranchFacade.GetPoints(BranchIndex);
+					for (int32 BranchPointIndex = 0; BranchPointIndex < BranchPoints.Num(); BranchPointIndex++)
+					{
+						const int32 PointIndex = BranchPoints[BranchPointIndex];
+						const int32 PointBudNumber = PointFacade.GetBudNumber(PointIndex);
+						Positions[PointIndex] += Offsets[PointBudNumber];
+
+						const int32 PreviousPointIndex = BranchPointIndex > 0
+							? BranchPoints[BranchPointIndex - 1]
+							: GetParentPoint(BranchIndex);
+						if (PreviousPointIndex != INDEX_NONE)
+						{
+							const FVector3f OldDirection = OriginalPositions[PointIndex] - OriginalPositions[PreviousPointIndex];
+							const FVector3f NewDirection = Positions[PointIndex] - Positions[PreviousPointIndex];
+							const FQuat4f Rotation = FQuat4f::FindBetweenVectors(OldDirection, NewDirection);
+
+							TArray<FVector3f>& BudDirection = BudDirections[PointIndex];
+							BudDirection[PV::Facades::BudDirectionsApical] = Rotation.RotateVector(BudDirection[PV::Facades::BudDirectionsApical]);
+							BudDirection[PV::Facades::BudDirectionsAxillary] = Rotation.RotateVector(BudDirection[PV::Facades::BudDirectionsAxillary]);
+							BudDirection[PV::Facades::BudDirectionsGuideCurve] = Rotation.RotateVector(BudDirection[PV::Facades::BudDirectionsGuideCurve]);
+							BudDirection[PV::Facades::BudDirectionsUpVector] = Rotation.RotateVector(BudDirection[PV::Facades::BudDirectionsUpVector]);
+
+							const float PointLFR = PointFacade.GetLengthFromRoot(PointIndex);
+							const float PreviousPointLFR = PointFacade.GetLengthFromRoot(PreviousPointIndex);
+							for (const int32 FoliageEntryId : FoliageFacade.GetFoliageEntryIdsForBranch(BranchIndex))
+							{
+								const float FoliageLFR = FoliageFacade.GetLengthFromRoot(FoliageEntryId);
+								if (FoliageLFR <= PointLFR && FoliageLFR > PreviousPointLFR)
+								{
+									FoliagePositions[FoliageEntryId] += Offsets[PointBudNumber];
+									FoliageFacade.SetUpVector(
+										FoliageEntryId,
+										Rotation.RotateVector(FoliageFacade.GetUpVector(FoliageEntryId))
+									);
+									FoliageFacade.SetNormalVector(
+										FoliageEntryId,
+										Rotation.RotateVector(FoliageFacade.GetNormalVector(FoliageEntryId))
+									);
+								}
+							}
+						}
+					}
+				}
+			}
+			
+			TArray<bool> RemovedPoints = InputSettings->SkeletonEditSettings.RemovedPoints;
+
+			TArray<bool> PointsToRemove;
+			TArray<bool> BranchesToRemove;
+			TArray<bool> FoliageInstancesToRemove;
+			if (MaxBudNumber && RemovedPoints.Num() <= *MaxBudNumber + 1)
+			{
+				PointsToRemove.Init(false, PointFacade.GetElementCount());
+				BranchesToRemove.Init(false, BranchFacade.GetElementCount());
+				FoliageInstancesToRemove.Init(false, FoliageFacade.NumFoliageEntries());
+
+				for (int32 BranchIndex = 0; BranchIndex < BranchFacade.GetElementCount(); BranchIndex++)
+				{
+					const TArray<int32>& BranchPoints = BranchFacade.GetPoints(BranchIndex);
+					bool bAllPointsRemoved = true;
+					int32 FirstPointRemovedIndex = INDEX_NONE;
+					for (int32 BranchPointIndex = 0; BranchPointIndex < BranchPoints.Num(); BranchPointIndex++)
+					{
+						const int32 PointIndex = BranchPoints[BranchPointIndex];
+						const int32 PointBudNumber = PointFacade.GetBudNumber(PointIndex);
+						PointsToRemove[PointIndex] = RemovedPoints[PointBudNumber];
+						
+						bAllPointsRemoved &= PointsToRemove[PointIndex];
+						if (PointsToRemove[PointIndex])
+						{
+							if (FirstPointRemovedIndex == INDEX_NONE)
+							{
+								FirstPointRemovedIndex = BranchPointIndex;
+							}
+							const int32 PreviousPointIndex = BranchPointIndex > 0
+								? BranchPoints[BranchPointIndex - 1]
+								: GetParentPoint(BranchIndex);
+							const float PointLFR = PointFacade.GetLengthFromRoot(PointIndex);
+							const float PreviousPointLFR = PointFacade.GetLengthFromRoot(PreviousPointIndex);
+							for (const int32 FoliageEntryId : FoliageFacade.GetFoliageEntryIdsForBranch(BranchIndex))
+							{
+								const float FoliageLFR = FoliageFacade.GetLengthFromRoot(FoliageEntryId);
+								if (FoliageLFR <= PointLFR && FoliageLFR > PreviousPointLFR)
+								{
+									FoliageInstancesToRemove[FoliageEntryId] = true;
+								}
+							}
+						}
+					}
+					BranchesToRemove[BranchIndex] = bAllPointsRemoved;
+
+					if (FirstPointRemovedIndex != INDEX_NONE)
+					{
+						const float Ratio = static_cast<float>(FirstPointRemovedIndex) / static_cast<float>(BranchPoints.Num());
+						// Rescale the branches based on the cut point
+						for (int32 Index = 0; Index <= FirstPointRemovedIndex; ++Index)
+						{
+							PointFacade.SetPointScale(BranchPoints[Index], PointFacade.GetPointScale(BranchPoints[Index]) * Ratio);
+						}
+						for (int32 ChildBranch : BranchFacade.GetChildren(BranchIndex))
+						{
+							if (const int32 ChildBranchIndex = BranchFacade.GetBranchNumbers().Find(ChildBranch); ChildBranchIndex != INDEX_NONE)
+							{
+								for (int32 ChildBranchPoint : BranchFacade.GetPoints(ChildBranchIndex))
+								{
+									PointFacade.SetPointScale(ChildBranchPoint, PointFacade.GetPointScale(ChildBranchPoint) * Ratio);
+								}
+							}
+						}
+					}
+				}
+
+				const FManagedArrayCollection InCollection = OutCollection;
+				FPVCarve::RemoveEntriesAndRecomputeAttributes(
+					OutCollection,
+					InCollection,
+					PointsToRemove,
+					BranchesToRemove,
+					FoliageInstancesToRemove
+				);
+			}
+
+			UPVData* OutManagedArrayCollectionData = FPCGContext::NewObject_AnyThread<UPVGrowthData>(InContext);
+			OutManagedArrayCollectionData->Initialize(MoveTemp(OutCollection));
+			InContext->OutputData.TaggedData.Emplace(OutManagedArrayCollectionData);
+		}
+	}
+
+	return true;
+}
diff --git a/Engine/Plugins/Experimental/ProceduralVegetationEditor/Source/ProceduralVegetation/Public/Facades/PVBudVectorsFacade.h b/Engine/Plugins/Experimental/ProceduralVegetationEditor/Source/ProceduralVegetation/Public/Facades/PVBudVectorsFacade.h
new file mode 100644
--- /dev/null
+++ b/Engine/Plugins/Experimental/ProceduralVegetationEditor/Source/ProceduralVegetation/Public/Facades/PVBudVectorsFacade.h
@@ -0,0 +1,40 @@
+﻿// Copyright Epic Games, Inc. All Rights Reserved.
+
+#pragma once
+
+#include "CoreMinimal.h"
+#include "GeometryCollection/ManagedArrayAccessor.h"
+
+namespace PV::Facades
+{
+	constexpr int BudDirectionsApical = 0;
+	constexpr int BudDirectionsAxillary = 1;
+	constexpr int BudDirectionsLightOptimalIndex = 2;
+	constexpr int BudDirectionsLightSubOptimal = 3;
+	constexpr int BudDirectionsGuideCurve = 4;
+	constexpr int BudDirectionsUpVector = 5;
+	
+	constexpr int BudLightDetectedAvailableIndex = 0;
+	
+	/**
+	 * FBudVectorsFacade is used to access and manipulate the BudVectors data from Point Group with in the ProceduralVegetation's FManagedArrayCollection
+	 * BudVectors represent different growth related vectors in plant data and are associated with each point. 
+	 */
+	class FBudVectorsFacade
+	{
+	public:
+		FBudVectorsFacade(FManagedArrayCollection& InCollection);
+		FBudVectorsFacade(const FManagedArrayCollection& InCollection);
+
+		bool IsValid() const;
+
+		const TArray<FVector3f>& GetBudDirection(int32 Index) const;
+
+		TManagedArray<TArray<FVector3f>>& ModifyBudDirections();
+
+		const TManagedArray<TArray<FVector3f>>& GetBudDirections() const;
+		
+	private:
+		TManagedArrayAccessor<TArray<FVector3f>> BudDirection;
+	};
+}
diff --git a/Engine/Plugins/Experimental/ProceduralVegetationEditor/Source/ProceduralVegetation/Public/Nodes/PVObjectInteractionSettings.h b/Engine/Plugins/Experimental/ProceduralVegetationEditor/Source/ProceduralVegetation/Public/Nodes/PVObjectInteractionSettings.h
--- a/Engine/Plugins/Experimental/ProceduralVegetationEditor/Source/ProceduralVegetation/Public/Nodes/PVObjectInteractionSettings.h
+++ b/Engine/Plugins/Experimental/ProceduralVegetationEditor/Source/ProceduralVegetation/Public/Nodes/PVObjectInteractionSettings.h
@@ -66,8 +66,6 @@ public:
 	{
 		return TArray{EPVRenderType::PointData};
 	}
-
-	virtual TMap<UObject*, FTransform> GetViewportObjects() const override;
 #endif
 
 protected:
diff --git a/Engine/Plugins/Experimental/ProceduralVegetationEditor/Source/ProceduralVegetation/Public/Nodes/PVSkeletonEditSettings.h b/Engine/Plugins/Experimental/ProceduralVegetationEditor/Source/ProceduralVegetation/Public/Nodes/PVSkeletonEditSettings.h
new file mode 100644
--- /dev/null
+++ b/Engine/Plugins/Experimental/ProceduralVegetationEditor/Source/ProceduralVegetation/Public/Nodes/PVSkeletonEditSettings.h
@@ -0,0 +1,124 @@
+// Copyright Epic Games, Inc. All Rights Reserved.
+
+#pragma once
+
+#include "Nodes/PVBaseSettings.h"
+
+#include "PVSkeletonEditSettings.generated.h"
+
+UENUM()
+enum class EPointSelectionSmoothnessMethod: uint8
+{
+	Linear,
+	Smooth,
+	Sphere,
+	Root,
+	Sharp,
+	Sine,
+	Constant,
+};
+
+UENUM()
+enum class ESkeletonSelectionMode: uint8
+{
+	SelectByDistance,
+	SelectByNeighbours,
+};
+
+UENUM()
+enum class EBranchRemovalMode: uint8
+{
+	RemoveBranch,
+	TrimBranch,
+};
+
+UENUM()
+enum class ESkeletonEditMode: uint8
+{
+	TRS,
+	BranchRemoval,
+};
+
+USTRUCT()
+struct FPVSkeletonSelectionParams
+{
+	GENERATED_BODY()
+
+	UPROPERTY(EditAnywhere, Category = "Selection")
+	EPointSelectionSmoothnessMethod SelectionFalloff = EPointSelectionSmoothnessMethod::Smooth;
+
+	UPROPERTY(EditAnywhere, Category = "Selection")
+	ESkeletonSelectionMode SkeletonSelectionMode = ESkeletonSelectionMode::SelectByNeighbours;
+
+	UPROPERTY(EditAnywhere, Category = "Selection",
+		meta = (UIMin=0, ClampMin=0, EditCondition="SkeletonSelectionMode == ESkeletonSelectionMode::SelectByNeighbours", EditConditionHides))
+	int32 NumPointsSelection = 0;
+
+	UPROPERTY(EditAnywhere, Category = "Selection",
+		meta = (UIMin=0, ClampMin=0, EditCondition="SkeletonSelectionMode == ESkeletonSelectionMode::SelectByDistance", EditConditionHides))
+	float SelectionDistance = 0.0f;
+
+	UPROPERTY(EditAnywhere, Category = "Selection")
+	bool bRootSelectsBranch = false;
+};
+
+USTRUCT()
+struct FPVBranchRemovalParams
+{
+	GENERATED_BODY()
+
+	UPROPERTY(EditAnywhere, Category = "Branch Removal")
+	EBranchRemovalMode BranchRemovalMode = EBranchRemovalMode::TrimBranch;
+};
+
+USTRUCT()
+struct FPVSkeletonEditParams
+{
+	GENERATED_BODY()
+
+	UPROPERTY(EditAnywhere, Category = "Tool")
+	ESkeletonEditMode SkeletonEditMode = ESkeletonEditMode::TRS;
+
+	UPROPERTY(EditAnywhere, Category = "Selection",
+		meta=(ShowOnlyInnerProperties, EditCondition="SkeletonEditMode == ESkeletonEditMode::TRS", EditConditionHides))
+	FPVSkeletonSelectionParams SkeletonSelectionParams;
+
+	UPROPERTY(EditAnywhere, Category = "Selection",
+		meta=(ShowOnlyInnerProperties, EditCondition="SkeletonEditMode == ESkeletonEditMode::BranchRemoval", EditConditionHides))
+	FPVBranchRemovalParams BranchRemovalParams;
+
+	UPROPERTY(EditAnywhere, EditFixedSize, Category = "Edits")
+	TArray<FVector3f> RelativeOffsets;
+
+	UPROPERTY(EditAnywhere, EditFixedSize, Category = "Edits")
+	TArray<bool> RemovedPoints;
+};
+
+UCLASS(BlueprintType, ClassGroup = (Procedural))
+class UPVSkeletonEditSettings : public UPVBaseSettings
+{
+	GENERATED_BODY()
+
+public:
+#if WITH_EDITOR
+	virtual FName GetDefaultNodeName() const override;
+	virtual FText GetNodeTooltipText() const override;
+	virtual FLinearColor GetNodeTitleColor() const override;
+#endif
+
+protected:
+	virtual FPCGDataTypeIdentifier GetInputPinTypeIdentifier() const override;
+	virtual FPCGDataTypeIdentifier GetOutputPinTypeIdentifier() const override;
+
+	virtual FPCGElementPtr CreateElement() const override;
+
+public:
+	UPROPERTY(EditAnywhere, Category = "Object Interaction", meta = (ShowOnlyInnerProperties))
+	FPVSkeletonEditParams SkeletonEditSettings;
+};
+
+class FPVSkeletonEditElement : public FPVBaseElement
+{
+protected:
+	virtual bool ExecuteInternal(FPCGContext* InContext) const override;
+};
diff --git a/Engine/Plugins/Experimental/ProceduralVegetationEditor/Source/ProceduralVegetationEditor/Private/PVEditor.cpp b/Engine/Plugins/Experimental/ProceduralVegetationEditor/Source/ProceduralVegetationEditor/Private/PVEditor.cpp
--- a/Engine/Plugins/Experimental/ProceduralVegetationEditor/Source/ProceduralVegetationEditor/Private/PVEditor.cpp
+++ b/Engine/Plugins/Experimental/ProceduralVegetationEditor/Source/ProceduralVegetationEditor/Private/PVEditor.cpp
@@ -46,6 +46,9 @@
 
 #include "Subsystems/PCGEngineSubsystem.h"
 
+#include "Tools/EdModeInteractiveToolsContext.h"
+#include "Tools/PVEditorMode.h"
+
 #include "Widgets/SPVExportSelectionDialog.h"
 #include "Widgets/Docking/SDockTab.h"
 
@@ -381,7 +384,12 @@ TSharedRef<SPCGEditorViewport> FPVEditor::CreateViewportWidget()
 	TSharedPtr<SPVEditorViewport> ViewportWidget = nullptr;
 	if (EditorViewport == nullptr)
 	{
-		EditorViewport = ViewportWidget = SNew(SPVEditorViewport);
+		FAssetEditorModeManager* ModeManager = static_cast<FAssetEditorModeManager*>(&GetEditorModeManager());
+		EditorViewport = ViewportWidget = SNew(SPVEditorViewport).ModeTools(ModeManager);
+
+		ModeManager->SetPreviewScene(ViewportWidget->GetAdvancedPreviewScene());
+		ModeManager->SetSupportsViewportITF(true);
+		ModeManager->ActivateMode(UPVEditorMode::EM_PVEditorModeId);
 	}
 	else
 	{
@@ -396,6 +404,16 @@ TSharedRef<SCollectionSpreadSheetWidget> FPVEditor::CreateCollectionSpreadSheetW
 	return SNew(SCollectionSpreadSheetWidget);
 }
 
+UPVEditorMode* FPVEditor::GetEditorMode() const
+{
+	return Cast<UPVEditorMode>(EditorModeManager->GetActiveScriptableMode(UPVEditorMode::EM_PVEditorModeId));
+}
+
+UPVEditorMode* FPVEditor::GetEditorModeChecked() const
+{
+	return CastChecked<UPVEditorMode>(EditorModeManager->GetActiveScriptableMode(UPVEditorMode::EM_PVEditorModeId));
+}
+
 FPCGStack FPVEditor::GetStackFromPin(const UPCGPin* InPin) const
 {
 	check(InPin);
@@ -439,6 +457,19 @@ void FPVEditor::OnNodeInspected(const FPCGStack& InPCGStack)
 {
 	UpdateInspectedCollection();
 	UpdateStats();
+
+	if (NodeBeingInspected)
+	{
+		if (UPVEditorMode* const EditorMode = GetEditorMode())
+		{
+			FPVToolsRegistry& ToolRegistry = EditorMode->GetToolRegistry();
+			const UPVData* const NodeData = GetDataFromPin(GetInPinFromNode(GetNodeBeingInspected()));
+			if (!ToolRegistry.StartToolForNode(GetNodeSettings(GetNodeBeingInspected()), NodeData))
+			{
+				ToolRegistry.ShutdownCurrentTool(EToolShutdownType::Accept);
+			}
+		}
+	}
 }
 
 void FPVEditor::ChangeNodeInspection(UPCGEditorGraphNodeBase* InNode)
@@ -447,16 +478,23 @@ void FPVEditor::ChangeNodeInspection(UPCGEditorGraphNodeBase* InNode)
 	{
 		return;
 	}
-	
+
 	if (UPVBaseSettings* Settings = GetNodeSettings(InNode->GetPCGNode()))
 	{
 		if (EditorViewport.IsValid())
 		{
 			EditorViewport->OnNodeInspectionChanged(Settings);
 		}
-
-		SetNodeBeingInspected(InNode);
 	}
+	else
+	{
+		if (EditorViewport.IsValid())
+		{
+			EditorViewport->ClearNodeInspection();
+		}
+	}
+
+	SetNodeBeingInspected(InNode);
 }
 
 void FPVEditor::OnLockNodeSelection()
@@ -469,6 +507,11 @@ void FPVEditor::OnLockNodeSelection()
 		NodeBeingInspectedSettings->bLockInspection = !NodeBeingInspectedSettings->bLockInspection;
 	}
 
+	if (!SelectedNode || !SelectedNodeSettings)
+	{
+		return;
+	}
+
 	if (NodeBeingInspected != SelectedNode)
 	{
 		SelectedNodeSettings->bLockInspection = true;
@@ -575,11 +618,6 @@ void FPVEditor::UpdateCollectionSpreadSheet()
 			}
 		}
 
-		// TODO: Check if there is a way to select default output
-		// if (GetOutPinFromNode(NodeBeingInspected))
-		// {
-		// 	CollectionSpreadSheetWidget->GetCollectionTable()->SetSelectedOutput(PCGPinConstants::DefaultOutputLabel);
-		// }
 
 		CollectionSpreadSheetWidget->RefreshWidget();
 	}
diff --git a/Engine/Plugins/Experimental/ProceduralVegetationEditor/Source/ProceduralVegetationEditor/Private/PVEditor.h b/Engine/Plugins/Experimental/ProceduralVegetationEditor/Source/ProceduralVegetationEditor/Private/PVEditor.h
--- a/Engine/Plugins/Experimental/ProceduralVegetationEditor/Source/ProceduralVegetationEditor/Private/PVEditor.h
+++ b/Engine/Plugins/Experimental/ProceduralVegetationEditor/Source/ProceduralVegetationEditor/Private/PVEditor.h
@@ -67,6 +67,9 @@ protected:
 	static TSharedRef<SCollectionSpreadSheetWidget> CreateCollectionSpreadSheetWidget();
 
 private:
+	UPVEditorMode* GetEditorMode() const;
+	UPVEditorMode* GetEditorModeChecked() const;
+
 	FPCGStack GetStackFromPin(const UPCGPin* InPin) const;
 
 	const UPVData* GetDataFromPin(const UPCGPin* InPin);
diff --git a/Engine/Plugins/Experimental/ProceduralVegetationEditor/Source/ProceduralVegetationEditor/Private/SPVEditorViewport.cpp b/Engine/Plugins/Experimental/ProceduralVegetationEditor/Source/ProceduralVegetationEditor/Private/SPVEditorViewport.cpp
--- a/Engine/Plugins/Experimental/ProceduralVegetationEditor/Source/ProceduralVegetationEditor/Private/SPVEditorViewport.cpp
+++ b/Engine/Plugins/Experimental/ProceduralVegetationEditor/Source/ProceduralVegetationEditor/Private/SPVEditorViewport.cpp
@@ -114,7 +114,7 @@ SPVEditorViewport::~SPVEditorViewport()
 
 void SPVEditorViewport::Construct(const FArguments& InArgs)
 {
-	SPCGEditorViewport::Construct(SPCGEditorViewport::FArguments());
+	SPCGEditorViewport::Construct(SPCGEditorViewport::FArguments().ModeTools(InArgs._ModeTools));
 
 	UPCGEditorSettings* PCGEditorSettings = GetMutableDefault<UPCGEditorSettings>();
 
@@ -906,6 +906,13 @@ void SPVEditorViewport::OnNodeInspectionChanged(UPVBaseSettings* InSettings)
 	}
 }
 
+void SPVEditorViewport::ClearNodeInspection()
+{
+	SupportedRenderTypes = {};
+	CurrentVisualizationMode = EPVVisualizationMode::Default;
+	InspectedNodeSettings = nullptr;
+}
+
 void SPVEditorViewport::SetOverlayText(const FText& CurrentlyLockedNodeName)
 {
 	bIsPreviewingLockedNode = !CurrentlyLockedNodeName.IsEmpty();
diff --git a/Engine/Plugins/Experimental/ProceduralVegetationEditor/Source/ProceduralVegetationEditor/Private/SPVEditorViewport.h b/Engine/Plugins/Experimental/ProceduralVegetationEditor/Source/ProceduralVegetationEditor/Private/SPVEditorViewport.h
--- a/Engine/Plugins/Experimental/ProceduralVegetationEditor/Source/ProceduralVegetationEditor/Private/SPVEditorViewport.h
+++ b/Engine/Plugins/Experimental/ProceduralVegetationEditor/Source/ProceduralVegetationEditor/Private/SPVEditorViewport.h
@@ -62,6 +62,12 @@ public:
 class SPVEditorViewport : public SPCGEditorViewport
 {
 public:
+	SLATE_BEGIN_ARGS(SPVEditorViewport)
+			: _ModeTools(nullptr)
+	{}
+		SLATE_ARGUMENT(FEditorModeTools*, ModeTools)
+	SLATE_END_ARGS()
+
 	SPVEditorViewport();
 	virtual ~SPVEditorViewport() override;
 
@@ -70,6 +76,7 @@ public:
 	FToolMenuEntry CreateVisualizationModeToolbarMenu();
 
 	void OnNodeInspectionChanged(UPVBaseSettings* InSettings);
+	void ClearNodeInspection();
 
 	void SetOverlayText(const FText& CurrentlyLockedNodeName = FText::GetEmpty());
 	void PopulateStatsOverlayText(const TArrayView<FText> TextItems);
@@ -121,7 +128,7 @@ private:
 	void OnVisualizationModeChanged(EPVVisualizationMode InMode);
 
 	bool bFocusOnNextUpdate = true;
-
+	
 protected:
 	virtual void OnSetupScene() override;
 	virtual void OnResetScene() override;
diff --git a/Engine/Plugins/Experimental/ProceduralVegetationEditor/Source/ProceduralVegetationEditor/Private/Tools/Contexts/PVToolContext.h b/Engine/Plugins/Experimental/ProceduralVegetationEditor/Source/ProceduralVegetationEditor/Private/Tools/Contexts/PVToolContext.h
new file mode 100644
--- /dev/null
+++ b/Engine/Plugins/Experimental/ProceduralVegetationEditor/Source/ProceduralVegetationEditor/Private/Tools/Contexts/PVToolContext.h
@@ -0,0 +1,22 @@
+// Copyright Epic Games, Inc. All Rights Reserved.
+
+#pragma once
+
+#include "CoreMinimal.h"
+
+#include "Nodes/PVBaseSettings.h"
+
+#include "PVToolContext.generated.h"
+
+UCLASS()
+class UPVToolContextObject : public UObject
+{
+	GENERATED_BODY()
+
+public:
+	UPROPERTY()
+	TObjectPtr<UPVBaseSettings> NodeSettings = nullptr;
+
+	UPROPERTY()
+	TObjectPtr<const UPVData> NodeInputData = nullptr;
+};
diff --git a/Engine/Plugins/Experimental/ProceduralVegetationEditor/Source/ProceduralVegetationEditor/Private/Tools/NodeTools/PVColliderTransformTool.cpp b/Engine/Plugins/Experimental/ProceduralVegetationEditor/Source/ProceduralVegetationEditor/Private/Tools/NodeTools/PVColliderTransformTool.cpp
new file mode 100644
--- /dev/null
+++ b/Engine/Plugins/Experimental/ProceduralVegetationEditor/Source/ProceduralVegetationEditor/Private/Tools/NodeTools/PVColliderTransformTool.cpp
@@ -0,0 +1,185 @@
+﻿// Copyright Epic Games, Inc. All Rights Reserved.
+
+#include "PVColliderTransformTool.h"
+
+#include "ContextObjectStore.h"
+#include "Editor.h"
+#include "EditorModeManager.h"
+#include "EngineUtils.h"
+#include "Selection.h"
+
+#include "BaseBehaviors/SingleClickBehavior.h"
+
+#include "BaseGizmos/TransformGizmoUtil.h"
+
+#include "Components/PrimitiveComponent.h"
+#include "Components/StaticMeshComponent.h"
+#include "Components/Viewport.h"
+
+#include "Tools/EdModeInteractiveToolsContext.h"
+
+#define LOCTEXT_NAMESPACE "PVColliderTransformTool"
+
+UPVColliderTransformTool::UPVColliderTransformTool()
+{
+	SetFlags(RF_Transactional);
+}
+
+void UPVColliderTransformTool::Setup()
+{
+	Super::Setup();
+
+	SetupColliderPreviews();
+
+	USingleClickInputBehavior* Behavior = NewObject<USingleClickInputBehavior>(this);
+	Behavior->Initialize(this);
+	AddInputBehavior(Behavior);
+
+	TransformProxy = NewObject<UTransformProxy>(this);
+
+	CombinedGizmo = UE::TransformGizmoUtil::Create3AxisTransformGizmo(GetToolManager(), this);
+	CombinedGizmo->SetActiveTarget(TransformProxy, GetToolManager());
+	CombinedGizmo->SetVisibility(false);
+	CombinedGizmo->ActiveGizmoMode = EToolContextTransformGizmoMode::Combined;
+	CombinedGizmo->bUseContextGizmoMode = false;
+	CombinedGizmo->bUseContextCoordinateSystem = false;
+
+	TransformProxy->OnBeginTransformEdit.AddUObject(this, &UPVColliderTransformTool::OnGizmoTransformBegin);
+	TransformProxy->OnTransformChanged.AddUObject(this, &UPVColliderTransformTool::OnGizmoTransformChanged);
+	TransformProxy->OnEndTransformEdit.AddUObject(this, &UPVColliderTransformTool::OnGizmoTransformEnd);
+}
+
+void UPVColliderTransformTool::Shutdown(EToolShutdownType ShutdownType)
+{
+	ClearColliderPreviews();
+
+	if (CombinedGizmo)
+	{
+		CombinedGizmo->ClearActiveTarget();
+		CombinedGizmo->Shutdown();
+		CombinedGizmo = nullptr;
+	}
+
+	TransformProxy = nullptr;
+
+	Super::Shutdown(ShutdownType);
+}
+
+void UPVColliderTransformTool::OnNodeSettingsUpdated()
+{
+	ClearColliderPreviews();
+	SetupColliderPreviews();
+	
+	if (SelectedComponentIndex != INDEX_NONE)
+	{
+		CombinedGizmo->SetVisibility(true);
+		CombinedGizmo->SetNewGizmoTransform(MeshComponents[SelectedComponentIndex]->GetComponentTransform());
+	}
+}
+
+FInputRayHit UPVColliderTransformTool::IsHitByClick(const FInputDeviceRay& ClickPos)
+{
+	if (FViewport* FocusedViewport = GetToolManager()->GetContextQueriesAPI()->GetFocusedViewport())
+	{
+		if (FocusedViewport->GetHitProxy(ClickPos.ScreenPosition.X, ClickPos.ScreenPosition.Y))
+		{
+			return FInputRayHit(1.0f);
+		}
+	}
+
+	CombinedGizmo->SetVisibility(false);
+	return FInputRayHit();
+}
+
+void UPVColliderTransformTool::OnClicked(const FInputDeviceRay& ClickPos)
+{
+	FViewport* FocusedViewport = GetToolManager()->GetContextQueriesAPI()->GetFocusedViewport();
+	if (!FocusedViewport)
+	{
+		return;
+	}
+	const HHitProxy* HitProxy = FocusedViewport->GetHitProxy(ClickPos.ScreenPosition.X, ClickPos.ScreenPosition.Y);
+	if (!HitProxy)
+	{
+		return;
+	}
+
+	if (HitProxy->IsA(HActor::StaticGetType()))
+	{
+		const HActor* ActorHitProxy = static_cast<const HActor*>(HitProxy);
+		if (ActorHitProxy->Actor == PreviewActor)
+		{
+			const TObjectPtr<const UPrimitiveComponent> HitComponent = ActorHitProxy->PrimComponent;
+			SelectedComponentIndex = MeshComponents.IndexOfByPredicate([&](const auto& Comp)
+				{
+					return Comp == HitComponent;
+				}
+			);
+			if (SelectedComponentIndex != INDEX_NONE)
+			{
+				CombinedGizmo->SetVisibility(true);
+				CombinedGizmo->SetNewGizmoTransform(ActorHitProxy->PrimComponent->GetComponentTransform());
+			}
+		}
+	}
+}
+
+void UPVColliderTransformTool::SetupColliderPreviews()
+{
+	check(PreviewActor);
+
+	const TObjectPtr<UPVObjectInteractionSettings> ObjectInteractionSettings = GetNodeSettings<UPVObjectInteractionSettings>();
+
+	int32 Index = 0;
+	for (const FPVColliderParams& Collider : ObjectInteractionSettings->ObjectInteractionSettings.Colliders)
+	{
+		if (UStaticMesh* ColliderMesh = Collider.Mesh.LoadSynchronous())
+		{
+			UStaticMeshComponent* PreviewMeshComponent = NewObject<UStaticMeshComponent>(PreviewActor,
+				*("ToolPreviewStaticMesh" + FString::FromInt(Index)), RF_Transient);
+			PreviewMeshComponent->SetMobility(EComponentMobility::Movable);
+			PreviewMeshComponent->SetupAttachment(PreviewActor->GetRootComponent());
+			PreviewMeshComponent->SetStaticMesh(ColliderMesh);
+			PreviewMeshComponent->SetWorldTransform(Collider.Transform);
+			PreviewMeshComponent->RegisterComponent();
+
+			MeshComponents.Add(PreviewMeshComponent);
+		}
+		++Index;
+	}
+}
+
+void UPVColliderTransformTool::ClearColliderPreviews()
+{
+	for (const TObjectPtr<UStaticMeshComponent>& MeshComponent : MeshComponents)
+	{
+		MeshComponent->DestroyComponent();
+	}
+	MeshComponents.Empty();
+}
+
+void UPVColliderTransformTool::OnGizmoTransformBegin(UTransformProxy* Proxy)
+{
+	GetToolManager()->BeginUndoTransaction(NSLOCTEXT("PVColliderTransformTool", "MoveComponent", "Move Component"));
+}
+
+void UPVColliderTransformTool::OnGizmoTransformChanged(UTransformProxy* Proxy, FTransform NewTransform)
+{
+	if (SelectedComponentIndex != INDEX_NONE && MeshComponents.IsValidIndex(SelectedComponentIndex))
+	{
+		MeshComponents[SelectedComponentIndex]->SetWorldTransform(NewTransform);
+		GetNodeSettings<UPVObjectInteractionSettings>()->ObjectInteractionSettings.Colliders[SelectedComponentIndex].Transform = NewTransform;
+	}
+}
+
+void UPVColliderTransformTool::OnGizmoTransformEnd(UTransformProxy* Proxy)
+{
+	if (NodeSettings)
+	{
+		NodeSettings->PostEditChange();
+	}
+
+	GetToolManager()->EndUndoTransaction();
+}
+
+#undef LOCTEXT_NAMESPACE
diff --git a/Engine/Plugins/Experimental/ProceduralVegetationEditor/Source/ProceduralVegetationEditor/Private/Tools/NodeTools/PVColliderTransformTool.h b/Engine/Plugins/Experimental/ProceduralVegetationEditor/Source/ProceduralVegetationEditor/Private/Tools/NodeTools/PVColliderTransformTool.h
new file mode 100644
--- /dev/null
+++ b/Engine/Plugins/Experimental/ProceduralVegetationEditor/Source/ProceduralVegetationEditor/Private/Tools/NodeTools/PVColliderTransformTool.h
@@ -0,0 +1,55 @@
+﻿// Copyright Epic Games, Inc. All Rights Reserved.
+
+#pragma once
+
+#include "Tools/PVBaseInteractiveTool.h"
+
+#include "BaseBehaviors/BehaviorTargetInterfaces.h"
+
+#include "Components/StaticMeshComponent.h"
+
+#include "Nodes/PVObjectInteractionSettings.h"
+
+#include "PVColliderTransformTool.generated.h"
+
+class UTransformProxy;
+class UCombinedTransformGizmo;
+
+UCLASS()
+class UPVColliderTransformTool
+	: public UPVBaseInteractiveTool,
+	  public IClickBehaviorTarget
+{
+	GENERATED_BODY()
+
+public:
+	UPVColliderTransformTool();
+
+	virtual void Setup() override;
+	virtual void Shutdown(EToolShutdownType ShutdownType) override;
+	
+	virtual void OnNodeSettingsUpdated() override;
+
+	virtual FInputRayHit IsHitByClick(const FInputDeviceRay& ClickPos) override;
+	virtual void OnClicked(const FInputDeviceRay& ClickPos) override;
+
+protected:
+	void SetupColliderPreviews();
+	void ClearColliderPreviews();
+
+	void OnGizmoTransformBegin(UTransformProxy* Proxy);
+	void OnGizmoTransformChanged(UTransformProxy* Proxy, FTransform NewTransform);
+	void OnGizmoTransformEnd(UTransformProxy* Proxy);
+
+private:
+	UPROPERTY(Transient)
+	TArray<TObjectPtr<UStaticMeshComponent>> MeshComponents;
+
+	int32 SelectedComponentIndex = INDEX_NONE;
+
+	UPROPERTY(Transient)
+	TObjectPtr<UTransformProxy> TransformProxy = nullptr;
+
+	UPROPERTY(Transient)
+	TObjectPtr<UCombinedTransformGizmo> CombinedGizmo = nullptr;
+};
diff --git a/Engine/Plugins/Experimental/ProceduralVegetationEditor/Source/ProceduralVegetationEditor/Private/Tools/NodeTools/PVSkeletonEditTool.cpp b/Engine/Plugins/Experimental/ProceduralVegetationEditor/Source/ProceduralVegetationEditor/Private/Tools/NodeTools/PVSkeletonEditTool.cpp
new file mode 100644
--- /dev/null
+++ b/Engine/Plugins/Experimental/ProceduralVegetationEditor/Source/ProceduralVegetationEditor/Private/Tools/NodeTools/PVSkeletonEditTool.cpp
@@ -0,0 +1,815 @@
+// Copyright Epic Games, Inc. All Rights Reserved.
+
+#include "PVSkeletonEditTool.h"
+
+#include "EditorCategoryUtils.h"
+
+#include "BaseBehaviors/SingleClickBehavior.h"
+
+#include "BaseGizmos/TransformGizmoUtil.h"
+#include "BaseGizmos/TransformProxy.h"
+
+#include "Components/Viewport.h"
+
+#include "Facades/PVBranchFacade.h"
+#include "Facades/PVPointFacade.h"
+
+#include "Nodes/PVSkeletonEditSettings.h"
+
+IMPLEMENT_HIT_PROXY(HPVSkeletonPointProxy, HHitProxy)
+
+FSkeletonPointsSelection::FSkeletonPointsSelection(const FManagedArrayCollection* const InCollection)
+	: Collection(InCollection)
+{
+	check(Collection);
+}
+
+void FSkeletonPointsSelection::SelectPoint(const int32 InBranchIndex, const int32 InBranchPointIndex)
+{
+	PrimarySelectedPoint.BranchIndex = InBranchIndex;
+	PrimarySelectedPoint.BranchPointIndex = InBranchPointIndex;
+	PrimarySelectedPoint.Weight = 1.0f;
+	PrimarySelectedPoint.NeighbourIndex = 0;
+
+	SelectedPoints.Empty();
+}
+
+void FSkeletonPointsSelection::ExtendSelection(const int32 NumNeighbours, const EPointSelectionSmoothnessMethod SmoothnessMethod)
+{
+	check(Collection);
+
+	const PV::Facades::FPointFacade PointFacade(*Collection);
+	const PV::Facades::FBranchFacade BranchFacade(*Collection);
+	if (!PointFacade.IsValid() || !BranchFacade.IsValid())
+	{
+		return;
+	}
+
+	const auto GetParentPoint = [&](const int32 BranchIndex) -> TPair<int32, int32>
+		{
+			const int32 RootPointIndex = BranchFacade.GetRootPoint(BranchIndex);
+			if (RootPointIndex != INDEX_NONE)
+			{
+				const float PointLFR = PointFacade.GetLengthFromRoot(RootPointIndex);
+				const int32 ParentBranchIndex = BranchFacade.GetParentBranchIndex(BranchIndex);
+				if (ParentBranchIndex != INDEX_NONE)
+				{
+					const TArray<int32>& ParentBranchPoints = BranchFacade.GetPoints(ParentBranchIndex);
+					for (int32 BranchPointIndex = 0; BranchPointIndex < ParentBranchPoints.Num(); BranchPointIndex++)
+					{
+						if (PointFacade.GetLengthFromRoot(ParentBranchPoints[BranchPointIndex]) >= PointLFR)
+						{
+							return {ParentBranchIndex, BranchPointIndex};
+						}
+					}
+				}
+			}
+			return {INDEX_NONE, INDEX_NONE};
+		};
+
+	const auto GetChildPoints = [&](
+		const int32 BranchIndex,
+		const int32 BranchPointIndex,
+		const TArray<int32>& ImmediateChildBranches
+	) -> TArray<TPair<int32, int32>>
+		{
+			const TArray<int32>& BranchPoints = BranchFacade.GetPoints(BranchIndex);
+
+			const float PointLFR = PointFacade.GetLengthFromRoot(BranchPoints[BranchPointIndex]);
+			const float NextPointLFR = BranchPoints.IsValidIndex(BranchPointIndex + 1)
+				? PointFacade.GetLengthFromRoot(BranchPoints[BranchPointIndex + 1])
+				: UE_MAX_FLT;
+
+			TArray<TPair<int32, int32>> ChildPoints;
+			for (int32 ChildBranch : ImmediateChildBranches)
+			{
+				const int32 ChildBranchIndex = BranchFacade.GetBranchNumbers().Find(ChildBranch);
+				if (ChildBranchIndex != INDEX_NONE)
+				{
+					const int32 ChildBranchRootPointIndex = BranchFacade.GetRootPoint(ChildBranchIndex);
+					if (ChildBranchRootPointIndex != INDEX_NONE)
+					{
+						const float ChildLFR = PointFacade.GetLengthFromRoot(ChildBranchRootPointIndex);
+						if (ChildLFR >= PointLFR && ChildLFR < NextPointLFR)
+						{
+							ChildPoints.Emplace(ChildBranchIndex, 0);
+						}
+					}
+				}
+			}
+
+			return ChildPoints;
+		};
+
+	SelectedPoints.Empty();
+	PrimarySelectedPoint.Weight = 1.0f;
+	SelectedPoints.Add(PrimarySelectedPoint);
+
+	int32 CurrentNeighbour = 0;
+	while (CurrentNeighbour <= NumNeighbours)
+	{
+		for (int32 SelectedPointIndex = 0; SelectedPointIndex < SelectedPoints.Num(); SelectedPointIndex++)
+		{
+			const FSelectedPoint& SelectedPoint = SelectedPoints[SelectedPointIndex];
+			if (SelectedPoint.NeighbourIndex != CurrentNeighbour)
+			{
+				continue;
+			}
+
+			if (SelectedPoint.BranchPointIndex == 0)
+			{
+				const auto [ParentBranchIndex, ParentBranchPointIndex] = GetParentPoint(SelectedPoint.BranchIndex);
+				if (ParentBranchIndex != INDEX_NONE && ParentBranchPointIndex != INDEX_NONE)
+				{
+					SelectedPoints.AddUnique({ParentBranchIndex, ParentBranchPointIndex, 0.0f, CurrentNeighbour});
+				}
+			}
+			const TArray<TPair<int32, int32>> ChildPoints = GetChildPoints(
+				SelectedPoint.BranchIndex,
+				SelectedPoint.BranchPointIndex,
+				BranchFacade.GetImmediateChildren(SelectedPoint.BranchIndex)
+			);
+			for (const auto [ChildBranchIndex, ChildBranchPointIndex] : ChildPoints)
+			{
+				SelectedPoints.AddUnique({ChildBranchIndex, ChildBranchPointIndex, 0.0f, CurrentNeighbour});
+			}
+
+			if (CurrentNeighbour < NumNeighbours)
+			{
+				const TArray<int32>& BranchPoints = BranchFacade.GetPoints(SelectedPoint.BranchIndex);
+				if (BranchPoints.IsValidIndex(SelectedPoint.BranchPointIndex - 1))
+				{
+					SelectedPoints.AddUnique({SelectedPoint.BranchIndex, SelectedPoint.BranchPointIndex - 1, 0.0f, CurrentNeighbour + 1});
+				}
+				if (BranchPoints.IsValidIndex(SelectedPoint.BranchPointIndex + 1))
+				{
+					SelectedPoints.AddUnique({SelectedPoint.BranchIndex, SelectedPoint.BranchPointIndex + 1, 0.0f, CurrentNeighbour + 1});
+				}
+			}
+		}
+
+		++CurrentNeighbour;
+	}
+
+	const float MaxNeighbours = NumNeighbours + 1;
+	for (FSelectedPoint& SelectedPoint : SelectedPoints)
+	{
+		SelectedPoint.Weight = GetWeightBySelectedSmoothness(
+			(MaxNeighbours - SelectedPoint.NeighbourIndex) / MaxNeighbours,
+			SmoothnessMethod
+		);
+	}
+}
+
+void FSkeletonPointsSelection::ExtendSelectionByDistance(
+	const float MaxDistance,
+	const TArray<FVector3f>& Offsets,
+	const EPointSelectionSmoothnessMethod SmoothnessMethod
+)
+{
+	check(Collection);
+
+	const PV::Facades::FPointFacade PointFacade(*Collection);
+	const PV::Facades::FBranchFacade BranchFacade(*Collection);
+	if (!PointFacade.IsValid() || !BranchFacade.IsValid())
+	{
+		return;
+	}
+
+	const auto GetParentPoint = [&](const int32 BranchIndex) -> TPair<int32, int32>
+		{
+			const int32 RootPointIndex = BranchFacade.GetRootPoint(BranchIndex);
+			if (RootPointIndex != INDEX_NONE)
+			{
+				const float PointLFR = PointFacade.GetLengthFromRoot(RootPointIndex);
+				const int32 ParentBranchIndex = BranchFacade.GetParentBranchIndex(BranchIndex);
+				if (ParentBranchIndex != INDEX_NONE)
+				{
+					const TArray<int32>& ParentBranchPoints = BranchFacade.GetPoints(ParentBranchIndex);
+					for (int32 BranchPointIndex = 0; BranchPointIndex < ParentBranchPoints.Num(); BranchPointIndex++)
+					{
+						if (PointFacade.GetLengthFromRoot(ParentBranchPoints[BranchPointIndex]) >= PointLFR)
+						{
+							return {ParentBranchIndex, BranchPointIndex};
+						}
+					}
+				}
+			}
+			return {INDEX_NONE, INDEX_NONE};
+		};
+
+	const auto GetChildPoints = [&](
+		const int32 BranchIndex,
+		const int32 BranchPointIndex,
+		const TArray<int32>& ImmediateChildBranches
+	) -> TArray<TPair<int32, int32>>
+		{
+			const TArray<int32>& BranchPoints = BranchFacade.GetPoints(BranchIndex);
+
+			const float PointLFR = PointFacade.GetLengthFromRoot(BranchPoints[BranchPointIndex]);
+			const float NextPointLFR = BranchPoints.IsValidIndex(BranchPointIndex + 1)
+				? PointFacade.GetLengthFromRoot(BranchPoints[BranchPointIndex + 1])
+				: UE_MAX_FLT;
+
+			TArray<TPair<int32, int32>> ChildPoints;
+			for (int32 ChildBranch : ImmediateChildBranches)
+			{
+				const int32 ChildBranchIndex = BranchFacade.GetBranchNumbers().Find(ChildBranch);
+				if (ChildBranchIndex != INDEX_NONE)
+				{
+					const int32 ChildBranchRootPointIndex = BranchFacade.GetRootPoint(ChildBranchIndex);
+					if (ChildBranchRootPointIndex != INDEX_NONE)
+					{
+						const float ChildLFR = PointFacade.GetLengthFromRoot(ChildBranchRootPointIndex);
+						if (ChildLFR >= PointLFR && ChildLFR < NextPointLFR)
+						{
+							ChildPoints.Emplace(ChildBranchIndex, 0);
+						}
+					}
+				}
+			}
+
+			return ChildPoints;
+		};
+
+	SelectedPoints.Empty();
+	PrimarySelectedPoint.Weight = 0.0f;
+	SelectedPoints.Add(PrimarySelectedPoint);
+
+	int32 CurrentNeighbour = 0;
+	bool bPointAdded = true;
+	while (bPointAdded)
+	{
+		bPointAdded = false;
+		for (int32 SelectedPointIndex = 0; SelectedPointIndex < SelectedPoints.Num(); SelectedPointIndex++)
+		{
+			const FSelectedPoint& SelectedPoint = SelectedPoints[SelectedPointIndex];
+			if (SelectedPoint.NeighbourIndex != CurrentNeighbour)
+			{
+				continue;
+			}
+
+			if (SelectedPoint.BranchPointIndex == 0)
+			{
+				const auto [ParentBranchIndex, ParentBranchPointIndex] = GetParentPoint(SelectedPoint.BranchIndex);
+				if (ParentBranchIndex != INDEX_NONE && ParentBranchPointIndex != INDEX_NONE)
+				{
+					SelectedPoints.AddUnique({ParentBranchIndex, ParentBranchPointIndex, SelectedPoint.Weight, CurrentNeighbour});
+				}
+			}
+			const TArray<TPair<int32, int32>> ChildPoints = GetChildPoints(
+				SelectedPoint.BranchIndex,
+				SelectedPoint.BranchPointIndex,
+				BranchFacade.GetImmediateChildren(SelectedPoint.BranchIndex)
+			);
+			for (const auto [ChildBranchIndex, ChildBranchPointIndex] : ChildPoints)
+			{
+				SelectedPoints.AddUnique({ChildBranchIndex, ChildBranchPointIndex, SelectedPoint.Weight, CurrentNeighbour});
+			}
+
+			const TArray<int32>& BranchPoints = BranchFacade.GetPoints(SelectedPoint.BranchIndex);
+			const int32 SelectedIndex = BranchPoints[SelectedPoint.BranchPointIndex];
+			const int32 SelectedBudNumber = PointFacade.GetBudNumber(SelectedIndex);
+			const FVector3f SelectedPointPosition = PointFacade.GetPosition(SelectedIndex) + Offsets[SelectedBudNumber];
+
+			if (BranchPoints.IsValidIndex(SelectedPoint.BranchPointIndex - 1))
+			{
+				const int32 PointIndex = BranchPoints[SelectedPoint.BranchPointIndex - 1];
+				const int32 PointBudNumber = PointFacade.GetBudNumber(PointIndex);
+				const FVector3f PointPosition = PointFacade.GetPosition(PointIndex) + Offsets[PointBudNumber];
+				const float PointDistance = SelectedPoint.Weight + FVector3f::Distance(SelectedPointPosition, PointPosition);
+				if (PointDistance < MaxDistance)
+				{
+					SelectedPoints.AddUnique({
+						SelectedPoint.BranchIndex, SelectedPoint.BranchPointIndex - 1, PointDistance, CurrentNeighbour + 1
+					});
+					bPointAdded = true;
+				}
+			}
+			if (BranchPoints.IsValidIndex(SelectedPoint.BranchPointIndex + 1))
+			{
+				const int32 PointIndex = BranchPoints[SelectedPoint.BranchPointIndex + 1];
+				const int32 PointBudNumber = PointFacade.GetBudNumber(PointIndex);
+				const FVector3f PointPosition = PointFacade.GetPosition(PointIndex) + Offsets[PointBudNumber];
+				const float PointDistance = SelectedPoint.Weight + FVector3f::Distance(SelectedPointPosition, PointPosition);
+				if (PointDistance < MaxDistance)
+				{
+					SelectedPoints.AddUnique({
+						SelectedPoint.BranchIndex, SelectedPoint.BranchPointIndex + 1, PointDistance, CurrentNeighbour + 1
+					});
+					bPointAdded = true;
+				}
+			}
+		}
+
+		++CurrentNeighbour;
+	}
+
+	const float MaxNeighbours = Algo::MaxElementBy(
+		SelectedPoints, [](const FSelectedPoint& Point)
+			{
+				return Point.NeighbourIndex;
+			})->NeighbourIndex + 1;
+	for (FSelectedPoint& SelectedPoint : SelectedPoints)
+	{
+		SelectedPoint.Weight = GetWeightBySelectedSmoothness(
+			(MaxNeighbours - SelectedPoint.NeighbourIndex) / MaxNeighbours,
+			SmoothnessMethod
+		);
+	}
+}
+
+void FSkeletonPointsSelection::ClearSelection()
+{
+	SelectedPoints.Empty();
+	PrimarySelectedPoint = FSelectedPoint();
+}
+
+bool FSkeletonPointsSelection::HasSelection()
+{
+	return PrimarySelectedPoint.BranchIndex != INDEX_NONE && PrimarySelectedPoint.BranchPointIndex != INDEX_NONE;
+}
+
+float FSkeletonPointsSelection::GetWeightBySelectedSmoothness(const float Weight, const EPointSelectionSmoothnessMethod SmoothnessMethod)
+{
+	switch (SmoothnessMethod)
+	{
+	case EPointSelectionSmoothnessMethod::Linear:
+		return Weight;
+
+	case EPointSelectionSmoothnessMethod::Smooth:
+		return Weight * Weight * (3.0f - 2.0f * Weight);
+
+	case EPointSelectionSmoothnessMethod::Sphere:
+		return FMath::Sqrt(1 - FMath::Square(Weight - 1));
+
+	case EPointSelectionSmoothnessMethod::Root:
+		return FMath::Sqrt(Weight);
+
+	case EPointSelectionSmoothnessMethod::Sharp:
+		return FMath::Square(Weight);
+
+	case EPointSelectionSmoothnessMethod::Sine:
+		return FMath::Sin(Weight * UE_HALF_PI);
+
+	case EPointSelectionSmoothnessMethod::Constant:
+		return 1.0f;
+
+	default:
+		return Weight;
+	}
+};
+
+
+UPVSkeletonEditTool::UPVSkeletonEditTool()
+{
+	SetFlags(RF_Transactional);
+}
+
+void UPVSkeletonEditTool::Setup()
+{
+	if (!HasCollection())
+	{
+		GetToolManager()->PostActiveToolShutdownRequest(this, EToolShutdownType::Cancel);
+		return;
+	}
+
+	Super::Setup();
+
+	Selection = MakeUnique<FSkeletonPointsSelection>(&GetCollection());
+
+	const PV::Facades::FPointFacade PointFacade(GetCollection());
+
+	const TObjectPtr<UPVSkeletonEditSettings> SkeletonEditSettings = GetNodeSettings<UPVSkeletonEditSettings>();
+	TArray<FVector3f>& Offsets = SkeletonEditSettings->SkeletonEditSettings.RelativeOffsets;
+
+	const int32* MaxBudNumber = Algo::MaxElement(PointFacade.GetBudNumbersAttribute());
+	if (Offsets.Num() < (*MaxBudNumber + 1))
+	{
+		Offsets.SetNumZeroed(*MaxBudNumber + 1);
+	}
+
+	TArray<bool>& RemovedPoints = SkeletonEditSettings->SkeletonEditSettings.RemovedPoints;
+	if (RemovedPoints.Num() < (*MaxBudNumber + 1))
+	{
+		RemovedPoints.SetNumZeroed(*MaxBudNumber + 1);
+	}
+
+	USingleClickInputBehavior* Behavior = NewObject<USingleClickInputBehavior>(this);
+	Behavior->Initialize(this);
+	AddInputBehavior(Behavior);
+
+	TransformProxy = NewObject<UTransformProxy>(this);
+
+	CombinedGizmo = UE::TransformGizmoUtil::CreateCustomTransformGizmo(GetToolManager(), ETransformGizmoSubElements::StandardTranslateRotate, this);
+	CombinedGizmo->SetActiveTarget(TransformProxy, GetToolManager());
+	CombinedGizmo->SetVisibility(false);
+	CombinedGizmo->ActiveGizmoMode = EToolContextTransformGizmoMode::Combined;
+	CombinedGizmo->bUseContextGizmoMode = false;
+	CombinedGizmo->bUseContextCoordinateSystem = false;
+
+	TransformProxy->OnBeginTransformEdit.AddUObject(this, &UPVSkeletonEditTool::OnGizmoTransformBegin);
+	TransformProxy->OnTransformChanged.AddUObject(this, &UPVSkeletonEditTool::OnGizmoTransformChanged);
+	TransformProxy->OnEndTransformEdit.AddUObject(this, &UPVSkeletonEditTool::OnGizmoTransformEnd);
+}
+
+void UPVSkeletonEditTool::Shutdown(EToolShutdownType ShutdownType)
+{
+	if (CombinedGizmo)
+	{
+		CombinedGizmo->ClearActiveTarget();
+		CombinedGizmo->Shutdown();
+		CombinedGizmo = nullptr;
+	}
+
+	TransformProxy = nullptr;
+
+	Super::Shutdown(ShutdownType);
+}
+
+void UPVSkeletonEditTool::OnNodeSettingsUpdated()
+{
+	if (Selection->HasSelection())
+	{
+		const TObjectPtr<UPVSkeletonEditSettings> SkeletonEditSettings = GetNodeSettings<UPVSkeletonEditSettings>();
+		const FPVSkeletonSelectionParams& SelectionParams = SkeletonEditSettings->SkeletonEditSettings.SkeletonSelectionParams;
+
+		if (SelectionParams.SkeletonSelectionMode == ESkeletonSelectionMode::SelectByDistance)
+		{
+			Selection->ExtendSelectionByDistance(
+				SelectionParams.SelectionDistance,
+				SkeletonEditSettings->SkeletonEditSettings.RelativeOffsets,
+				SelectionParams.SelectionFalloff
+			);
+		}
+		else
+		{
+			Selection->ExtendSelection(
+				SelectionParams.NumPointsSelection,
+				SelectionParams.SelectionFalloff
+			);
+		}
+
+		const PV::Facades::FPointFacade PointFacade(GetCollection());
+		const PV::Facades::FBranchFacade BranchFacade(GetCollection());
+
+		TArray<FVector3f>& Offsets = SkeletonEditSettings->SkeletonEditSettings.RelativeOffsets;
+
+		const int32 SelectedPointIndex = BranchFacade.GetPoints(Selection->GetSelectedBranchIndex())[Selection->GetSelectedBranchPointIndex()];
+
+		const int32 SelectedPointBudBumber = PointFacade.GetBudNumber(SelectedPointIndex);
+		const FVector Position = static_cast<FVector>(PointFacade.GetPosition(SelectedPointIndex) + Offsets[SelectedPointBudBumber]);
+		CombinedGizmo->SetVisibility(true);
+		CombinedGizmo->ReinitializeGizmoTransform(FTransform(Position));
+	}
+	else
+	{
+		CombinedGizmo->SetVisibility(false);
+	}
+}
+
+FInputRayHit UPVSkeletonEditTool::IsHitByClick(const FInputDeviceRay& ClickPos)
+{
+	if (FViewport* FocusedViewport = GetToolManager()->GetContextQueriesAPI()->GetFocusedViewport())
+	{
+		if (FocusedViewport->GetHitProxy(ClickPos.ScreenPosition.X, ClickPos.ScreenPosition.Y))
+		{
+			return FInputRayHit(1.0f);
+		}
+	}
+
+	Selection->ClearSelection();
+	CombinedGizmo->SetVisibility(false);
+	return FInputRayHit();
+}
+
+void UPVSkeletonEditTool::OnClicked(const FInputDeviceRay& ClickPos)
+{
+	FViewport* FocusedViewport = GetToolManager()->GetContextQueriesAPI()->GetFocusedViewport();
+	if (!FocusedViewport)
+	{
+		return;
+	}
+	const HHitProxy* HitProxy = FocusedViewport->GetHitProxy(ClickPos.ScreenPosition.X, ClickPos.ScreenPosition.Y);
+	if (!HitProxy)
+	{
+		return;
+	}
+
+	if (HitProxy->IsA(HPVSkeletonPointProxy::StaticGetType()))
+	{
+		const PV::Facades::FPointFacade PointFacade(GetCollection());
+		const PV::Facades::FBranchFacade BranchFacade(GetCollection());
+
+		const auto GetParentPoint = [&](const int32 BranchIndex) -> TPair<int32, int32>
+			{
+				const int32 RootPointIndex = BranchFacade.GetRootPoint(BranchIndex);
+				if (RootPointIndex != INDEX_NONE)
+				{
+					const float PointLFR = PointFacade.GetLengthFromRoot(RootPointIndex);
+					const int32 ParentBranchIndex = BranchFacade.GetParentBranchIndex(BranchIndex);
+					if (ParentBranchIndex != INDEX_NONE)
+					{
+						const TArray<int32>& ParentBranchPoints = BranchFacade.GetPoints(ParentBranchIndex);
+						for (int32 BranchPointIndex = 0; BranchPointIndex < ParentBranchPoints.Num(); BranchPointIndex++)
+						{
+							if (PointFacade.GetLengthFromRoot(ParentBranchPoints[BranchPointIndex]) >= PointLFR)
+							{
+								return {ParentBranchIndex, BranchPointIndex};
+							}
+						}
+					}
+				}
+				return {INDEX_NONE, INDEX_NONE};
+			};
+
+		const HPVSkeletonPointProxy* SkeletonPointProxy = static_cast<const HPVSkeletonPointProxy*>(HitProxy);
+
+		const TObjectPtr<UPVSkeletonEditSettings> SkeletonEditSettings = GetNodeSettings<UPVSkeletonEditSettings>();
+
+		if (SkeletonEditSettings->SkeletonEditSettings.SkeletonEditMode == ESkeletonEditMode::TRS)
+		{
+			const FPVSkeletonSelectionParams& SelectionParams = SkeletonEditSettings->SkeletonEditSettings.SkeletonSelectionParams;
+
+			if (!SelectionParams.bRootSelectsBranch && SkeletonPointProxy->BranchPointIndex == 0)
+			{
+				const auto [ParentBranchIndex, ParentBranchPointIndex] = GetParentPoint(SkeletonPointProxy->BranchIndex);
+				if (ParentBranchIndex == INDEX_NONE || ParentBranchPointIndex == INDEX_NONE)
+				{
+					Selection->SelectPoint(SkeletonPointProxy->BranchIndex, SkeletonPointProxy->BranchPointIndex);
+				}
+				else
+				{
+					Selection->SelectPoint(ParentBranchIndex, ParentBranchPointIndex);
+				}
+			}
+			else
+			{
+				Selection->SelectPoint(SkeletonPointProxy->BranchIndex, SkeletonPointProxy->BranchPointIndex);
+			}
+
+			if (SelectionParams.SkeletonSelectionMode == ESkeletonSelectionMode::SelectByDistance)
+			{
+				Selection->ExtendSelectionByDistance(
+					SelectionParams.SelectionDistance,
+					SkeletonEditSettings->SkeletonEditSettings.RelativeOffsets,
+					SelectionParams.SelectionFalloff
+				);
+			}
+			else
+			{
+				Selection->ExtendSelection(
+					SelectionParams.NumPointsSelection,
+					SelectionParams.SelectionFalloff
+				);
+			}
+
+			TArray<FVector3f>& Offsets = SkeletonEditSettings->SkeletonEditSettings.RelativeOffsets;
+
+			const int32 SelectedPointIndex = BranchFacade.GetPoints(SkeletonPointProxy->BranchIndex)[SkeletonPointProxy->BranchPointIndex];
+			const int32 SelectedPointBudNumber = PointFacade.GetBudNumber(SelectedPointIndex);
+			const FVector Position = static_cast<FVector>(PointFacade.GetPosition(SelectedPointIndex) + Offsets[SelectedPointBudNumber]);
+			PreviousRotation = FQuat4f::Identity;
+			CombinedGizmo->SetVisibility(true);
+			CombinedGizmo->ReinitializeGizmoTransform(FTransform(Position));
+		}
+		else if (SkeletonEditSettings->SkeletonEditSettings.SkeletonEditMode == ESkeletonEditMode::BranchRemoval)
+		{
+			PreviousRotation = FQuat4f::Identity;
+			Selection->ClearSelection();
+			CombinedGizmo->SetVisibility(false);
+
+			TArray<bool>& RemovedPoints = SkeletonEditSettings->SkeletonEditSettings.RemovedPoints;
+
+			const TArray<int32>& BranchPoints = BranchFacade.GetPoints(SkeletonPointProxy->BranchIndex);
+
+			int32 StartIndex =
+				SkeletonEditSettings->SkeletonEditSettings.BranchRemovalParams.BranchRemovalMode == EBranchRemovalMode::TrimBranch
+				? SkeletonPointProxy->BranchPointIndex
+				: 0;
+			StartIndex = StartIndex == 0
+				? 0
+				: StartIndex + 1;
+			for (int32 BranchPointIndex = StartIndex; BranchPointIndex < BranchPoints.Num(); BranchPointIndex++)
+			{
+				const int32 PointIndex = BranchPoints[BranchPointIndex];
+				const int32 PointBudNumber = PointFacade.GetBudNumber(PointIndex);
+				RemovedPoints[PointBudNumber] = true;
+			}
+
+			const int32 RemovedPointIndex = BranchPoints[StartIndex];
+			const float PointLFR = PointFacade.GetLengthFromRoot(RemovedPointIndex);
+			const TManagedArray<int32>& BranchNumbers = BranchFacade.GetBranchNumbers();
+			for (int32 BranchChildNumber : BranchFacade.GetImmediateChildren(SkeletonPointProxy->BranchIndex))
+			{
+				if (const int32 BranchIndex = BranchNumbers.Find(BranchChildNumber); BranchIndex != INDEX_NONE)
+				{
+					const TArray<int32>& ChildBranchPoints = BranchFacade.GetPoints(BranchIndex);
+					if (PointFacade.GetLengthFromRoot(ChildBranchPoints[0]) >= PointLFR)
+					{
+						for (const int32 ChildBranchPointIndex : ChildBranchPoints)
+						{
+							const int32 ChildPointBudNumber = PointFacade.GetBudNumber(ChildBranchPointIndex);
+							RemovedPoints[ChildPointBudNumber] = true;
+						}
+
+						for (int32 SubBranchChildNumber : BranchFacade.GetChildren(BranchIndex))
+						{
+							if (const int32 SubBranchIndex = BranchNumbers.Find(SubBranchChildNumber); SubBranchIndex != INDEX_NONE)
+							{
+								for (const int32 ChildSubBranchPointIndex : BranchFacade.GetPoints(SubBranchIndex))
+								{
+									const int32 SubBranchChildPointBudNumber = PointFacade.GetBudNumber(ChildSubBranchPointIndex);
+									RemovedPoints[SubBranchChildPointBudNumber] = true;
+								}
+							}
+						}
+					}
+				}
+			}
+
+			GetNodeSettings()->PostEditChange();
+		}
+	}
+	else
+	{
+		PreviousRotation = FQuat4f::Identity;
+		Selection->ClearSelection();
+		CombinedGizmo->SetVisibility(false);
+	}
+}
+
+void UPVSkeletonEditTool::Render(IToolsContextRenderAPI* RenderAPI)
+{
+	const PV::Facades::FPointFacade PointFacade(GetCollection());
+	const PV::Facades::FBranchFacade BranchFacade(GetCollection());
+	if (!PointFacade.IsValid() || !BranchFacade.IsValid())
+	{
+		return;
+	}
+
+	FPrimitiveDrawInterface* PDI = RenderAPI->GetPrimitiveDrawInterface();
+	const TManagedArray<FVector3f>& Positions = PointFacade.GetPositions();
+
+	const TObjectPtr<UPVSkeletonEditSettings> SkeletonEditSettings = GetNodeSettings<UPVSkeletonEditSettings>();
+	TArray<FVector3f>& Offsets = SkeletonEditSettings->SkeletonEditSettings.RelativeOffsets;
+	TArray<bool>& RemovedPoints = SkeletonEditSettings->SkeletonEditSettings.RemovedPoints;
+
+	int32 PreviousSelectedPointIndex = INDEX_NONE;
+	for (int32 BranchIndex = 0; BranchIndex < BranchFacade.GetElementCount(); BranchIndex++)
+	{
+		const TArray<int32>& BranchPoints = BranchFacade.GetPoints(BranchIndex);
+		for (int32 BranchPointIndex = 0; BranchPointIndex < BranchPoints.Num(); BranchPointIndex++)
+		{
+			const int32 PointIndex = BranchPoints[BranchPointIndex];
+			const int32 PointBudNumber = PointFacade.GetBudNumber(PointIndex);
+			if (RemovedPoints[PointBudNumber])
+			{
+				break;
+			}
+			const int32 BudNumber = PointFacade.GetBudNumber(PointIndex);
+			const FVector CurrentPosition = static_cast<FVector>(Positions[PointIndex] + Offsets[BudNumber]);
+			FLinearColor DrawColor = FLinearColor::White;
+
+			const int32 SelectedPointIndex = Selection->SelectedPoints.Find({BranchIndex, BranchPointIndex});
+			if (SelectedPointIndex != INDEX_NONE)
+			{
+				DrawColor = PV::Utilities::GetRandomHueColor(1.0f - Selection->SelectedPoints[SelectedPointIndex].Weight);
+			}
+
+			PDI->SetHitProxy(new HPVSkeletonPointProxy(BranchIndex, BranchPointIndex));
+			PDI->DrawPoint(CurrentPosition, DrawColor, 5.0f, SDPG_World);
+			PDI->SetHitProxy(nullptr);
+
+			if (BranchPointIndex != 0)
+			{
+				const int32 PreviousPointIndex = BranchPoints[BranchPointIndex - 1];
+				const int32 PreviousBudNumber = PointFacade.GetBudNumber(PreviousPointIndex);
+				const FVector PreviousPosition = static_cast<FVector>(Positions[PreviousPointIndex] + Offsets[PreviousBudNumber]);
+
+				if (PreviousSelectedPointIndex != INDEX_NONE && SelectedPointIndex != INDEX_NONE)
+				{
+					if (Selection->SelectedPoints[SelectedPointIndex].Weight < Selection->SelectedPoints[PreviousSelectedPointIndex].Weight)
+					{
+						DrawColor = PV::Utilities::GetRandomHueColor(1.0f - Selection->SelectedPoints[PreviousSelectedPointIndex].Weight);
+					}
+				}
+				else if (PreviousSelectedPointIndex == INDEX_NONE)
+				{
+					DrawColor = FLinearColor::White;
+				}
+
+				PDI->DrawLine(
+					PreviousPosition,
+					CurrentPosition,
+					DrawColor,
+					SDPG_World,
+					2.0f,
+					0.0f,
+					true
+				);
+			}
+
+			PreviousSelectedPointIndex = SelectedPointIndex;
+		}
+	}
+}
+
+void UPVSkeletonEditTool::OnGizmoTransformBegin(UTransformProxy* InTransformProxy)
+{
+	GetToolManager()->BeginUndoTransaction(NSLOCTEXT("PVSkeletonEditTool", "MoveTransform", "Move Transform"));
+}
+
+void UPVSkeletonEditTool::OnGizmoTransformChanged(UTransformProxy* InTransformProxy, FTransform InTransform)
+{
+	if (Selection->HasSelection())
+	{
+		const PV::Facades::FPointFacade PointFacade(GetCollection());
+		const PV::Facades::FBranchFacade BranchFacade(GetCollection());
+
+		const TObjectPtr<UPVSkeletonEditSettings> SkeletonEditSettings = GetNodeSettings<UPVSkeletonEditSettings>();
+		TArray<FVector3f>& Offsets = SkeletonEditSettings->SkeletonEditSettings.RelativeOffsets;
+
+		const int32 PrimarySelectedBranch = Selection->GetSelectedBranchIndex();
+		const int32 PrimarySelectedBranchPointIndex = Selection->GetSelectedBranchPointIndex();
+
+		const TArray<int32>& PrimaryBranchPoints = BranchFacade.GetPoints(PrimarySelectedBranch);
+		const int32 PrimarySelectedPointIndex = PrimaryBranchPoints[PrimarySelectedBranchPointIndex];
+		const int32 PrimaryPointBudNumber = PointFacade.GetBudNumber(PrimarySelectedPointIndex);
+
+		const FVector3f OriginalPosition = PointFacade.GetPosition(PrimarySelectedPointIndex) + Offsets[PrimaryPointBudNumber];
+		const FVector3f NewPosition = static_cast<FVector3f>(InTransform.GetLocation());
+		const FVector3f Offset = NewPosition - OriginalPosition;
+
+		if (!Offset.IsZero())
+		{
+			for (const FSkeletonPointsSelection::FSelectedPoint& SelectedPoint : Selection->SelectedPoints)
+			{
+				const TArray<int32>& BranchPoints = BranchFacade.GetPoints(SelectedPoint.BranchIndex);
+				const int32 SelectedPointIndex = BranchPoints[SelectedPoint.BranchPointIndex];
+				const int32 SelectedPointBudNumber = PointFacade.GetBudNumber(SelectedPointIndex);
+				Offsets[SelectedPointBudNumber] += Offset * SelectedPoint.Weight;
+			}
+		}
+
+		const FQuat4f CurrentRotation = static_cast<FQuat4f>(InTransform.GetRotation());
+		const FQuat4f DeltaRotation = CurrentRotation * PreviousRotation.Inverse();
+		if (!DeltaRotation.IsIdentity())
+		{
+			for (int32 BranchPointIndex = PrimarySelectedBranchPointIndex + 1; BranchPointIndex < PrimaryBranchPoints.Num(); BranchPointIndex++)
+			{
+				const int32 PointIndex = PrimaryBranchPoints[BranchPointIndex];
+				const int32 PointBudNumber = PointFacade.GetBudNumber(PointIndex);
+				const FVector3f PointPosition = PointFacade.GetPosition(PointIndex) + Offsets[PointBudNumber];
+				const FVector3f UpdatedPosition = OriginalPosition + DeltaRotation.RotateVector(PointPosition - OriginalPosition);
+				const FVector3f PointOffset = UpdatedPosition - PointPosition;
+				Offsets[PointBudNumber] += PointOffset;
+			}
+
+			const float PointLFR = PointFacade.GetLengthFromRoot(PrimarySelectedPointIndex);
+
+			const TManagedArray<int32>& BranchNumbers = BranchFacade.GetBranchNumbers();
+			for (int32 BranchChildNumber : BranchFacade.GetImmediateChildren(PrimarySelectedBranch))
+			{
+				if (const int32 BranchIndex = BranchNumbers.Find(BranchChildNumber); BranchIndex != INDEX_NONE)
+				{
+					const TArray<int32>& ChildBranchPoints = BranchFacade.GetPoints(BranchIndex);
+					if (PointFacade.GetLengthFromRoot(ChildBranchPoints[0]) > PointLFR)
+					{
+						for (int32 ChildBranchPointIndex : ChildBranchPoints)
+						{
+							const int32 ChildPointBudNumber = PointFacade.GetBudNumber(ChildBranchPointIndex);
+							const FVector3f PointPosition = PointFacade.GetPosition(ChildBranchPointIndex) + Offsets[ChildPointBudNumber];
+							const FVector3f UpdatedPosition = OriginalPosition + DeltaRotation.RotateVector(PointPosition - OriginalPosition);
+							const FVector3f PointOffset = UpdatedPosition - PointPosition;
+							Offsets[ChildPointBudNumber] += PointOffset;
+						}
+
+						for (int32 SubBranchChildNumber : BranchFacade.GetChildren(BranchIndex))
+						{
+							if (const int32 SubBranchIndex = BranchNumbers.Find(SubBranchChildNumber); SubBranchIndex != INDEX_NONE)
+							{
+								for (int32 ChildSubBranchPointIndex : BranchFacade.GetPoints(SubBranchIndex))
+								{
+									const int32 SubPointBudNumber = PointFacade.GetBudNumber(ChildSubBranchPointIndex);
+									const FVector3f PointPosition = PointFacade.GetPosition(ChildSubBranchPointIndex) + Offsets[SubPointBudNumber];
+									const FVector3f UpdatedPosition = OriginalPosition + DeltaRotation.RotateVector(PointPosition - OriginalPosition);
+									const FVector3f PointOffset = UpdatedPosition - PointPosition;
+									Offsets[SubPointBudNumber] += PointOffset;
+								}
+							}
+						}
+					}
+				}
+			}
+		}
+
+		PreviousRotation = CurrentRotation;
+	}
+}
+
+void UPVSkeletonEditTool::OnGizmoTransformEnd(UTransformProxy* InTransformProxy)
+{
+	GetNodeSettings()->PostEditChange();
+	GetToolManager()->EndUndoTransaction();
+}
diff --git a/Engine/Plugins/Experimental/ProceduralVegetationEditor/Source/ProceduralVegetationEditor/Private/Tools/NodeTools/PVSkeletonEditTool.h b/Engine/Plugins/Experimental/ProceduralVegetationEditor/Source/ProceduralVegetationEditor/Private/Tools/NodeTools/PVSkeletonEditTool.h
new file mode 100644
--- /dev/null
+++ b/Engine/Plugins/Experimental/ProceduralVegetationEditor/Source/ProceduralVegetationEditor/Private/Tools/NodeTools/PVSkeletonEditTool.h
@@ -0,0 +1,118 @@
+// Copyright Epic Games, Inc. All Rights Reserved.
+
+#pragma once
+
+#include "ComponentVisualizer.h"
+
+#include "Tools/PVBaseInteractiveTool.h"
+
+#include "BaseBehaviors/BehaviorTargetInterfaces.h"
+
+#include "Nodes/PVSkeletonEditSettings.h"
+
+#include "PVSkeletonEditTool.generated.h"
+
+class UTransformProxy;
+class UCombinedTransformGizmo;
+
+struct HPVSkeletonPointProxy : HHitProxy
+{
+	DECLARE_HIT_PROXY();
+
+	HPVSkeletonPointProxy(const int32 InBranchIndex, const int32 InBranchPointIndex)
+		: HHitProxy(HPP_Wireframe)
+		, BranchIndex(InBranchIndex)
+		, BranchPointIndex(InBranchPointIndex)
+	{}
+
+	virtual EMouseCursor::Type GetMouseCursor() override
+	{
+		return EMouseCursor::Crosshairs;
+	}
+
+	int32 BranchIndex = INDEX_NONE;
+	int32 BranchPointIndex = INDEX_NONE;
+};
+
+
+
+struct FSkeletonPointsSelection
+{
+	FSkeletonPointsSelection(const FManagedArrayCollection* const InCollection);
+
+	void SelectPoint(int32 InBranchIndex, int32 InBranchPointIndex);
+	void ExtendSelection(const int32 NumNeighbours, const EPointSelectionSmoothnessMethod SmoothnessMethod = EPointSelectionSmoothnessMethod::Linear);
+	void ExtendSelectionByDistance(const float MaxDistance, const TArray<FVector3f>& Offsets, const EPointSelectionSmoothnessMethod SmoothnessMethod = EPointSelectionSmoothnessMethod::Linear);
+	void ClearSelection();
+
+	bool HasSelection();
+	
+	float GetWeightBySelectedSmoothness(float Weight, const EPointSelectionSmoothnessMethod SmoothnessMethod = EPointSelectionSmoothnessMethod::Linear);
+
+	int32 GetSelectedBranchIndex()
+	{
+		return PrimarySelectedPoint.BranchIndex;
+	}
+
+	int32 GetSelectedBranchPointIndex()
+	{
+		return PrimarySelectedPoint.BranchPointIndex;
+	}
+
+	struct FSelectedPoint
+	{
+		int32 BranchIndex = INDEX_NONE;
+		int32 BranchPointIndex = INDEX_NONE;
+		float Weight = 0.0f;
+
+		int32 NeighbourIndex = INDEX_NONE;
+
+		bool operator==(const FSelectedPoint& Other) const
+		{
+			return BranchIndex == Other.BranchIndex && BranchPointIndex == Other.BranchPointIndex;
+		}
+	};
+
+	const FManagedArrayCollection* const Collection = nullptr;
+
+	FSelectedPoint PrimarySelectedPoint;
+	TArray<FSelectedPoint> SelectedPoints;
+};
+
+UCLASS()
+class UPVSkeletonEditTool
+	: public UPVBaseInteractiveTool,
+	  public IClickBehaviorTarget
+{
+	GENERATED_BODY()
+
+public:
+	UPVSkeletonEditTool();
+
+	virtual void Setup() override;
+	virtual void Shutdown(EToolShutdownType ShutdownType) override;
+	
+	virtual void OnNodeSettingsUpdated() override;
+
+	virtual FInputRayHit IsHitByClick(const FInputDeviceRay& ClickPos) override;
+	virtual void OnClicked(const FInputDeviceRay& ClickPos) override;
+
+	virtual void Render(IToolsContextRenderAPI* RenderAPI) override;
+	
+private:
+	void OnGizmoTransformBegin(UTransformProxy* TransformProxy);
+	void OnGizmoTransformChanged(UTransformProxy* TransformProxy, FTransform Transform);
+	void OnGizmoTransformEnd(UTransformProxy* TransformProxy);
+
+protected:
+	UPROPERTY(Transient)
+	TObjectPtr<UTransformProxy> TransformProxy = nullptr;
+
+	UPROPERTY(Transient)
+	TObjectPtr<UCombinedTransformGizmo> CombinedGizmo = nullptr;
+
+private:
+	TUniquePtr<FSkeletonPointsSelection> Selection;
+	
+	FQuat4f PreviousRotation = FQuat4f::Identity;
+};
diff --git a/Engine/Plugins/Experimental/ProceduralVegetationEditor/Source/ProceduralVegetationEditor/Private/Tools/PVBaseInteractiveTool.cpp b/Engine/Plugins/Experimental/ProceduralVegetationEditor/Source/ProceduralVegetationEditor/Private/Tools/PVBaseInteractiveTool.cpp
new file mode 100644
--- /dev/null
+++ b/Engine/Plugins/Experimental/ProceduralVegetationEditor/Source/ProceduralVegetationEditor/Private/Tools/PVBaseInteractiveTool.cpp
@@ -0,0 +1,113 @@
+// Copyright Epic Games, Inc. All Rights Reserved.
+
+#include "PVBaseInteractiveTool.h"
+
+#include "ContextObjectStore.h"
+#include "InteractiveToolManager.h"
+
+#include "Contexts/PVToolContext.h"
+
+#include "Engine/World.h"
+
+#include "Tools/EdModeInteractiveToolsContext.h"
+
+class UContextObjectStore;
+
+void UPVBaseInteractiveTool::Setup()
+{
+	Super::Setup();
+
+	InteractiveToolsContext = Cast<UEditorInteractiveToolsContext>(GetToolManager()->GetOuter());
+
+	if (ensureMsgf(TargetWorld, TEXT("Target World should be set before calling tool setup")))
+	{
+		FActorSpawnParameters Params;
+		Params.ObjectFlags = RF_Transient;
+		Params.bHideFromSceneOutliner = false;
+		PreviewActor = TargetWorld->SpawnActor<AActor>(Params);
+
+		USceneComponent* Root = NewObject<USceneComponent>(PreviewActor, TEXT("PVToolPreviewRoot"), RF_Transient);
+		Root->RegisterComponent();
+		PreviewActor->SetRootComponent(Root);
+	}
+
+	NodeSettings->OnSettingsChangedDelegate.AddUObject(this, &UPVBaseInteractiveTool::OnLinkedNodeSettingsChanged);
+}
+
+void UPVBaseInteractiveTool::Shutdown(EToolShutdownType ShutdownType)
+{
+	NodeSettings->OnSettingsChangedDelegate.RemoveAll(this);
+
+	if (PreviewActor)
+	{
+		if (USceneComponent* const RootComponent = PreviewActor->GetRootComponent())
+		{
+			RootComponent->DestroyComponent();
+		}
+		PreviewActor->Destroy();
+		PreviewActor = nullptr;
+	}
+
+	Super::Shutdown(ShutdownType);
+}
+
+void UPVBaseInteractiveTool::SetTargetWorld(TObjectPtr<UWorld> InWorld)
+{
+	check(InWorld);
+	TargetWorld = InWorld;
+}
+
+TObjectPtr<UWorld> UPVBaseInteractiveTool::GetTargetWorld()
+{
+	return TargetWorld;
+}
+
+void UPVBaseInteractiveTool::SetCollection(const FManagedArrayCollection& InCollection)
+{
+	Collection = MakeUnique<FManagedArrayCollection>(InCollection);
+}
+
+bool UPVBaseInteractiveTool::HasCollection() const
+{
+	return Collection.IsValid();
+}
+
+const FManagedArrayCollection& UPVBaseInteractiveTool::GetCollection()
+{
+	check(HasCollection());
+	return *Collection.Get();
+}
+
+void UPVBaseInteractiveTool::OnLinkedNodeSettingsChanged(UPCGSettings* InPCGSettings, EPCGChangeType InChangeType)
+{
+	if (InPCGSettings == NodeSettings && InChangeType == EPCGChangeType::Settings)
+	{
+		OnNodeSettingsUpdated();
+	}
+}
+
+bool UPVToolBuilder::CanBuildTool(const FToolBuilderState& SceneState) const
+{
+	return ToolClass != nullptr;
+}
+
+UInteractiveTool* UPVToolBuilder::BuildTool(const FToolBuilderState& SceneState) const
+{
+	checkf(ToolClass, TEXT("Tool class should be set for PV Tool Builder"));
+	UPVBaseInteractiveTool* const Tool = NewObject<UPVBaseInteractiveTool>(SceneState.ToolManager, ToolClass);
+
+	Tool->SetTargetWorld(SceneState.World);
+	if (const TObjectPtr<UContextObjectStore> Store = SceneState.ToolManager->GetContextObjectStore())
+	{
+		if (const UPVToolContextObject* const ToolContext = Store->FindContext<UPVToolContextObject>())
+		{
+			Tool->NodeSettings = ToolContext->NodeSettings;
+			if (ToolContext->NodeInputData)
+			{
+				Tool->SetCollection(ToolContext->NodeInputData->GetCollection());
+			}
+		}
+	}
+
+	return Tool;
+}
diff --git a/Engine/Plugins/Experimental/ProceduralVegetationEditor/Source/ProceduralVegetationEditor/Private/Tools/PVBaseInteractiveTool.h b/Engine/Plugins/Experimental/ProceduralVegetationEditor/Source/ProceduralVegetationEditor/Private/Tools/PVBaseInteractiveTool.h
new file mode 100644
--- /dev/null
+++ b/Engine/Plugins/Experimental/ProceduralVegetationEditor/Source/ProceduralVegetationEditor/Private/Tools/PVBaseInteractiveTool.h
@@ -0,0 +1,76 @@
+// Copyright Epic Games, Inc. All Rights Reserved.
+
+#pragma once
+
+#include "CoreMinimal.h"
+#include "InteractiveTool.h"
+#include "InteractiveToolBuilder.h"
+#include "PCGSettings.h"
+
+#include "GeometryCollection/ManagedArrayCollection.h"
+
+#include "PVBaseInteractiveTool.generated.h"
+
+class UPVBaseSettings;
+class UEditorInteractiveToolsContext;
+
+UCLASS()
+class UPVBaseInteractiveTool : public UInteractiveTool
+{
+	GENERATED_BODY()
+
+public:
+	virtual void Setup() override;
+	virtual void Shutdown(EToolShutdownType ShutdownType) override;
+
+	virtual void OnNodeSettingsUpdated() {}
+	
+	void SetTargetWorld(TObjectPtr<UWorld> InWorld);
+	TObjectPtr<UWorld> GetTargetWorld();
+
+	void SetCollection(const FManagedArrayCollection& InCollection);
+	bool HasCollection() const;
+	const FManagedArrayCollection& GetCollection();
+
+	template <typename T=UPVBaseSettings>
+	TObjectPtr<T> GetNodeSettings() { return CastChecked<T>(NodeSettings); }
+	
+	template <typename T=UPVBaseSettings>
+	TObjectPtr<const T> GetNodeSettings() const { return CastChecked<T>(NodeSettings); }
+
+private:
+	void OnLinkedNodeSettingsChanged(UPCGSettings* InPCGSettings, EPCGChangeType InChangeType);
+
+protected:
+	UPROPERTY(Transient)
+	TObjectPtr<UWorld> TargetWorld = nullptr;
+
+	UPROPERTY(Transient)
+	TObjectPtr<AActor> PreviewActor = nullptr;
+
+	UPROPERTY(Transient)
+	TObjectPtr<UPVBaseSettings> NodeSettings = nullptr;
+
+	UPROPERTY(Transient)
+	TObjectPtr<UEditorInteractiveToolsContext> InteractiveToolsContext;
+
+private:
+	TUniquePtr<FManagedArrayCollection> Collection;
+
+	friend class UPVToolBuilder;
+};
+
+UCLASS()
+class UPVToolBuilder : public UInteractiveToolBuilder
+{
+	GENERATED_BODY()
+
+public:
+	virtual bool CanBuildTool(const FToolBuilderState& SceneState) const override;
+
+	virtual UInteractiveTool* BuildTool(const FToolBuilderState& SceneState) const override;
+
+public:
+	UPROPERTY()
+	TObjectPtr<UClass> ToolClass = nullptr;
+};
diff --git a/Engine/Plugins/Experimental/ProceduralVegetationEditor/Source/ProceduralVegetationEditor/Private/Tools/PVEditorMode.cpp b/Engine/Plugins/Experimental/ProceduralVegetationEditor/Source/ProceduralVegetationEditor/Private/Tools/PVEditorMode.cpp
new file mode 100644
--- /dev/null
+++ b/Engine/Plugins/Experimental/ProceduralVegetationEditor/Source/ProceduralVegetationEditor/Private/Tools/PVEditorMode.cpp
@@ -0,0 +1,62 @@
+﻿// Copyright Epic Games, Inc. All Rights Reserved.
+
+#include "PVEditorMode.h"
+
+#include "AssetEditorModeManager.h"
+#include "ContextObjectStore.h"
+
+#include "BaseGizmos/GizmoElementShared.h"
+#include "BaseGizmos/TransformGizmoUtil.h"
+
+#include "Contexts/PVToolContext.h"
+
+#include "Snapping/ModelingSceneSnappingManager.h"
+
+#include "Tools/EdModeInteractiveToolsContext.h"
+
+#define LOCTEXT_NAMESPACE "PVEditorMode"
+
+const FEditorModeID UPVEditorMode::EM_PVEditorModeId = TEXT("PVEditorModeId");
+
+UPVEditorMode::UPVEditorMode()
+{
+	Info = FEditorModeInfo(
+		EM_PVEditorModeId,
+		LOCTEXT("PVEditorModeName", "PV Editor Mode"),
+		FSlateIcon(),
+		false);
+}
+
+void UPVEditorMode::Enter()
+{
+	Super::Enter();
+
+	const TObjectPtr<UEditorInteractiveToolsContext> ToolsContext = GetInteractiveToolsContext();
+
+	ToolRegistry = MakeUnique<FPVToolsRegistry>(ToolsContext);
+	ToolRegistry->RegisterTools();
+
+	UE::TransformGizmoUtil::RegisterTransformGizmoContextObject(ToolsContext);
+
+	ToolsContext->SetForceCombinedGizmoMode(true);
+
+	const TObjectPtr<UContextObjectStore> ContextObjectStore = ToolsContext->ContextObjectStore;
+
+	const TObjectPtr<UPVToolContextObject> PVToolContextObject = NewObject<UPVToolContextObject>(ToolsContext);
+	ContextObjectStore->AddContextObject(PVToolContextObject);
+}
+
+void UPVEditorMode::Exit()
+{
+	const TObjectPtr<UEditorInteractiveToolsContext> ToolsContext = GetInteractiveToolsContext();
+
+	ToolRegistry->UnregisterTools();
+	ToolRegistry.Reset();
+
+	const TObjectPtr<UContextObjectStore> ContextObjectStore = ToolsContext->ContextObjectStore;
+	ContextObjectStore->Shutdown();
+
+	Super::Exit();
+}
+
+#undef LOCTEXT_NAMESPACE
diff --git a/Engine/Plugins/Experimental/ProceduralVegetationEditor/Source/ProceduralVegetationEditor/Private/Tools/PVEditorMode.h b/Engine/Plugins/Experimental/ProceduralVegetationEditor/Source/ProceduralVegetationEditor/Private/Tools/PVEditorMode.h
new file mode 100644
--- /dev/null
+++ b/Engine/Plugins/Experimental/ProceduralVegetationEditor/Source/ProceduralVegetationEditor/Private/Tools/PVEditorMode.h
@@ -0,0 +1,34 @@
+﻿// Copyright Epic Games, Inc. All Rights Reserved.
+
+#pragma once
+
+#include "CoreMinimal.h"
+#include "EditorViewportClient.h"
+#include "PVToolsRegistry.h"
+#include "SceneView.h"
+#include "Selection.h"
+
+#include "Tools/UEdMode.h"
+
+#include "PVEditorMode.generated.h"
+
+
+UCLASS()
+class UPVEditorMode : public UEdMode
+{
+	GENERATED_BODY()
+
+public:
+	const static FEditorModeID EM_PVEditorModeId;
+
+	UPVEditorMode();
+
+	virtual void Enter() override;
+	virtual void Exit() override;
+
+	FPVToolsRegistry& GetToolRegistry() { return *ToolRegistry; }
+	const FPVToolsRegistry& GetToolRegistry() const { return *ToolRegistry; }
+
+protected:
+	TUniquePtr<FPVToolsRegistry> ToolRegistry;
+};
diff --git a/Engine/Plugins/Experimental/ProceduralVegetationEditor/Source/ProceduralVegetationEditor/Private/Tools/PVToolsRegistry.cpp b/Engine/Plugins/Experimental/ProceduralVegetationEditor/Source/ProceduralVegetationEditor/Private/Tools/PVToolsRegistry.cpp
new file mode 100644
--- /dev/null
+++ b/Engine/Plugins/Experimental/ProceduralVegetationEditor/Source/ProceduralVegetationEditor/Private/Tools/PVToolsRegistry.cpp
@@ -0,0 +1,98 @@
+// Copyright Epic Games, Inc. All Rights Reserved.
+
+#include "PVToolsRegistry.h"
+
+#include "ContextObjectStore.h"
+#include "PVBaseInteractiveTool.h"
+
+#include "Contexts/PVToolContext.h"
+
+#include "Nodes/PVSkeletonEditSettings.h"
+
+#include "NodeTools/PVColliderTransformTool.h"
+#include "NodeTools/PVSkeletonEditTool.h"
+
+#include "Tools/EdModeInteractiveToolsContext.h"
+
+#define REGISTER_PVE_TOOL(Tool, Node)																								\
+	{																																\
+		UPVToolBuilder* const ToolBuilder = NewObject<UPVToolBuilder>(ToolsContext);												\
+		ToolBuilder->ToolClass = Tool::StaticClass();																				\
+		ToolManager->RegisterToolType(Tool::StaticClass()->GetName(), ToolBuilder);													\
+		NodeSettingsToToolMapping.Add(Node::StaticClass(), Tool::StaticClass()->GetName());											\
+	}
+
+FPVToolsRegistry::FPVToolsRegistry(const TObjectPtr<UEditorInteractiveToolsContext>& InToolsContext)
+	: ToolsContext(InToolsContext)
+{
+	check(ToolsContext);
+
+	ToolManager = ToolsContext->ToolManager;
+}
+
+void FPVToolsRegistry::RegisterTools()
+{
+	check(ToolsContext);
+
+	REGISTER_PVE_TOOL(UPVColliderTransformTool, UPVObjectInteractionSettings);
+	REGISTER_PVE_TOOL(UPVSkeletonEditTool, UPVSkeletonEditSettings);
+}
+
+void FPVToolsRegistry::UnregisterTools()
+{
+	check(ToolsContext);
+
+	for (const auto& [_, ToolID] : NodeSettingsToToolMapping)
+	{
+		ToolManager->UnregisterToolType(ToolID);
+	}
+
+	NodeSettingsToToolMapping.Empty();
+}
+
+bool FPVToolsRegistry::StartToolForNode(const TObjectPtr<UPVBaseSettings>& NodeSettings, const TObjectPtr<const UPVData>& InData)
+{
+	if (!NodeSettings)
+	{
+		return false;
+	}
+	
+	if (const TObjectPtr<UPVBaseInteractiveTool> ToolInstance = Cast<UPVBaseInteractiveTool>(ToolManager->ActiveLeftTool))
+	{
+		if (ToolInstance->GetNodeSettings() == NodeSettings)
+		{
+			return true;
+		}
+	}
+
+	if (const FString* ToolIdentifier = NodeSettingsToToolMapping.Find(NodeSettings->GetClass()))
+	{
+		if (UPVToolContextObject* const ToolContext = ToolsContext->ContextObjectStore->FindContext<UPVToolContextObject>())
+		{
+			ToolContext->NodeSettings = NodeSettings;
+			ToolContext->NodeInputData = InData;
+		}
+		
+		ToolManager->SelectActiveToolType(EToolSide::Left, *ToolIdentifier);
+		if (ToolManager->ActivateTool(EToolSide::Left))
+		{
+			return true;
+		}
+	}
+	return false;
+}
+
+bool FPVToolsRegistry::ShutdownCurrentTool(const EToolShutdownType ShutdownType)
+{
+	if (ToolManager->HasActiveTool(EToolSide::Left))
+	{
+		ToolManager->DeactivateTool(EToolSide::Left, ShutdownType);
+		return true;
+	}
+	return false;
+}
+
+bool FPVToolsRegistry::HasActiveTool()
+{
+	return ToolManager->HasActiveTool(EToolSide::Left);
+}
diff --git a/Engine/Plugins/Experimental/ProceduralVegetationEditor/Source/ProceduralVegetationEditor/Private/Tools/PVToolsRegistry.h b/Engine/Plugins/Experimental/ProceduralVegetationEditor/Source/ProceduralVegetationEditor/Private/Tools/PVToolsRegistry.h
new file mode 100644
--- /dev/null
+++ b/Engine/Plugins/Experimental/ProceduralVegetationEditor/Source/ProceduralVegetationEditor/Private/Tools/PVToolsRegistry.h
@@ -0,0 +1,52 @@
+// Copyright Epic Games, Inc. All Rights Reserved.
+
+#pragma once
+
+#include "CoreMinimal.h"
+#include "InteractiveTool.h"
+
+#include "DataTypes/PVData.h"
+
+#include "UObject/ObjectPtr.h"
+
+class UPVBaseSettings;
+class UInteractiveToolManager;
+class UEditorInteractiveToolsContext;
+
+class FPVToolsRegistry
+{
+public:
+	FPVToolsRegistry(const TObjectPtr<UEditorInteractiveToolsContext>& InToolsContext);
+
+	/*
+	 * Register all PVE Tools in this function
+	 */
+	void RegisterTools();
+
+	/*
+	 * Unregister all PVE Tools in this function
+	 */
+	void UnregisterTools();
+
+	/*
+	 * Starts the corresponding tool for the settings
+	 */
+	bool StartToolForNode(const TObjectPtr<UPVBaseSettings>& NodeSettings, const TObjectPtr<const UPVData>& InData = nullptr);
+
+	/*
+	 * Stops the current running tool, if any
+	 */
+	bool ShutdownCurrentTool(EToolShutdownType ShutdownType = EToolShutdownType::Accept);
+	
+	/*
+	 * Checks if there is a tool running
+	 */
+	bool HasActiveTool();
+
+protected:
+	TMap<TObjectPtr<UClass>, FString> NodeSettingsToToolMapping;
+	
+	TObjectPtr<UEditorInteractiveToolsContext> ToolsContext = nullptr;
+
+	TObjectPtr<UInteractiveToolManager> ToolManager = nullptr;
+};
diff --git a/Engine/Plugins/Experimental/ProceduralVegetationEditor/Source/ProceduralVegetationEditor/ProceduralVegetationEditor.Build.cs b/Engine/Plugins/Experimental/ProceduralVegetationEditor/Source/ProceduralVegetationEditor/ProceduralVegetationEditor.Build.cs
--- a/Engine/Plugins/Experimental/ProceduralVegetationEditor/Source/ProceduralVegetationEditor/ProceduralVegetationEditor.Build.cs
+++ b/Engine/Plugins/Experimental/ProceduralVegetationEditor/Source/ProceduralVegetationEditor/ProceduralVegetationEditor.Build.cs
@@ -21,35 +21,38 @@ public class ProceduralVegetationEditor : ModuleRules
 				"CoreUObject",
 				"CurveEditor",
 				"DataflowEditor",
+				"DynamicWind",
+				"DynamicWindEditor",
+				"EditorFramework",
+				"EditorWidgets",
 				"Engine",
 				"GeometryCore",
 				"GeometryFramework",
 				"GeometryScriptingCore",
 				"InputCore",
+				"InteractiveToolsFramework",
 				"Json",
 				"JsonUtilities",
 				"LevelEditor",
-				"MeshConversionEngineTypes",
 				"MeshConversion",
+				"MeshConversionEngineTypes",
 				"ModelingComponents",
-				"PackagesDialog",
 				"PCG",
 				"PCGEditor",
+				"PackagesDialog",
+				"PlanarCut",
 				"ProceduralVegetation",
 				"Projects",
 				"PropertyEditor",
-				"PlanarCut",
+				"RHI",
 				"RenderCore",
+				"SkeletalMeshDescription",
 				"Slate",
 				"SlateCore",
-				"SkeletalMeshDescription",
 				"StaticMeshDescription",
 				"ToolMenus",
 				"ToolWidgets",
-				"UnrealEd",
-				"DynamicWind",
-				"DynamicWindEditor",
-				"RHI"
+				"UnrealEd"
 			}
 		);
 	}
