[PVE] Swarm Warning Fix

#slack_thread: https://epicgames.slack.com/archives/C042RSGQVPD/p1772717071202879

#Synced-CL 51373704

diff --git a/Engine/Plugins/Experimental/ProceduralVegetationEditor/Source/ProceduralVegetationEditor/Private/PVEditor.cpp b/Engine/Plugins/Experimental/ProceduralVegetationEditor/Source/ProceduralVegetationEditor/Private/PVEditor.cpp
--- a/Engine/Plugins/Experimental/ProceduralVegetationEditor/Source/ProceduralVegetationEditor/Private/PVEditor.cpp
+++ b/Engine/Plugins/Experimental/ProceduralVegetationEditor/Source/ProceduralVegetationEditor/Private/PVEditor.cpp
@@ -499,23 +499,21 @@ void FPVEditor::ChangeNodeInspection(UPCGEditorGraphNodeBase* InNode)
 
 void FPVEditor::OnLockNodeSelection()
 {
-	UPCGEditorGraphNodeBase* const SelectedNode = GetFirstSelectedEdNode();
-	UPVBaseSettings* SelectedNodeSettings = GetNodeSettings(SelectedNode->GetPCGNode());
-
 	if (UPVBaseSettings* NodeBeingInspectedSettings = GetNodeSettings(GetNodeBeingInspected()))
 	{
 		NodeBeingInspectedSettings->bLockInspection = !NodeBeingInspectedSettings->bLockInspection;
 	}
 
-	if (!SelectedNode || !SelectedNodeSettings)
-	{
-		return;
-	}
-
-	if (NodeBeingInspected != SelectedNode)
+	if (UPCGEditorGraphNodeBase* const SelectedNode = GetFirstSelectedEdNode())
 	{
-		SelectedNodeSettings->bLockInspection = true;
-		ChangeNodeInspection(SelectedNode);
+		if (UPVBaseSettings* SelectedNodeSettings = GetNodeSettings(SelectedNode->GetPCGNode()))
+		{
+			if (NodeBeingInspected != SelectedNode)
+			{
+				SelectedNodeSettings->bLockInspection = true;
+				ChangeNodeInspection(SelectedNode);
+			}
+		}
 	}
 
 	UpdateStats();
