[PVE] BUG: Interactive Tools Inverted Condition Fix

#Synced-CL 51373704

diff --git a/Engine/Plugins/Experimental/ProceduralVegetationEditor/Source/ProceduralVegetation/Private/Nodes/PVSkeletonEditSettings.cpp b/Engine/Plugins/Experimental/ProceduralVegetationEditor/Source/ProceduralVegetation/Private/Nodes/PVSkeletonEditSettings.cpp
--- a/Engine/Plugins/Experimental/ProceduralVegetationEditor/Source/ProceduralVegetation/Private/Nodes/PVSkeletonEditSettings.cpp
+++ b/Engine/Plugins/Experimental/ProceduralVegetationEditor/Source/ProceduralVegetation/Private/Nodes/PVSkeletonEditSettings.cpp
@@ -90,7 +90,7 @@ bool FPVSkeletonEditElement::ExecuteInternal(FPCGContext* InContext) const
 			const TArray<FVector3f>& Offsets = InputSettings->SkeletonEditSettings.RelativeOffsets;
 
 			const int32* MaxBudNumber = Algo::MaxElement(PointFacade.GetBudNumbersAttribute());
-			if (MaxBudNumber && Offsets.Num() <= *MaxBudNumber + 1)
+			if (MaxBudNumber && Offsets.Num() >= *MaxBudNumber + 1)
 			{
 				TManagedArray<FVector3f>& Positions = PointFacade.ModifyPositions();
 				TManagedArray<FVector3f>& FoliagePositions = FoliageFacade.ModifyPivotPositions();
@@ -150,7 +150,7 @@ bool FPVSkeletonEditElement::ExecuteInternal(FPCGContext* InContext) const
 			TArray<bool> PointsToRemove;
 			TArray<bool> BranchesToRemove;
 			TArray<bool> FoliageInstancesToRemove;
-			if (MaxBudNumber && RemovedPoints.Num() <= *MaxBudNumber + 1)
+			if (MaxBudNumber && RemovedPoints.Num() >= *MaxBudNumber + 1)
 			{
 				PointsToRemove.Init(false, PointFacade.GetElementCount());
 				BranchesToRemove.Init(false, BranchFacade.GetElementCount());
