[PVE] Add grafter palette node and condition system to grafter workflow

- Adds a new grafter palette node that allows users to combine multiple grafts and label them with condition attribute values for light, up alignment, health, and tip attachment
- Encodes the condition attributes into the ManagedArrayCollection for each graft
- Adds a new data type for node output: UPVGrafterPaletteData which inherits from UPCGData
- Grafter node now accepts UPVGrafterPaletteData as input for grafting
- Extracts condition system related logic from the foliage distributor code into a common helper that is re-used in the grafter palette and grafter node
- Adds condition system rules to grafter node for candidate selection for both hormone based and parametric settings based distribution
- Fixes a bug in grafter node where spacing basis was being used instead of scale ramp basis
- Adds minor bug fixes and performance improvements to the foliage distributor

#Synced-CL 51471516

diff --git a/Engine/Plugins/Experimental/ProceduralVegetationEditor/Source/ProceduralVegetation/Private/DataTypes/PVGrafterPaletteData.cpp b/Engine/Plugins/Experimental/ProceduralVegetationEditor/Source/ProceduralVegetation/Private/DataTypes/PVGrafterPaletteData.cpp
new file mode 100644
--- /dev/null
+++ b/Engine/Plugins/Experimental/ProceduralVegetationEditor/Source/ProceduralVegetation/Private/DataTypes/PVGrafterPaletteData.cpp
@@ -0,0 +1,41 @@
+// Copyright Epic Games, Inc. All Rights Reserved.
+
+#include "DataTypes/PVGrafterPaletteData.h"
+
+PCG_DEFINE_TYPE_INFO(FPVDataTypeInfoGrafterPalette, UPVGrafterPaletteData);
+
+void UPVGrafterPaletteData::Initialize(TArray<TObjectPtr<UPVGrowthData>>&& InGrowthDataElements)
+{
+	GrowthDataElements = MoveTemp(InGrowthDataElements);
+}
+
+void UPVGrafterPaletteData::AddToCrc(FArchiveCrc32& Ar, bool bFullDataCrc) const
+{
+	Super::AddToCrc(Ar, bFullDataCrc);
+	
+	if (bFullDataCrc)
+	{
+		FString ClassName = StaticClass()->GetPathName();
+
+		Ar << ClassName;
+		
+		int32 Num = GrowthDataElements.Num();
+		Ar << Num;
+		
+		for (const auto& Element : GrowthDataElements)
+		{
+			bool bIsValid = (Element != nullptr);
+			Ar << bIsValid;
+			
+			if (bIsValid)
+			{
+				Element->AddToCrc(Ar, true);
+			}
+		}
+	}
+	else
+	{
+		AddUIDToCrc(Ar);
+	}
+}
+
diff --git a/Engine/Plugins/Experimental/ProceduralVegetationEditor/Source/ProceduralVegetation/Private/Facades/PVFoliageConditionFacade.h b/Engine/Plugins/Experimental/ProceduralVegetationEditor/Source/ProceduralVegetation/Private/Facades/PVFoliageConditionFacade.h
--- a/Engine/Plugins/Experimental/ProceduralVegetationEditor/Source/ProceduralVegetation/Private/Facades/PVFoliageConditionFacade.h
+++ b/Engine/Plugins/Experimental/ProceduralVegetationEditor/Source/ProceduralVegetation/Private/Facades/PVFoliageConditionFacade.h
@@ -4,6 +4,7 @@
 
 #include "CoreMinimal.h"
 #include "GeometryCollection/ManagedArrayAccessor.h"
+#include "Helpers/PVFoliageHelper.h"
 #include "Implementations/PVFoliage.h"
 
 struct FPVFoliageInfo;
diff --git a/Engine/Plugins/Experimental/ProceduralVegetationEditor/Source/ProceduralVegetation/Private/Facades/PVFoliageFacade.cpp b/Engine/Plugins/Experimental/ProceduralVegetationEditor/Source/ProceduralVegetation/Private/Facades/PVFoliageFacade.cpp
--- a/Engine/Plugins/Experimental/ProceduralVegetationEditor/Source/ProceduralVegetation/Private/Facades/PVFoliageFacade.cpp
+++ b/Engine/Plugins/Experimental/ProceduralVegetationEditor/Source/ProceduralVegetation/Private/Facades/PVFoliageFacade.cpp
@@ -8,6 +8,7 @@
 #include "VisualizeTexture.h"
 #include "Chaos/Deformable/MuscleActivationConstraints.h"
 #include "Facades/PVAttributesNames.h"
+#include "Helpers/PVFoliageHelper.h"
 #include "Implementations/PVFoliage.h"
  
 namespace PV::Facades
diff --git a/Engine/Plugins/Experimental/ProceduralVegetationEditor/Source/ProceduralVegetation/Private/Facades/PVMetaInfoFacade.cpp b/Engine/Plugins/Experimental/ProceduralVegetationEditor/Source/ProceduralVegetation/Private/Facades/PVMetaInfoFacade.cpp
--- a/Engine/Plugins/Experimental/ProceduralVegetationEditor/Source/ProceduralVegetation/Private/Facades/PVMetaInfoFacade.cpp
+++ b/Engine/Plugins/Experimental/ProceduralVegetationEditor/Source/ProceduralVegetation/Private/Facades/PVMetaInfoFacade.cpp
@@ -12,6 +12,10 @@ namespace PV::Facades
 		, LeafGrowthAttribute(InCollection, AttributeNames::LeafGrowth, GroupNames::DetailsGroup)
 		, AbscissionSenescenseAttribute(InCollection, AttributeNames::AbscissionSenescense, GroupNames::DetailsGroup)
 		, LateralElongationAttribute(InCollection, AttributeNames::LateralElongation, GroupNames::DetailsGroup)
+		, GraftLightAttribute(InCollection, AttributeNames::GraftLight, GroupNames::DetailsGroup)
+		, GraftUpAlignmentAttribute(InCollection, AttributeNames::GraftUpAlignment, GroupNames::DetailsGroup)
+		, GraftHealthAttribute(InCollection, AttributeNames::GraftHealth, GroupNames::DetailsGroup)
+		, GraftTipAttribute(InCollection, AttributeNames::GraftTip, GroupNames::DetailsGroup)
 	{
 		DefineSchema(InitialSize);
 	}
@@ -23,6 +27,10 @@ namespace PV::Facades
 		, LeafGrowthAttribute(InCollection, AttributeNames::LeafGrowth, GroupNames::DetailsGroup)
 		, AbscissionSenescenseAttribute(InCollection, AttributeNames::AbscissionSenescense, GroupNames::DetailsGroup)
 		, LateralElongationAttribute(InCollection, AttributeNames::LateralElongation, GroupNames::DetailsGroup)
+		, GraftLightAttribute(InCollection, AttributeNames::GraftLight, GroupNames::DetailsGroup)
+		, GraftUpAlignmentAttribute(InCollection, AttributeNames::GraftUpAlignment, GroupNames::DetailsGroup)
+		, GraftHealthAttribute(InCollection, AttributeNames::GraftHealth, GroupNames::DetailsGroup)
+		, GraftTipAttribute(InCollection, AttributeNames::GraftTip, GroupNames::DetailsGroup)
 	{}
 
 	void FMetaInfoFacade::DefineSchema(const int32 InitialSize)
@@ -35,6 +43,10 @@ namespace PV::Facades
 		}
 
 		GuidAttribute.Add();
+		GraftLightAttribute.Add();
+		GraftUpAlignmentAttribute.Add();
+		GraftHealthAttribute.Add();
+		GraftTipAttribute.Add();
 	}
 
 	bool FMetaInfoFacade::IsValid() const
@@ -66,6 +78,42 @@ namespace PV::Facades
 		return FGuid::NewGuid();
 	}
 
+	void FMetaInfoFacade::SetGraftAttributes(const FGraftData& InGraftData)
+	{
+		check(!IsConst());
+
+		int32 NumElements = Collection->NumElements(GroupNames::DetailsGroup);
+
+		if(NumElements == 0)
+		{
+			Collection->AddElements(1, GroupNames::DetailsGroup);
+		}
+		
+		GraftLightAttribute.Modify()[0] = InGraftData.Light;
+		GraftUpAlignmentAttribute.Modify()[0] = InGraftData.UpAlignment;
+		GraftHealthAttribute.Modify()[0] = InGraftData.Health;
+		GraftTipAttribute.Modify()[0] = InGraftData.Tip ? 1 : 0;
+	}
+
+	FGraftData FMetaInfoFacade::GetGraftAttributes() const
+	{
+		FGraftData GraftData = {};
+		
+		if (GraftLightAttribute.IsValidIndex(0))
+			GraftData.Light = GraftLightAttribute[0];
+		
+		if (GraftUpAlignmentAttribute.IsValidIndex(0))
+			GraftData.UpAlignment = GraftUpAlignmentAttribute[0];
+		
+		if (GraftHealthAttribute.IsValidIndex(0))
+			GraftData.Health = GraftHealthAttribute[0];
+			
+		if (GraftTipAttribute.IsValidIndex(0))
+			GraftData.Tip = (GraftTipAttribute[0] > 0);
+		
+		return GraftData;
+	}
+
 	const TArray<float>& FMetaInfoFacade::GetLeafGrowth() const
 	{
 		if(LeafGrowthAttribute.IsValid() && LeafGrowthAttribute.IsValidIndex(0))
diff --git a/Engine/Plugins/Experimental/ProceduralVegetationEditor/Source/ProceduralVegetation/Private/Helpers/PVAnalyticsHelper.cpp b/Engine/Plugins/Experimental/ProceduralVegetationEditor/Source/ProceduralVegetation/Private/Helpers/PVAnalyticsHelper.cpp
--- a/Engine/Plugins/Experimental/ProceduralVegetationEditor/Source/ProceduralVegetation/Private/Helpers/PVAnalyticsHelper.cpp
+++ b/Engine/Plugins/Experimental/ProceduralVegetationEditor/Source/ProceduralVegetation/Private/Helpers/PVAnalyticsHelper.cpp
@@ -4,6 +4,7 @@
  
 #include "EngineAnalytics.h"
 #include "ProceduralVegetationModule.h"
+#include "PVFoliageHelper.h"
 #include "Params/PVExportParams.h"
 #include "Engine/RendererSettings.h"
 #include "Facades/PVFoliageFacade.h"
diff --git a/Engine/Plugins/Experimental/ProceduralVegetationEditor/Source/ProceduralVegetation/Private/Helpers/PVFoliageHelper.cpp b/Engine/Plugins/Experimental/ProceduralVegetationEditor/Source/ProceduralVegetation/Private/Helpers/PVFoliageHelper.cpp
new file mode 100644
--- /dev/null
+++ b/Engine/Plugins/Experimental/ProceduralVegetationEditor/Source/ProceduralVegetation/Private/Helpers/PVFoliageHelper.cpp
@@ -0,0 +1,488 @@
+// Copyright Epic Games, Inc. All Rights Reserved.
+
+#include "Helpers/PVFoliageHelper.h"
+#include "Utils/PVAttributes.h"
+
+
+
+namespace PV::DistributionConditionUtils
+{
+	float FConditionAttributes::GetConditionValue(EFoliageDistributionCondition InCondition, const int Index) const
+	{
+		const TArray<float>* Result = nullptr;
+		switch (InCondition)
+		{
+		case EFoliageDistributionCondition::Light:
+			{
+				Result = &Light;
+				break;
+			}
+		case EFoliageDistributionCondition::Scale:
+			{
+				Result = &Scale;
+				break;;
+			}
+		case EFoliageDistributionCondition::Tip:
+			{
+				Result = &Tip;
+				break;
+			}
+		case EFoliageDistributionCondition::UpAlignment:
+			{
+				Result = &UpAlignment;
+				break;
+			}
+		case EFoliageDistributionCondition::Health:
+			{
+				Result = &Health;
+				break;
+			}
+				
+		default:
+			break;
+		}
+ 
+		if (Result && Result->IsValidIndex(Index))
+			[[likely]]
+			{
+				return (*Result)[Index];
+			}
+ 
+		return 0;
+	}
+	
+	void FNormalizedPointCaches::Reset(const int32 NumPoints)
+	{
+		PScales.Reset();
+		Light.Reset();
+		Health.Reset();
+			
+		PScales.SetNumZeroed(NumPoints);
+		Light.SetNumZeroed(NumPoints);
+		Health.SetNumZeroed(NumPoints);
+			
+		bHasPScales = false;
+		bHasLight = false;
+		bHasHealth = false;
+	}
+
+	void BuildNormalizedPointCaches(
+		const FManagedArrayCollection& Collection,
+		const FBuildOptions& Options,
+		FNormalizedPointCaches& OutCaches)
+	{
+		auto PointBudLightDetectedAttribute = PV::FBudLightDetectedAttribute::GetAttribute(Collection);
+		auto PointScaleAttribute = PV::FPointScaleAttribute::GetAttribute(Collection);
+		auto PointBudHormoneLevelsAttribute = PV::FBudHormoneLevelsAttribute::GetAttribute(Collection);
+		auto PointBudDevelopmentAttribute = PV::FBudDevelopmentAttribute::GetAttribute(Collection);
+
+		const int32 NumOfPoints = PointScaleAttribute.Num();
+		OutCaches.Reset(NumOfPoints);
+		if (NumOfPoints == 0)
+		{
+			return;
+		}
+
+		float MinPScale = TNumericLimits<float>::Max();
+		float MaxPScale = TNumericLimits<float>::Lowest();
+
+		float MinLight = TNumericLimits<float>::Max();
+		float MaxLight = TNumericLimits<float>::Lowest();
+
+		float MinHealth = TNumericLimits<float>::Max();
+		float MaxHealth = TNumericLimits<float>::Lowest();
+
+		for (int32 PointIndex = 0; PointIndex < NumOfPoints; ++PointIndex)
+		{
+			if (Options.bComputePScales)
+			{
+				const float PScale = PointScaleAttribute[PointIndex];
+				MinPScale = FMath::Min(MinPScale, PScale);
+				MaxPScale = FMath::Max(MaxPScale, PScale);
+			}
+
+			if (Options.bComputeLight)
+			{
+				const float Available = PointBudLightDetectedAttribute[PointIndex].Availible;
+				MinLight = FMath::Min(MinLight, Available);
+				MaxLight = FMath::Max(MaxLight, Available);
+			}
+
+			if (Options.bComputeHealth)
+			{
+				const float Ethylene = PointBudHormoneLevelsAttribute[PointIndex].Ethylene;
+				const int32 AgeSen = PointBudDevelopmentAttribute[PointIndex].AgeSenescense;
+				const int32 LightSen = PointBudDevelopmentAttribute[PointIndex].LightSenescense;
+				const int32 Senescence = FMath::Max(AgeSen, LightSen) + 1;
+
+				const float Health = 1.0f - (Ethylene * (1.0f / FMath::Max(1, Senescence)));
+				OutCaches.Health[PointIndex] = Health;
+
+				MinHealth = FMath::Min(MinHealth, Health);
+				MaxHealth = FMath::Max(MaxHealth, Health);
+			}
+		}
+
+		if (Options.bComputePScales)
+		{
+			OutCaches.bHasPScales = (MaxPScale > MinPScale + KINDA_SMALL_NUMBER);
+			for (int32 i = 0; i < NumOfPoints; ++i)
+			{
+				const float PScale = PointScaleAttribute[i];
+				OutCaches.PScales[i] = OutCaches.bHasPScales
+					? FMath::GetMappedRangeValueClamped(
+						FVector2f(MinPScale, MaxPScale),
+						FVector2f(0.f, 1.f),
+						PScale)
+					: 0.0f;
+			}
+		}
+
+		if (Options.bComputeLight)
+		{
+			OutCaches.bHasLight = (MaxLight > MinLight + KINDA_SMALL_NUMBER);
+			for (int32 i = 0; i < NumOfPoints; ++i)
+			{
+				const float Available = PointBudLightDetectedAttribute[i].Availible;
+				OutCaches.Light[i] = OutCaches.bHasLight
+					? FMath::GetMappedRangeValueClamped(
+						FVector2f(MinLight, MaxLight),
+						FVector2f(0.0f, 1.0f),
+						Available)
+					: 0.0f;
+			}
+		}
+
+		if (Options.bComputeHealth)
+		{
+			OutCaches.bHasHealth = (MaxHealth > MinHealth + KINDA_SMALL_NUMBER);
+			for (int32 i = 0; i < NumOfPoints; ++i)
+			{
+				const float Current = OutCaches.Health[i];
+				OutCaches.Health[i] = OutCaches.bHasHealth
+					? FMath::GetMappedRangeValueClamped(
+						FVector2f(MinHealth, MaxHealth),
+						FVector2f(0.f, 1.f),
+						Current)
+					: 0.0f;
+			}
+		}
+	}
+
+	namespace ConditionPicker
+	{
+		struct FWeightInfo
+		{
+			int32 CandidateIndex = INDEX_NONE;
+			float Weight = 0.0f;
+		};
+
+		bool GetActiveConditionSettings(
+			const FPVFoliageDistributionConditions& DistributionConditions,
+			const EFoliageDistributionCondition Condition,
+			FPVFoliageDistributionConditionSettings& OutSettings)
+		{
+			if (!DistributionConditions.IsActiveCondition(Condition))
+			{
+				return false;
+			}
+
+			return DistributionConditions.GetSettings(Condition, OutSettings);
+		}
+
+		static void ShuffleInPlace(TArray<FWeightInfo>& Weights, const FRandomStream& RandomStream)
+		{
+			for (int32 i = Weights.Num() - 1; i > 0; --i)
+			{
+				const int32 SwapIndex = RandomStream.RandRange(0, i);
+				Weights.Swap(i, SwapIndex);
+			}
+		}
+
+		int32 PickCandidateIndex(
+			const int32 NumCandidates,
+			const FPVFoliageDistributionConditions& DistributionConditions,
+			const FPointConditionSample& PointSample,
+			const FRandomStream& RandomStream,
+			const FCandidateValueGetter& GetCandidateValue,
+			const FCandidateValidFn& IsCandidateValid)
+		{
+			if (NumCandidates < 1)
+			{
+				return INDEX_NONE;
+			}
+
+			TArray<FWeightInfo> Weights;
+			Weights.Reserve(NumCandidates);
+
+			for (int32 i = 0; i < NumCandidates; ++i)
+			{
+				if (IsCandidateValid(i))
+				{
+					Weights.Add({i, 0.0f});
+				}
+			}
+
+			if (Weights.Num() == 0)
+			{
+				return INDEX_NONE;
+			}
+
+			if (!DistributionConditions.HasActiveCondition())
+			{
+				return Weights[RandomStream.RandRange(0, Weights.Num() - 1)].CandidateIndex;
+			}
+
+			ShuffleInPlace(Weights, RandomStream);
+
+			for (const EFoliageDistributionCondition Condition : TEnumRange<EFoliageDistributionCondition>())
+			{
+				FPVFoliageDistributionConditionSettings ConditionSettings;
+				if (!GetActiveConditionSettings(DistributionConditions, Condition, ConditionSettings))
+				{
+					continue;
+				}
+
+				const float Target = PointSample.Get(Condition) + ConditionSettings.Offset;
+
+				for (auto& [CandidateIndex, Weight] : Weights)
+				{
+					const float CandidateValue = GetCandidateValue(CandidateIndex, Condition);
+					const float Dist = FMath::Abs(CandidateValue - Target);
+					Weight += (Dist * ConditionSettings.Weight);
+				}
+			}
+
+			float MinWeight = TNumericLimits<float>::Max();
+			float MaxWeight = TNumericLimits<float>::Lowest();
+			for (const auto& [CandidateIndex, Weight] : Weights)
+			{
+				MinWeight = FMath::Min(MinWeight, Weight);
+				MaxWeight = FMath::Max(MaxWeight, Weight);
+			}
+
+			const bool bHasRange = (MaxWeight > MinWeight + KINDA_SMALL_NUMBER);
+			for (auto& [CandidateIndex, Weight] : Weights)
+			{
+				Weight = bHasRange
+					? FMath::GetMappedRangeValueClamped(
+						FVector2f(MinWeight, MaxWeight),
+						FVector2f(0.0f, 1.0f),
+						Weight)
+					: 0.0f;
+			}
+
+			Weights.Sort([](const FWeightInfo& A, const FWeightInfo& B)
+				{
+					return A.Weight < B.Weight;
+				});
+
+			const int32 MinCandidates = FMath::Clamp(DistributionConditions.MinimumCandidates, 1, Weights.Num());
+
+			TArray<int32> Picked;
+			Picked.Reserve(Weights.Num());
+
+			for (const auto& [CandidateIndex, Weight] : Weights)
+			{
+				if (Weight < DistributionConditions.CutoffThreshold)
+				{
+					Picked.Add(CandidateIndex);
+				}
+				else if (Picked.Num() < MinCandidates)
+				{
+					Picked.Add(CandidateIndex);
+				}
+				else
+				{
+					break;
+				}
+			}
+
+			if (Picked.Num() == 0)
+			{
+				return Weights[0].CandidateIndex;
+			}
+
+			return Picked[RandomStream.RandRange(0, Picked.Num() - 1)];
+		}
+
+		int32 PickCandidateIndex(
+			const int32 NumCandidates,
+			const FPVFoliageDistributionConditions& DistributionConditions,
+			const FPointConditionSample& PointSample,
+			const FRandomStream& RandomStream,
+			const FCandidateValueGetter& GetCandidateValue)
+		{
+			return PickCandidateIndex(
+				NumCandidates,
+				DistributionConditions,
+				PointSample,
+				RandomStream,
+				GetCandidateValue,
+				[](int32)
+					{
+						return true;
+					});
+		}
+	}
+}
+
+bool FPVFoliageDistributionConditions::HasActiveCondition() const
+{
+		for (EFoliageDistributionCondition Condition : TEnumRange<EFoliageDistributionCondition>())
+		{
+			if (IsActiveCondition(Condition))
+			{
+				return true;
+			}
+		}
+		return false;
+}
+
+bool FPVFoliageDistributionConditions::IsActiveCondition(const EFoliageDistributionCondition InCondition) const
+{
+		switch (InCondition)
+		{
+		case EFoliageDistributionCondition::Light:
+			return bActivateLight && Light.Weight > 0.0f;
+			
+		case EFoliageDistributionCondition::Scale:
+			return bActivateScale && Scale.Weight > 0.0f; 
+			
+		case EFoliageDistributionCondition::UpAlignment:
+			return bActivateUpAlignment && UpAlignment.Weight > 0.0f;
+			
+		case EFoliageDistributionCondition::Tip:
+			return bActivateTip && Tip.Weight > 0.0f;
+	
+		case EFoliageDistributionCondition::Health:
+			return bActivateHealth && Health.Weight > 0.0f;
+	
+		default:
+			return false;
+		}
+}
+
+void FPVFoliageDistributionConditions::SetConditionState(const EFoliageDistributionCondition InCondition, const bool InNewState)
+{
+	switch (InCondition)
+	{
+	case EFoliageDistributionCondition::Light:
+		bActivateLight = InNewState;
+		break;
+			
+			
+	case EFoliageDistributionCondition::Scale:
+		bActivateScale = InNewState;
+		break;
+			
+	case EFoliageDistributionCondition::UpAlignment:
+		bActivateUpAlignment = InNewState;
+		break;
+			
+	case EFoliageDistributionCondition::Tip:
+		bActivateTip = InNewState;
+		break;
+ 
+	case EFoliageDistributionCondition::Health:
+		bActivateHealth = InNewState;
+		break;
+ 
+	default:
+		check(false);
+		break;
+	}
+}
+
+FPVFoliageDistributionConditionSettings* FPVFoliageDistributionConditions::GetSettings(const EFoliageDistributionCondition InCondition)
+{
+	switch (InCondition)
+	{
+	case EFoliageDistributionCondition::Light:
+		return &Light;
+				
+	case EFoliageDistributionCondition::Scale:
+		return &Scale;
+			
+				
+	case EFoliageDistributionCondition::UpAlignment:
+		return &UpAlignment;
+				
+	case EFoliageDistributionCondition::Tip:
+		return &Tip;
+ 
+	case EFoliageDistributionCondition::Health:
+		return &Health;
+ 
+	default:
+		return nullptr;
+	}
+}
+
+bool FPVFoliageDistributionConditions::GetSettings(const EFoliageDistributionCondition InCondition,
+	FPVFoliageDistributionConditionSettings& OutSettings) const
+{
+	switch (InCondition)
+	{
+	case EFoliageDistributionCondition::Light:
+		{
+			OutSettings =  Light;
+			return true;
+		}
+			
+	case EFoliageDistributionCondition::Scale:
+		{
+			OutSettings = Scale;
+			return true;
+		}
+			
+	case EFoliageDistributionCondition::UpAlignment:
+		{
+			OutSettings = UpAlignment;
+			return true;
+		}
+			
+	case EFoliageDistributionCondition::Tip:
+		{
+			OutSettings = Tip;
+			return true;
+		}
+ 
+	case EFoliageDistributionCondition::Health:
+		{
+			OutSettings = Health;
+			return true;
+		}
+	default:
+		return false;
+	}
+}
+
+float FPVFoliageAttributes::GetAttributeValue(const EFoliageDistributionCondition InCondition) const
+{
+	switch (InCondition)
+	{
+	case EFoliageDistributionCondition::Light:
+		return Light;
+			
+	case EFoliageDistributionCondition::Scale:
+		return Scale;
+			
+	case EFoliageDistributionCondition::UpAlignment:
+		return UpAlignment;
+			
+	case EFoliageDistributionCondition::Tip:
+		return Tip ? 1.0f : 0.0f;
+ 
+	case EFoliageDistributionCondition::Health:
+		return Health;
+ 
+	default:
+		return 0.0f;
+	}
+}
+
+FString FPVFoliageAttributes::ToString() const
+{
+	return FString::Printf(TEXT("Light = %f, Scale = %f, UpAlignment = %f, Health = %f, Tip = %s"), Light, Scale, UpAlignment, Health, Tip ? TEXT("True") : TEXT("False"));
+}
diff --git a/Engine/Plugins/Experimental/ProceduralVegetationEditor/Source/ProceduralVegetation/Private/Helpers/PVFoliageHelper.h b/Engine/Plugins/Experimental/ProceduralVegetationEditor/Source/ProceduralVegetation/Private/Helpers/PVFoliageHelper.h
new file mode 100644
--- /dev/null
+++ b/Engine/Plugins/Experimental/ProceduralVegetationEditor/Source/ProceduralVegetation/Private/Helpers/PVFoliageHelper.h
@@ -0,0 +1,249 @@
+// Copyright Epic Games, Inc. All Rights Reserved.
+
+#pragma once
+
+#include "CoreMinimal.h"
+#include "GeometryCollection/ManagedArrayCollection.h"
+#include "Misc/EnumRange.h"
+
+#include "PVFoliageHelper.generated.h"
+
+
+UENUM()
+enum class EFoliageDistributionCondition : uint8
+{
+	Light,
+	Scale,
+	UpAlignment,
+	Tip,
+	Health,
+	Count UMETA(Hidden),
+	None UMETA(Hidden),
+};
+ 
+ENUM_RANGE_BY_COUNT(EFoliageDistributionCondition, EFoliageDistributionCondition::Count)
+
+USTRUCT()
+struct FPVFoliageDistributionConditionSettings
+{
+	GENERATED_BODY()
+ 
+	UPROPERTY(EditAnywhere, Category="ConditionSettings", meta=(ClampMin=0.0f, ClampMax=1.0f, UIMin=0.0f, UIMax=1.0f, Tooltip=""))
+	float Weight = 1.0f;
+ 
+	UPROPERTY(EditAnywhere, Category="ConditionSettings", meta=(ClampMin=-1.0f, ClampMax=1.0f, UIMin=-1.0f, UIMax=1.0f, Tooltip=""))
+	float Offset = 0.0f;
+ 
+	FString ToString() const
+	{
+		return FString::Printf(TEXT("Weight = %f, Offset = %f"), Weight, Offset);
+	}
+};
+ 
+ 
+USTRUCT()
+struct FPVFoliageDistributionConditions
+{
+	GENERATED_BODY()
+ 
+	UPROPERTY(EditAnywhere, Category="ConditionSettings", meta=(ClampMin=0.0f, ClampMax=1.0f, UIMin=0.0f, UIMax=1.0f, Tooltip=""))
+	float CutoffThreshold = 0.3f;
+ 
+	UPROPERTY(EditAnywhere, Category="ConditionSettings", meta=(ClampMin=1, ClampMax=10, UIMin=0, UIMax=10, Tooltip=""))
+	int32 MinimumCandidates = 1;
+	
+	UPROPERTY(EditAnywhere, Category="ConditionSettings", meta=(PinHiddenByDefault, InlineEditConditionToggle))
+	bool bActivateLight = false;
+ 
+	UPROPERTY(EditAnywhere, Category="ConditionSettings", meta=(EditCondition="bActivateLight", Tooltip=""))
+	FPVFoliageDistributionConditionSettings Light;
+ 
+	UPROPERTY()
+	bool bActivateScale = false;
+ 
+	UPROPERTY()
+	FPVFoliageDistributionConditionSettings Scale;
+ 
+	UPROPERTY(EditAnywhere, Category="ConditionSettings", meta=(PinHiddenByDefault, InlineEditConditionToggle))
+	bool bActivateUpAlignment = false;
+ 
+	UPROPERTY(EditAnywhere, Category="ConditionSettings", meta=(EditCondition="bActivateUpAlignment", Tooltip=""))
+	FPVFoliageDistributionConditionSettings UpAlignment;
+ 
+	UPROPERTY(EditAnywhere, Category="ConditionSettings", meta=(PinHiddenByDefault, InlineEditConditionToggle))
+	bool bActivateTip = false;
+	
+	UPROPERTY(EditAnywhere, Category="ConditionSettings", meta=(EditCondition="bActivateTip", Tooltip=""))
+	FPVFoliageDistributionConditionSettings Tip;
+ 
+	UPROPERTY(EditAnywhere, Category="ConditionSettings", meta=(PinHiddenByDefault, InlineEditConditionToggle))
+	bool bActivateHealth = false;
+	
+	UPROPERTY(EditAnywhere, Category="ConditionSettings", meta=(EditCondition="bActivateHealth", Tooltip=""))
+	FPVFoliageDistributionConditionSettings Health;
+ 
+	bool HasActiveCondition() const;
+ 
+	bool IsActiveCondition(const EFoliageDistributionCondition InCondition) const;
+ 
+	void SetConditionState(const EFoliageDistributionCondition InCondition, const bool InNewState);
+ 
+	FPVFoliageDistributionConditionSettings* GetSettings(const EFoliageDistributionCondition InCondition);
+ 
+	bool GetSettings(const EFoliageDistributionCondition InCondition, FPVFoliageDistributionConditionSettings& OutSettings) const;
+};
+
+
+USTRUCT()
+struct FPVFoliageAttributes
+{
+	GENERATED_BODY()
+ 
+	UPROPERTY(EditAnywhere, Category="Info", meta=(ClampMin=0.0f, ClampMax=1.0f, UIMin=0.0f, UIMax=1.0f, Tooltip=""))
+	float Light = 0.0f;
+ 
+	UPROPERTY()
+	float Scale = 0.0f;
+ 
+	UPROPERTY(EditAnywhere, Category="Info", meta=(ClampMin=0.0f, ClampMax=1.0f, UIMin=0.0f, UIMax=1.0f, Tooltip=""))
+	float UpAlignment = 0.0f;
+ 
+	UPROPERTY(EditAnywhere, Category="Info", meta=(ClampMin=0.0f, ClampMax=1.0f, UIMin=0.0f, UIMax=1.0f, Tooltip=""))
+	float Health = 0.0f;
+ 
+	UPROPERTY(EditAnywhere, Category="Info", meta=(Tooltip=""))
+	bool Tip = false;
+ 
+	float GetAttributeValue(const EFoliageDistributionCondition InCondition) const;
+ 
+	FString ToString() const;
+};
+
+
+USTRUCT()
+struct FPVFoliageInfo
+{
+	GENERATED_BODY()
+	
+	UPROPERTY(EditAnywhere, Category = "Info")
+	TSoftObjectPtr<UObject> Mesh;
+ 
+	UPROPERTY(EditAnywhere, Category="Info", meta=(ShowInnerProperties, FullyExpand="true", Tooltip=""))
+	FPVFoliageAttributes Attributes;
+ 
+	bool IsValid() const
+	{
+		// Asset is loaded and pointer is valid
+		// Path is valid, but asset is NOT loaded yet
+		return Mesh.IsValid() || Mesh.ToSoftObjectPath().IsValid(); 
+	}
+};
+
+namespace PV::DistributionConditionUtils
+{
+	struct FConditionAttributes
+	{
+		TArray<float> UpAlignment;
+		TArray<float> Light;
+		TArray<float> Scale;
+		TArray<float> Tip;
+		TArray<float> Health;
+	
+ 
+		float GetConditionValue(EFoliageDistributionCondition InCondition, const int Index) const;
+	};
+	
+	struct FNormalizedPointCaches
+	{
+		TArray<float> PScales;
+		TArray<float> Light;
+		TArray<float> Health;
+		bool bHasPScales = false;
+		bool bHasLight  = false;
+		bool bHasHealth = false;
+
+		void Reset(const int32 NumPoints);
+	};
+
+	struct FBuildOptions
+	{
+		bool bComputePScales = true;
+		bool bComputeLight  = true;
+		bool bComputeHealth = true;
+	};
+
+	FORCEINLINE float UpAlignmentNormalized(const FVector3f& ApicalDirection)
+	{
+		const float Dot = FVector3f::DotProduct(ApicalDirection.GetSafeNormal(), FVector3f::UpVector);
+		
+		return FMath::GetMappedRangeValueClamped(
+			FVector2f(-1.0f, 1.0f), 
+			FVector2f(0.0f, 1.0f), 
+			Dot);
+	}
+
+	FORCEINLINE float Tip(const bool bIsTip)
+	{
+		return bIsTip ? 1.0f : 0.0f;
+	}
+
+	void BuildNormalizedPointCaches(
+		const FManagedArrayCollection& Collection,
+		const FBuildOptions& Options,
+		FNormalizedPointCaches& OutCaches);
+
+	namespace ConditionPicker
+	{
+		struct FPointConditionSample
+		{
+			float Light = 0.0f;
+			float Scale = 0.0f;
+			float UpAlignment = 0.0f;
+			float Tip = 0.0f;
+			float Health = 0.0f;
+
+			FORCEINLINE float Get(EFoliageDistributionCondition Condition) const
+			{
+				switch (Condition)
+				{
+				case EFoliageDistributionCondition::Light:
+					return Light;
+				case EFoliageDistributionCondition::Scale:
+					return Scale;
+				case EFoliageDistributionCondition::UpAlignment:
+					return UpAlignment;
+				case EFoliageDistributionCondition::Tip:
+					return Tip;
+				case EFoliageDistributionCondition::Health:
+					return Health;
+				default:
+					return 0.0f;
+				}
+			}
+		};
+	
+		using FCandidateValueGetter = TFunctionRef<float(int32, EFoliageDistributionCondition)>;
+
+		using FCandidateValidFn = TFunctionRef<bool(int32)>;
+
+		int32 PickCandidateIndex(
+			const int32 NumCandidates,
+			const FPVFoliageDistributionConditions& DistributionConditions,
+			const FPointConditionSample& PointSample,
+			const FRandomStream& RandomStream,
+			const FCandidateValueGetter& GetCandidateValue,
+			const FCandidateValidFn& IsCandidateValid);
+		
+		int32 PickCandidateIndex(
+			const int32 NumCandidates,
+			const FPVFoliageDistributionConditions& DistributionConditions,
+			const FPointConditionSample& PointSample,
+			const FRandomStream& RandomStream,
+			const FCandidateValueGetter& GetCandidateValue);
+	
+		bool GetActiveConditionSettings(
+			const FPVFoliageDistributionConditions& DistributionConditions,
+			const EFoliageDistributionCondition Condition,
+			FPVFoliageDistributionConditionSettings& OutSettings);
+	}
+}
diff --git a/Engine/Plugins/Experimental/ProceduralVegetationEditor/Source/ProceduralVegetation/Private/Implementations/PVFoliage.cpp b/Engine/Plugins/Experimental/ProceduralVegetationEditor/Source/ProceduralVegetation/Private/Implementations/PVFoliage.cpp
--- a/Engine/Plugins/Experimental/ProceduralVegetationEditor/Source/ProceduralVegetation/Private/Implementations/PVFoliage.cpp
+++ b/Engine/Plugins/Experimental/ProceduralVegetationEditor/Source/ProceduralVegetation/Private/Implementations/PVFoliage.cpp
@@ -5,76 +5,24 @@
 #include "ProceduralVegetationModule.h"
 #include "Utils/PVFloatRamp.h"
 #include "Algo/RandomShuffle.h"
- 
+
 #include "Facades/PVAttributesNames.h"
 #include "Facades/PVBranchFacade.h"
 #include "Facades/PVBudVectorsFacade.h"
 #include "Facades/PVFoliageFacade.h"
 #include "Facades/PVPointFacade.h"
+#include "Helpers/PVFoliageHelper.h"
 #include "Helpers/PVPhyllotaxyHelper.h"
 
-struct ConditionAttributes
-{
-	TArray<float> UpAlignment;
-	TArray<float> Light;
-	TArray<float> Scale;
-	TArray<float> Tip;
-	TArray<float> Health;
-	
- 
-	float GetConditionValue(EFoliageDistributionCondition InCondition, const int Index) const
-	{
-		TArray<float> Result;
-		switch (InCondition)
-		{
-		case EFoliageDistributionCondition::Light:
-			{
-				Result = Light;
-				break;
-			}
-		case EFoliageDistributionCondition::Scale:
-			{
-				Result = Scale;
-				break;;
-			}
-		case EFoliageDistributionCondition::Tip:
-			{
-				Result = Tip;
-				break;
-			}
-		case EFoliageDistributionCondition::UpAlignment:
-			{
-				Result = UpAlignment;
-				break;
-			}
-		case EFoliageDistributionCondition::Health:
-			{
-				Result = Health;
-				break;
-			}
-		}
- 
-		if (Result.IsValidIndex(Index))
-		{
-			return Result[Index];
-		}
- 
-		return 0;
-	}
-};
- 
-// Condition Weight
-struct FFoliageWeightInfo
-{
-	int FoliageIndex = 0;
-	float Weight = 0;
-};
  
 void FPVFoliage::DistributeFoliage(FManagedArrayCollection& OutCollection, const FManagedArrayCollection& SourceCollection, const FManagedArrayCollection& FoliageCollection,
                                    const FDistributionSettings& DistributionSettings, const FScaleSettings& ScaleSettings,
                                    const FVectorSettings& VectorSettings, const FPhyllotaxySettings& PhyllotaxySettings,
                                    const FMiscSettings& MiscSettings, const FPVFoliageDistributionConditions& DistributionConditions)
 {
+	using namespace PV::DistributionConditionUtils;
+	using PV::DistributionConditionUtils::FConditionAttributes;
+	
 	PV::Facades::FFoliageFacade FoliageFacadeOutput(OutCollection);
 	PV::Facades::FFoliageFacade FoliagePaletteFacade(FoliageCollection);
  
@@ -166,74 +114,17 @@ void FPVFoliage::DistributeFoliage(FManagedArrayCollection& OutCollection, const
 	if (DistributionSettings.OverrideDistribution)
 	{
 		TRACE_CPUPROFILER_EVENT_SCOPE(PVFoliageDistributorNode::OverrideDistribution);
- 
-		TArray<float> NormalisedPScales;
-		NormalisedPScales.Init(0, PointFacadeSource.GetElementCount());
- 
-		TArray<float> NormalisedBudLightDetected;
-		NormalisedBudLightDetected.Init(0, PointFacadeSource.GetElementCount());
- 
-		TArray<float> NormalisedHealth;
-		NormalisedHealth.Init(0, PointFacadeSource.GetElementCount());
- 
-		TArray<PV::Facades::FFoliageEntryData> FoliageEntries;
- 
-		// TODO : Optimize
-		// if (DistributionConditions.IsActiveCondition(EFoliageDistributionCondition::Scale) || DistributionConditions.IsActiveCondition(EFoliageDistributionCondition::Light))
-		// {
-		float MinPScale = std::numeric_limits<float>::max();
-		float MaxPScale = std::numeric_limits<float>::min();
-			
-		float MinBudLightDetected = std::numeric_limits<float>::max();
-		float MaxBudLightDetected = std::numeric_limits<float>::min();
- 
-		float MinHealth = std::numeric_limits<float>::max();
-		float MaxHealth = std::numeric_limits<float>::min();
- 
-		for (int PointIndex = 0; PointIndex < PointFacadeSource.GetElementCount(); PointIndex++)
-		{
-			// Scale
-			float PScale = PointFacadeSource.GetPointScale(PointIndex);
-			MinPScale = FMath::Min(MinPScale, PScale);	
-			MaxPScale = FMath::Max(MaxPScale, PScale);
- 
-			// Light
-			const TArray<float>& CurrentBudLightDetected = PointFacadeSource.GetBudLightDetected(PointIndex);
-			const float CurrentLightAvailable = CurrentBudLightDetected.IsValidIndex(PV::Facades::BudLightDetectedAvailableIndex) ? CurrentBudLightDetected[PV::Facades::BudLightDetectedAvailableIndex] : 0.0f;
-				
-			MinBudLightDetected = FMath::Min(MinBudLightDetected, CurrentLightAvailable);	
-			MaxBudLightDetected = FMath::Max(MaxBudLightDetected, CurrentLightAvailable);
- 
-			// Health
-			float Ethylene = BudHormoneLevels[PointIndex][4];
-			int BudAgeSenescence = BudDevelopment[PointIndex][3];
-			int BudLightSenescence = BudDevelopment[PointIndex][4];
-			
-			int Senescence = FMath::Max(BudAgeSenescence,BudLightSenescence) + 1;
-			NormalisedHealth[PointIndex] = 1 - (Ethylene * (1/Senescence));
-			
-			MinHealth = FMath::Min(MinHealth, NormalisedHealth[PointIndex]);
-			MaxHealth = FMath::Max(MaxHealth, NormalisedHealth[PointIndex]);
-		}
- 
-		for (int PointIndex = 0; PointIndex < PointFacadeSource.GetElementCount(); PointIndex++)
-		{
-			// Scale
-			float PScale = PointFacadeSource.GetPointScale(PointIndex);
-			NormalisedPScales[PointIndex] = FMath::GetMappedRangeValueClamped(FVector2f(MinPScale, MaxPScale),FVector2f(0, 1),PScale);
- 
-			// UE_LOG(LogTemp, Log, TEXT("Normalized PScale : Index : %i : Value : %f"),
-			// 					PointIndex, NormalisedPScales[PointIndex]);
- 
-			// Light
-			const TArray<float>& CurrentBudLightDetected = PointFacadeSource.GetBudLightDetected(PointIndex);
-			const float CurrentLightAvailable = CurrentBudLightDetected.IsValidIndex(PV::Facades::BudLightDetectedAvailableIndex) ? CurrentBudLightDetected[PV::Facades::BudLightDetectedAvailableIndex] : 0.0f;
-			NormalisedBudLightDetected[PointIndex] = FMath::GetMappedRangeValueClamped(FVector2f(MinBudLightDetected, MaxBudLightDetected),FVector2f(0, 1), CurrentLightAvailable);
- 
-			// Health
-			NormalisedHealth[PointIndex] = FMath::GetMappedRangeValueClamped(FVector2f(MinHealth, MaxHealth),FVector2f(0, 1),NormalisedHealth[PointIndex]);
-		}
-		// }
+		
+		FNormalizedPointCaches NormalizedPointCaches;
+		FBuildOptions BuildOptions;
+		BuildOptions.bComputePScales = DistributionConditions.IsActiveCondition(EFoliageDistributionCondition::Scale);
+		BuildOptions.bComputeLight = DistributionConditions.IsActiveCondition(EFoliageDistributionCondition::Light);
+		BuildOptions.bComputeHealth = DistributionConditions.IsActiveCondition(EFoliageDistributionCondition::Health);
+
+		BuildNormalizedPointCaches(
+			SourceCollection,
+			BuildOptions,
+			NormalizedPointCaches);
 		
 		// Compute attachment points in terms of normalized distances
 		OutCollection.EmptyGroup(PV::GroupNames::FoliageGroup);
@@ -256,19 +147,29 @@ void FPVFoliage::DistributeFoliage(FManagedArrayCollection& OutCollection, const
 			return;
 		}
 		
+		const int32 NumFoliageInfo = FoliagePaletteFacade.NumFoliageInfo();
+		const TArray<FPVFoliageInfo>& FoliageInfos = FoliagePaletteFacade.GetFoliageInfos();
+
+		auto IsValidCandidate = [&](int32 CandidateIndex)
+			{
+				return FoliageInfos.IsValidIndex(CandidateIndex) && FoliageInfos[CandidateIndex].IsValid();
+			};
+
+		auto GetCandidateValue = [&](int32 CandidateIndex, EFoliageDistributionCondition Condition) -> float
+			{
+				return FoliageInfos[CandidateIndex].Attributes.GetAttributeValue(Condition);
+			};
+		
 		for (int32 BranchIndex = 0; BranchIndex < NumBranches; ++BranchIndex)
 		{
-			ConditionAttributes ConditionAttributes;
+			FConditionAttributes ConditionAttrs;
 			TArray<float> NormalizedAttachmentPoints;
-			TArray<float> NormalizedAttachmentPointsFinal;
 			TArray<FVector3f> AttachmentPointsPositions;
 			TArray<float> AttachmentPointsLengthsFromRoot;
 			TArray<FVector3f> AttachmentPointsUpVectors;
 			TArray<FVector3f> AttachmentPointsNormalVectors;
 			TArray<float> AttachmentPointsScales;
 			
-			const int32 NumFoliageInfo = FoliagePaletteFacade.NumFoliageInfo();
-			
 			TArray<int32> BranchPoints = BranchFacadeSource.GetPoints(BranchIndex);
  
 			if (BranchPoints.Num() < 2)
@@ -279,6 +180,12 @@ void FPVFoliage::DistributeFoliage(FManagedArrayCollection& OutCollection, const
 			const float FirstPointLengthFromRoot = PointFacadeSource.GetLengthFromRoot(BranchPoints[0]);
 			const float LastPointLengthFromRoot = PointFacadeSource.GetLengthFromRoot(BranchPoints.Last());
 			const float BranchLength = LastPointLengthFromRoot - FirstPointLengthFromRoot;
+			
+			if (BranchLength < KINDA_SMALL_NUMBER)
+			{
+				continue;
+			}
+			
 			float Increment = DistributionSettings.InstanceSpacing / BranchLength;
 			const float AdjustedMaxPerBranch = DistributionSettings.MaxPerBranch == -1
 				? BranchPoints.Num()
@@ -339,21 +246,23 @@ void FPVFoliage::DistributeFoliage(FManagedArrayCollection& OutCollection, const
  
 					const float CurrentPointLengthFromRoot = PointFacadeSource.GetLengthFromRoot(CurrentPointIndex);
 					const float NextPointLengthFromRoot = PointFacadeSource.GetLengthFromRoot(NextPointIndex);
-					
+
 					const float CurrentPointLfrNormalized = FMath::GetMappedRangeValueClamped(
 						FVector2f(FirstPointLengthFromRoot, LastPointLengthFromRoot),
 						FVector2f(0.0f, 1.0f), CurrentPointLengthFromRoot);
- 
+
 					const float NextPointLfrNormalized = FMath::GetMappedRangeValueClamped(
 						FVector2f(FirstPointLengthFromRoot, LastPointLengthFromRoot),
 						FVector2f(0.0f, 1.0f),
 						NextPointLengthFromRoot);
- 
+
 					if (NormalizedAttachmentPoints[Idx] >= CurrentPointLfrNormalized && NormalizedAttachmentPoints[Idx] <=
 						NextPointLfrNormalized)
 					{
-						const float DistanceAlpha = (NormalizedAttachmentPoints[Idx] -
-							CurrentPointLfrNormalized) / (NextPointLfrNormalized - CurrentPointLfrNormalized);
+						const float DistanceAlpha = (NextPointLfrNormalized - CurrentPointLfrNormalized) < KINDA_SMALL_NUMBER
+							? 1.0f
+							: (NormalizedAttachmentPoints[Idx] -
+								CurrentPointLfrNormalized) / (NextPointLfrNormalized - CurrentPointLfrNormalized);
  
 						const float EthyleneLevel = FMath::Lerp(EthyleneLevelAtCurrentPoint, EthyleneLevelAtNextPoint, DistanceAlpha);
  
@@ -406,9 +315,8 @@ void FPVFoliage::DistributeFoliage(FManagedArrayCollection& OutCollection, const
 							NextPointApicalDirectionVector, DistanceAlpha);
  
 						// Condition - Y Dot
-						float YDot = FVector3f::DotProduct(BudApicalDirectionVector, FVector3f::UpVector);
-						YDot = FMath::GetMappedRangeValueClamped(FVector2f(-1, 1),FVector2f(0, 1),YDot);
-						ConditionAttributes.UpAlignment.Add(YDot);
+						float YDot = UpAlignmentNormalized(BudApicalDirectionVector);
+						ConditionAttrs.UpAlignment.Add(YDot);
  
 						// Condition - Light
 						// const TArray<float>& CurrentBudLightDetected = PointFacadeSource.GetBudLightDetected(CurrentPointIndex);
@@ -417,36 +325,30 @@ void FPVFoliage::DistributeFoliage(FManagedArrayCollection& OutCollection, const
 						// const TArray<float>& NextBudLightDetected = PointFacadeSource.GetBudLightDetected(NextPointIndex);
 						// const float NextLightAvailable = NextBudLightDetected.IsValidIndex(PV::Facades::BudLightDetectedAvailableIndex) ? NextBudLightDetected[PV::Facades::BudLightDetectedAvailableIndex] : 0.0f;	
  
-						const float CurrentLightAvailable = NormalisedBudLightDetected[CurrentPointIndex];
-						const float NextLightAvailable = NormalisedBudLightDetected[NextPointIndex];
+						const float CurrentLightAvailable = NormalizedPointCaches.Light[CurrentPointIndex];
+						const float NextLightAvailable = NormalizedPointCaches.Light[NextPointIndex];
 						const float BudLightReceived = FMath::Lerp(CurrentLightAvailable,NextLightAvailable, DistanceAlpha);
-						ConditionAttributes.Light.Add(BudLightReceived);
+						ConditionAttrs.Light.Add(BudLightReceived);
  
 						// Condition - Scale
 						// const float CurrentPointScaleGradient = PointFacadeSource.GetPointScaleGradient(CurrentPointIndex);
 						// const float NextPointScaleGradient = PointFacadeSource.GetPointScaleGradient(NextPointIndex);
 						// const float PointScaleGradient = FMath::Lerp(CurrentPointScaleGradient, NextPointScaleGradient, DistanceAlpha);
  
-						const float CurrentPScale = NormalisedPScales[CurrentPointIndex];
-						const float NextPScale = NormalisedPScales[NextPointIndex];
+						const float CurrentPScale = NormalizedPointCaches.PScales[CurrentPointIndex];
+						const float NextPScale = NormalizedPointCaches.PScales[NextPointIndex];
 						const float PointPScale = FMath::Lerp(CurrentPScale, NextPScale, DistanceAlpha);
-						ConditionAttributes.Scale.Add(PointPScale);
+						ConditionAttrs.Scale.Add(PointPScale);
  
 						// Health
-						const float CurrentHealth = NormalisedHealth[CurrentPointIndex];
-						const float NextHealth = NormalisedHealth[NextPointIndex];
+						const float CurrentHealth = NormalizedPointCaches.Health[CurrentPointIndex];
+						const float NextHealth = NormalizedPointCaches.Health[NextPointIndex];
 						const float BudHealth = FMath::Lerp(CurrentHealth,NextHealth, DistanceAlpha);
-						ConditionAttributes.Health.Add(BudHealth);
+						ConditionAttrs.Health.Add(BudHealth);
  
 						// Condition - Tip
-						if (BranchPointIndex == BranchPoints.Num() - 2 && DistanceAlpha == 1)
-						{
-							ConditionAttributes.Tip.Add(1);
-						}
-						else
-						{
-							ConditionAttributes.Tip.Add(0);
-						}
+						const bool bIsTip = (BranchPointIndex == BranchPoints.Num() - 2) && FMath::IsNearlyEqual(DistanceAlpha, 1.0f, 1e-3f);
+						ConditionAttrs.Tip.Add(Tip(bIsTip));
  
 						int32 BudsToAdd = 1;
  
@@ -510,20 +412,20 @@ void FPVFoliage::DistributeFoliage(FManagedArrayCollection& OutCollection, const
 										AttachmentPointsLengthsFromRoot.Add(LastLengthFromRootAdded);
  
 										// Conditions
-										float LastConditionScale = ConditionAttributes.Scale.Last();
-										ConditionAttributes.Scale.Add(LastConditionScale);
+										float LastConditionScale = ConditionAttrs.Scale.Last();
+										ConditionAttrs.Scale.Add(LastConditionScale);
 									
-										float LastConditionLight = ConditionAttributes.Light.Last();
-										ConditionAttributes.Light.Add(LastConditionLight);
+										float LastConditionLight = ConditionAttrs.Light.Last();
+										ConditionAttrs.Light.Add(LastConditionLight);
 									
-										float LastConditionUpAlignment  = ConditionAttributes.UpAlignment.Last();
-										ConditionAttributes.UpAlignment.Add(LastConditionUpAlignment);
+										float LastConditionUpAlignment  = ConditionAttrs.UpAlignment.Last();
+										ConditionAttrs.UpAlignment.Add(LastConditionUpAlignment);
  
-										float LastConditionTip = ConditionAttributes.Tip.Last();
-										ConditionAttributes.Tip.Add(LastConditionTip);
+										float LastConditionTip = ConditionAttrs.Tip.Last();
+										ConditionAttrs.Tip.Add(LastConditionTip);
  
-										float LastConditionHealth = ConditionAttributes.Health.Last();
-										ConditionAttributes.Health.Add(LastConditionHealth);
+										float LastConditionHealth = ConditionAttrs.Health.Last();
+										ConditionAttrs.Health.Add(LastConditionHealth);
 									}
 								}
  
@@ -552,116 +454,29 @@ void FPVFoliage::DistributeFoliage(FManagedArrayCollection& OutCollection, const
 			// Set foliage entries
 			TManagedArrayAccessor<FVector3f> PivotPointsAttribute(OutCollection, PV::AttributeNames::FoliagePivotPoint,
 				PV::GroupNames::FoliageGroup);
- 
-			const TArray<FPVFoliageInfo> FoliageInfos = FoliagePaletteFacade.GetFoliageInfos();
-			
- 
-			const int32 BranchNumber = BranchFacadeSource.GetBranchNumber(BranchIndex);
- 
-			FRandomStream FoliageRandomStream;
-			FoliageRandomStream.Initialize(BranchNumber);
 			
 			for (int32 Idx = 0; Idx < AttachmentPointsLengthsFromRoot.Num(); ++Idx)
 			{
-				TArray<FFoliageWeightInfo> AvailableFoliageWeights;
-				
-				for (int i = 0; i < NumFoliageInfo; ++i)
-				{
-					auto FoliageInfo = FoliageInfos[i];
- 
-					if (FoliageInfo.IsValid())
-					{
-						AvailableFoliageWeights.Add(FFoliageWeightInfo(i, 0));
-					}
-				}
- 
-				// Shuffle foliage
-				for (int i = AvailableFoliageWeights.Num() - 1; i > 0; --i)
-				{
-					int Index = RandomStream.RandRange(0, i);
-					AvailableFoliageWeights.Swap(i, Index);
-				}
- 
-				if(DistributionConditions.HasActiveCondition())
-				{
-					for (EFoliageDistributionCondition Condition : TEnumRange<EFoliageDistributionCondition>())
-					{
-						FPVFoliageDistributionConditionSettings OutConditionSettings;
-						
-						if (DistributionConditions.IsActiveCondition(Condition) && DistributionConditions.GetSettings(Condition, OutConditionSettings))
-						{
-							float PointAttributeValue = ConditionAttributes.GetConditionValue(Condition,Idx);
-					
-							const float AttributeValue = PointAttributeValue + OutConditionSettings.Offset;
- 
-							for (int Index = 0; Index < AvailableFoliageWeights.Num(); ++Index)
-							{
-								int FoliageIndex = AvailableFoliageWeights[Index].FoliageIndex;
- 
-								if (FoliageInfos.IsValidIndex(FoliageIndex))
-								{
-									auto FoliageInfo = FoliageInfos[FoliageIndex];
-									float FoliageValue = FoliageInfo.Attributes.GetAttributeValue(Condition);
-									float FoliageDistribution = FMath::Abs(FoliageValue - AttributeValue) * OutConditionSettings.Weight;
-									
-									// UE_LOG(LogProceduralVegetation, Log, TEXT("FoliageWeight : Index : %i : FoliageValue : %f : AttributeValue : %f : Weight : %f : Final Value: %f"),
-									// 	FoliageIndex,
-									// 	FoliageValue,
-									// 	AttributeValue,
-									// 	OutConditionSettings.Weight,
-									// 	FoliageDistribution);	
- 
-									AvailableFoliageWeights[Index].Weight += FoliageDistribution;
-								}
-								else
-								{
-									check(false);
-								}
-							}	
-						}
-					}
- 
-					float MinFoliageWeight = Algo::MinElementBy(AvailableFoliageWeights, [](const FFoliageWeightInfo& Info) { return Info.Weight; })->Weight;
-					float MaxFoliageWeight = Algo::MaxElementBy(AvailableFoliageWeights, [](const FFoliageWeightInfo& Info) { return Info.Weight; })->Weight;
- 
-					for (FFoliageWeightInfo& Info : AvailableFoliageWeights)
-					{
-						Info.Weight = FMath::GetMappedRangeValueClamped(
-									FVector2f(MinFoliageWeight, MaxFoliageWeight),
-									FVector2f(0, 1),
-									Info.Weight);
-					}
-				}
- 
-				// Pick foliage
-				AvailableFoliageWeights.Sort([&](const FFoliageWeightInfo& A, const FFoliageWeightInfo& B)
-					{
-						return A.Weight < B.Weight; // Ascending
-					});
- 
-				
-				TArray<int> PickedFoliage;
- 
-				
-				for (const FFoliageWeightInfo& Info : AvailableFoliageWeights)
+				ConditionPicker::FPointConditionSample Sample;
+				Sample.Light = ConditionAttrs.Light[Idx];
+				Sample.Scale = ConditionAttrs.Scale[Idx];
+				Sample.UpAlignment = ConditionAttrs.UpAlignment[Idx];
+				Sample.Tip = ConditionAttrs.Tip[Idx];
+				Sample.Health = ConditionAttrs.Health[Idx];
+
+				const int32 PickedFoliageIndex = ConditionPicker::PickCandidateIndex(
+					NumFoliageInfo,
+					DistributionConditions,
+					Sample,
+					RandomStream,
+					GetCandidateValue,
+					IsValidCandidate);
+
+				if (PickedFoliageIndex == INDEX_NONE)
 				{
-					if (Info.Weight < DistributionConditions.CutoffThreshold)
-					{
-						PickedFoliage.Add(Info.FoliageIndex);
-					}
-					else if (PickedFoliage.Num() < DistributionConditions.MinimumCandidates)
-					{
-						PickedFoliage.Add(Info.FoliageIndex);
-					}
-					else
-					{
-						break;
-					}
+					continue;
 				}
  
-				const int PickedIndex = RandomStream.RandRange(0, PickedFoliage.Num() - 1);
-				int PickedFoliageIndex = PickedFoliage[PickedIndex];
- 
 				const int32 NewIndex = PivotPointsAttribute.AddElements(1);
 				FoliageFacadeOutput.SetFoliageEntry(NewIndex, {
 					.NameId = PickedFoliageIndex,
@@ -671,11 +486,11 @@ void FPVFoliage::DistributeFoliage(FManagedArrayCollection& OutCollection, const
 					.NormalVector = AttachmentPointsNormalVectors[Idx],
 					.Scale = AttachmentPointsScales[Idx],
 					.LengthFromRoot = AttachmentPointsLengthsFromRoot[Idx],
-					.ConditionUpAlignment = ConditionAttributes.UpAlignment[Idx],
-					.ConditionTip = ConditionAttributes.Tip[Idx],
-					.ConditionLight = ConditionAttributes.Light[Idx],
-					.ConditionScale = ConditionAttributes.Scale[Idx],
-					.ConditionHealth = ConditionAttributes.Health[Idx]
+					.ConditionUpAlignment = ConditionAttrs.UpAlignment[Idx],
+					.ConditionTip = ConditionAttrs.Tip[Idx],
+					.ConditionLight = ConditionAttrs.Light[Idx],
+					.ConditionScale = ConditionAttrs.Scale[Idx],
+					.ConditionHealth = ConditionAttrs.Health[Idx]
 				});
 			}
 		}
diff --git a/Engine/Plugins/Experimental/ProceduralVegetationEditor/Source/ProceduralVegetation/Private/Implementations/PVFoliage.h b/Engine/Plugins/Experimental/ProceduralVegetationEditor/Source/ProceduralVegetation/Private/Implementations/PVFoliage.h
--- a/Engine/Plugins/Experimental/ProceduralVegetationEditor/Source/ProceduralVegetation/Private/Implementations/PVFoliage.h
+++ b/Engine/Plugins/Experimental/ProceduralVegetationEditor/Source/ProceduralVegetation/Private/Implementations/PVFoliage.h
@@ -3,12 +3,9 @@
 #pragma once
  
 #include "CoreMinimal.h"
-#include "Misc/EnumRange.h"
- 
 #include "GeometryCollection/ManagedArrayCollection.h"
 
-#include "PVFoliage.generated.h"
-
+struct FPVFoliageDistributionConditions;
 enum class EPhyllotaxyFormation : uint8;
 enum class EPhyllotaxyType : uint8;
 struct FPVFloatRamp;
@@ -52,284 +49,13 @@ struct FPhyllotaxySettings
 	const int32 MaximumNodeBuds;
 	const float PhyllotaxyAdditionalAngle;
 };
- 
-USTRUCT()
-struct FPVFoliageDistributionConditionSettings
-{
-	GENERATED_BODY()
- 
-	UPROPERTY(EditAnywhere, Category="ConditionSettings", meta=(ClampMin=0.0f, ClampMax=1.0f, UIMin=0.0f, UIMax=1.0f, Tooltip=""))
-	float Weight = 1.0f;
- 
-	UPROPERTY(EditAnywhere, Category="ConditionSettings", meta=(ClampMin=-1.0f, ClampMax=1.0f, UIMin=-1.0f, UIMax=1.0f, Tooltip=""))
-	float Offset = 0.0f;
- 
-	FString ToString() const
-	{
-		return FString::Printf(TEXT("Weight = %f, Offset = %f"), Weight, Offset);
-	}
-};
- 
-UENUM()
-enum class EFoliageDistributionCondition : uint8
-{
-	Light,
-	Scale,
-	UpAlignment,
-	Tip,
-	Health,
-	Count UMETA(Hidden),
-	None UMETA(Hidden),
-};
- 
-ENUM_RANGE_BY_COUNT(EFoliageDistributionCondition, EFoliageDistributionCondition::Count)
- 
- 
-USTRUCT()
-struct FPVFoliageDistributionConditions
-{
-	GENERATED_BODY()
- 
-	UPROPERTY(EditAnywhere, Category="ConditionSettings", meta=(ClampMin=0.0f, ClampMax=1.0f, UIMin=0.0f, UIMax=1.0f, Tooltip=""))
-	float CutoffThreshold = 0.3f;
- 
-	UPROPERTY(EditAnywhere, Category="ConditionSettings", meta=(ClampMin=1, ClampMax=10, UIMin=0, UIMax=10, Tooltip=""))
-	int32 MinimumCandidates = 1;
-	
-	UPROPERTY(EditAnywhere, Category="ConditionSettings", meta=(PinHiddenByDefault, InlineEditConditionToggle))
-	bool bActivateLight = false;
- 
-	UPROPERTY(EditAnywhere, Category="ConditionSettings", meta=(EditCondition="bActivateLight", Tooltip=""))
-	FPVFoliageDistributionConditionSettings Light;
- 
-	UPROPERTY()
-	bool bActivateScale = false;
- 
-	UPROPERTY()
-	FPVFoliageDistributionConditionSettings Scale;
- 
-	UPROPERTY(EditAnywhere, Category="ConditionSettings", meta=(PinHiddenByDefault, InlineEditConditionToggle))
-	bool bActivateUpAlignment = false;
- 
-	UPROPERTY(EditAnywhere, Category="ConditionSettings", meta=(EditCondition="bActivateUpAlignment", Tooltip=""))
-	FPVFoliageDistributionConditionSettings UpAlignment;
- 
-	UPROPERTY(EditAnywhere, Category="ConditionSettings", meta=(PinHiddenByDefault, InlineEditConditionToggle))
-	bool bActivateTip = false;
-	
-	UPROPERTY(EditAnywhere, Category="ConditionSettings", meta=(EditCondition="bActivateTip", Tooltip=""))
-	FPVFoliageDistributionConditionSettings Tip;
- 
-	UPROPERTY(EditAnywhere, Category="ConditionSettings", meta=(PinHiddenByDefault, InlineEditConditionToggle))
-	bool bActivateHealth = false;
-	
-	UPROPERTY(EditAnywhere, Category="ConditionSettings", meta=(EditCondition="bActivateHealth", Tooltip=""))
-	FPVFoliageDistributionConditionSettings Health;
- 
-	bool HasActiveCondition() const
-	{
-		for (EFoliageDistributionCondition Condition : TEnumRange<EFoliageDistributionCondition>())
-		{
-			if (IsActiveCondition(Condition))
-			{
-				return true;
-			}
-		}
-		return false;
-	}
- 
-	bool IsActiveCondition(const EFoliageDistributionCondition InCondition) const 
-	{
-		switch (InCondition)
-		{
-		case EFoliageDistributionCondition::Light:
-			return bActivateLight && Light.Weight > 0.0f;
-			
-		case EFoliageDistributionCondition::Scale:
-			return bActivateScale && Scale.Weight > 0.0f; 
-			
-		case EFoliageDistributionCondition::UpAlignment:
-			return bActivateUpAlignment && UpAlignment.Weight > 0.0f;
-			
-		case EFoliageDistributionCondition::Tip:
-			return bActivateTip && Tip.Weight > 0.0f;
- 
-		case EFoliageDistributionCondition::Health:
-			return bActivateHealth && Health.Weight > 0.0f;
- 
-		default:
-			return false;
-		}
-	}
- 
-	void SetConditionState(const EFoliageDistributionCondition InCondition, const bool InNewState) 
-	{
-		switch (InCondition)
-		{
-		case EFoliageDistributionCondition::Light:
-			bActivateLight = InNewState;
-			break;
-			
-			
-		case EFoliageDistributionCondition::Scale:
-			bActivateScale = InNewState;
-			break;
-			
-		case EFoliageDistributionCondition::UpAlignment:
-			bActivateUpAlignment = InNewState;
-			break;
-			
-		case EFoliageDistributionCondition::Tip:
-			bActivateTip = InNewState;
-			break;
- 
-		case EFoliageDistributionCondition::Health:
-			bActivateHealth = InNewState;
-			break;
- 
-			default:
-			check(false);
-		}
-	}
- 
-	FPVFoliageDistributionConditionSettings* GetSettings(const EFoliageDistributionCondition InCondition)
-	{
-		switch (InCondition)
-		{
-			case EFoliageDistributionCondition::Light:
-				return &Light;
-				
-			case EFoliageDistributionCondition::Scale:
-				return &Scale;
-			
-				
-			case EFoliageDistributionCondition::UpAlignment:
-				return &UpAlignment;
-				
-			case EFoliageDistributionCondition::Tip:
-				return &Tip;
- 
-			case EFoliageDistributionCondition::Health:
-				return &Health;
- 
-			default:
-			return nullptr;
-		}
-	}
- 
-	bool GetSettings(const EFoliageDistributionCondition InCondition, FPVFoliageDistributionConditionSettings& OutSettings) const
-	{
-		switch (InCondition)
-		{
-			case EFoliageDistributionCondition::Light:
-			{
-				OutSettings =  Light;
-				return true;
-			}
-			
-			case EFoliageDistributionCondition::Scale:
-			{
-				OutSettings = Scale;
-				return true;
-			}
-			
-			case EFoliageDistributionCondition::UpAlignment:
-			{
-				OutSettings = UpAlignment;
-				return true;
-			}
-			
-			case EFoliageDistributionCondition::Tip:
-			{
-				OutSettings = Tip;
-				return true;
-			}
- 
-			case EFoliageDistributionCondition::Health:
-			{
-				OutSettings = Health;
-				return true;
-			}
-		}
- 
-		return false;
-	}
-};
+
  
 struct FMiscSettings
 {
 	const int32 RandomSeed;
 };
- 
- 
-USTRUCT()
-struct FPVFoliageAttributes
-{
-	GENERATED_BODY()
- 
-	UPROPERTY(EditAnywhere, Category="Info", meta=(ClampMin=0.0f, ClampMax=1.0f, UIMin=0.0f, UIMax=1.0f, Tooltip=""))
-	float Light = 0.0f;
- 
-	UPROPERTY()
-	float Scale = 0.0f;
- 
-	UPROPERTY(EditAnywhere, Category="Info", meta=(ClampMin=0.0f, ClampMax=1.0f, UIMin=0.0f, UIMax=1.0f, Tooltip=""))
-	float UpAlignment = 0.0f;
- 
-	UPROPERTY(EditAnywhere, Category="Info", meta=(ClampMin=0.0f, ClampMax=1.0f, UIMin=0.0f, UIMax=1.0f, Tooltip=""))
-	float Health = 0.0f;
- 
-	UPROPERTY(EditAnywhere, Category="Info", meta=(Tooltip=""))
-	bool Tip = false;
- 
-	float GetAttributeValue(const EFoliageDistributionCondition InCondition) const
-	{
-		switch (InCondition)
-		{
-			case EFoliageDistributionCondition::Light:
-			return Light;
-			
-			case EFoliageDistributionCondition::Scale:
-			return Scale;
-			
-			case EFoliageDistributionCondition::UpAlignment:
-			return UpAlignment;
-			
-			case EFoliageDistributionCondition::Tip:
-			return Tip ? 1.0f : 0.0f;
- 
-			case EFoliageDistributionCondition::Health:
-			return Health;
- 
-			default:
-			return 0.0f;
-		}
-	}
- 
-	FString ToString() const
-	{
-		return FString::Printf(TEXT("Light = %f, Scale = %f, UpAlignment = %f, Health = %f, Tip = %s"), Light, Scale, UpAlignment, Health, Tip ? TEXT("True") : TEXT("False"));
-	}
-};
- 
-USTRUCT()
-struct FPVFoliageInfo
-{
-	GENERATED_BODY()
-	
-	UPROPERTY(EditAnywhere, Category = "Info")
-	TSoftObjectPtr<UObject> Mesh;
- 
-	UPROPERTY(EditAnywhere, Category="Info", meta=(ShowInnerProperties, FullyExpand="true", Tooltip=""))
-	FPVFoliageAttributes Attributes;
- 
-	bool IsValid() const
-	{
-		// Asset is loaded and pointer is valid
-		// Path is valid, but asset is NOT loaded yet
-		return Mesh.IsValid() || Mesh.ToSoftObjectPath().IsValid(); 
-	}
-};
+
  
 struct FPVFoliage
 {
diff --git a/Engine/Plugins/Experimental/ProceduralVegetationEditor/Source/ProceduralVegetation/Private/Implementations/PVGrafter.cpp b/Engine/Plugins/Experimental/ProceduralVegetationEditor/Source/ProceduralVegetation/Private/Implementations/PVGrafter.cpp
--- a/Engine/Plugins/Experimental/ProceduralVegetationEditor/Source/ProceduralVegetation/Private/Implementations/PVGrafter.cpp
+++ b/Engine/Plugins/Experimental/ProceduralVegetationEditor/Source/ProceduralVegetation/Private/Implementations/PVGrafter.cpp
@@ -8,6 +8,7 @@
 
 #include "Facades/PVAttributesNames.h"
 #include "Facades/PVBranchFacade.h"
+#include "Facades/PVMetaInfoFacade.h"
 #include "Facades/PVPlantFacade.h"
 #include "Facades/PVPointFacade.h"
 #include "Facades/PVTreeFacade.h"
@@ -16,6 +17,8 @@
 #include "Helpers/PVUtilities.h"
 #include "Helpers/PVAttributesHelper.h"
 
+using namespace PV::DistributionConditionUtils;
+using PV::DistributionConditionUtils::FConditionAttributes;
 
 template <typename Fn>
 void FPVGrafter::ForEachGraftPointExcludingRoot(const int32 NumPoints, const int32 RootIndex, Fn&& Func)
@@ -607,7 +610,7 @@ bool FPVGrafter::RemoveFoliageDataIfPresent(FManagedArrayCollection& OutCollecti
 TArray<FManagedArrayCollection> FPVGrafter::ExtractIndividualPlants(const FManagedArrayCollection& InCollection)
 {
 	const PV::Facades::FPlantFacade PlantFacadeSourceGraft(InCollection);
-	
+
 	TArray<int32> PlantNumbersInGraft = PlantFacadeSourceGraft.GetPlantNumbers();
 	TArray<FManagedArrayCollection> IndividualPlants;
 	if (PlantNumbersInGraft.Num() > 1)
@@ -629,7 +632,7 @@ TArray<FManagedArrayCollection> FPVGrafter::ExtractIndividualPlants(const FManag
 					BranchesToRemove.Add(BranchIndex);
 				}
 			}
-			
+
 			PV::Facades::FTreeFacade::RemoveBranches(BranchFacade, BranchesToRemove, NewCollection);
 			IndividualPlants.Add(NewCollection);
 		}
@@ -638,18 +641,19 @@ TArray<FManagedArrayCollection> FPVGrafter::ExtractIndividualPlants(const FManag
 	{
 		IndividualPlants.Add(InCollection);
 	}
-	
+
 	return IndividualPlants;
 }
 
 void FPVGrafter::DistributeGraftWithHormoneBasedSettings(FManagedArrayCollection& OutCollection,
                                                          const FManagedArrayCollection& SourceSkeletonCollection,
-                                                         FManagedArrayCollection& SourceGraftCollection,
+                                                         TArray<FManagedArrayCollection>& SourceGraftCollections,
                                                          const FHormoneBasedSettings::FDistributionSettings& DistributionSettings,
                                                          const FHormoneBasedSettings::FPhyllotaxySettings& PhyllotaxySettings,
                                                          const FHormoneBasedSettings::FScaleSettings& ScaleSettings,
                                                          const FHormoneBasedSettings::FVectorSettings& VectorSettings,
-                                                         const FMiscSettings& MiscSettings)
+                                                         const FMiscSettings& MiscSettings,
+                                                         const FPVFoliageDistributionConditions& DistributionConditions)
 {
 	TRACE_CPUPROFILER_EVENT_SCOPE(PVGrafter::DistributeGraftWithHormoneBasedSettings);
 
@@ -676,9 +680,35 @@ void FPVGrafter::DistributeGraftWithHormoneBasedSettings(FManagedArrayCollection
 		BranchChildrenAttribute,
 		BranchNumberAttribute,
 		BranchPointsAttribute);
-	
+
+	FNormalizedPointCaches NormalizedPointCaches;
+	FBuildOptions BuildOptions;
+	BuildOptions.bComputePScales = DistributionConditions.IsActiveCondition(EFoliageDistributionCondition::Scale);
+	BuildOptions.bComputeLight = DistributionConditions.IsActiveCondition(EFoliageDistributionCondition::Light);
+	BuildOptions.bComputeHealth = DistributionConditions.IsActiveCondition(EFoliageDistributionCondition::Health);
+
+	BuildNormalizedPointCaches(
+		SourceSkeletonCollection,
+		BuildOptions,
+		NormalizedPointCaches);
+
 	// We extract all individual plants from the graft source as separate collections for simplicity
-	TArray<FManagedArrayCollection> GraftSources = ExtractIndividualPlants(SourceGraftCollection);
+	TArray<FManagedArrayCollection> GraftSourcesExtracted;
+	for (auto& Source : SourceGraftCollections)
+	{
+		GraftSourcesExtracted.Append(ExtractIndividualPlants(Source));
+	}
+
+	check(GraftSourcesExtracted.Num() > 0);
+
+	TArray<PV::Facades::FGraftData> CandidateGraftData;
+	CandidateGraftData.SetNum(GraftSourcesExtracted.Num());
+
+	for (int32 i = 0; i < GraftSourcesExtracted.Num(); ++i)
+	{
+		const PV::Facades::FMetaInfoFacade MetaInfoFacade(GraftSourcesExtracted[i]);
+		CandidateGraftData[i] = MetaInfoFacade.GetGraftAttributes();
+	}
 
 	TManagedArrayAccessor<TArray<FVector3f>> BudDirectionAttribute(OutCollection, PV::AttributeNames::BudDirection,
 		PV::GroupNames::PointGroup);
@@ -686,13 +716,13 @@ void FPVGrafter::DistributeGraftWithHormoneBasedSettings(FManagedArrayCollection
 		PV::GroupNames::PointGroup);
 	TManagedArrayAccessor<TArray<float>> BudHormoneLevels(OutCollection, PV::AttributeNames::BudHormoneLevels,
 		PV::GroupNames::PointGroup);
-	
+
 	float PhyllotaxyFormationLeaf = 180;
 	float MinBudsLeaf = 1;
 	float MaxBudsLeaf = 1;
 	bool ResetPhyllotaxyLeaf = PhyllotaxySettings.ResetPhyllotaxy;
 	float PhyllotaxyOffsetLeaf = PhyllotaxySettings.PhyllotaxyOffset;
-	
+
 	switch (PhyllotaxySettings.PhyllotaxyType)
 	{
 	case EPhyllotaxyType::Alternate: PhyllotaxyFormationLeaf = 180 + PhyllotaxySettings.PhyllotaxyAdditionalAngle;
@@ -724,7 +754,7 @@ void FPVGrafter::DistributeGraftWithHormoneBasedSettings(FManagedArrayCollection
 			case EPhyllotaxyFormation::Parastichous: PhyllotaxyFormationLeaf = PhyllotaxySettings.PhyllotaxyAdditionalAngle;
 				break;
 			}
- 
+
 			MinBudsLeaf = MaxBudsLeaf = 1;
 			break;
 		}
@@ -737,9 +767,39 @@ void FPVGrafter::DistributeGraftWithHormoneBasedSettings(FManagedArrayCollection
 	const int32 MaxBudNumber = *Algo::MaxElement(SourceSkeletonBudNumbers.GetConstArray());
 	int32 CurrentBudNumber = MaxBudNumber + 1;
 
+	auto IsValidCandidate = [&](int32 CandidateIndex) -> bool
+		{
+			return CandidateGraftData.IsValidIndex(CandidateIndex);
+		};
+
+	auto GetCandidateValue = [&](int32 CandidateIndex, EFoliageDistributionCondition Condition) -> float
+		{
+			const auto& [Light, UpAlignment, Health, Tip] = CandidateGraftData[CandidateIndex];
+
+			switch (Condition)
+			{
+			case EFoliageDistributionCondition::Light:
+				return Light;
+			case EFoliageDistributionCondition::UpAlignment:
+				return UpAlignment;
+			case EFoliageDistributionCondition::Health:
+				return Health;
+			case EFoliageDistributionCondition::Tip:
+				return Tip
+					? 1.0f
+					: 0.0f;
+			case EFoliageDistributionCondition::Scale:
+				return 0.0f;
+
+			default:
+				return 0.0f;
+			}
+		};
+
 	// Compute attachment points in terms of normalized distances
 	for (int32 BranchIndex = 0; BranchIndex < NumBranches; ++BranchIndex)
 	{
+		FConditionAttributes ConditionAttrs;
 		TArray<float> NormalizedAttachmentPoints;
 		TArray<float> NormalizedAttachmentPointsFinal;
 		TArray<FVector3f> AttachmentPointsPositions;
@@ -758,6 +818,12 @@ void FPVGrafter::DistributeGraftWithHormoneBasedSettings(FManagedArrayCollection
 		const float FirstPointLengthFromRoot = PointFacadeSourceSkeleton.GetLengthFromRoot(BranchPoints[0]);
 		const float LastPointLengthFromRoot = PointFacadeSourceSkeleton.GetLengthFromRoot(BranchPoints.Last());
 		const float BranchLength = LastPointLengthFromRoot - FirstPointLengthFromRoot;
+
+		if (BranchLength <= KINDA_SMALL_NUMBER)
+		{
+			continue;
+		}
+
 		float Increment = DistributionSettings.InstanceSpacing / BranchLength;
 		const float AdjustedMaxPerBranch = DistributionSettings.MaxPerBranch == -1
 			? BranchPoints.Num()
@@ -826,8 +892,10 @@ void FPVGrafter::DistributeGraftWithHormoneBasedSettings(FManagedArrayCollection
 				if (NormalizedAttachmentPoints[Idx] >= CurrentPointLfrNormalized && NormalizedAttachmentPoints[Idx] <=
 					NextPointLfrNormalized)
 				{
-					const float DistanceAlpha = (NormalizedAttachmentPoints[Idx] -
-						CurrentPointLfrNormalized) / (NextPointLfrNormalized - CurrentPointLfrNormalized);
+					const float DistanceAlpha = (NextPointLfrNormalized - CurrentPointLfrNormalized) < KINDA_SMALL_NUMBER
+						? 1.0f
+						: (NormalizedAttachmentPoints[Idx] -
+							CurrentPointLfrNormalized) / (NextPointLfrNormalized - CurrentPointLfrNormalized);
 
 					const float EthyleneLevel = FMath::Lerp(EthyleneLevelAtCurrentPoint, EthyleneLevelAtNextPoint, DistanceAlpha);
 
@@ -901,6 +969,28 @@ void FPVGrafter::DistributeGraftWithHormoneBasedSettings(FManagedArrayCollection
 					const FVector3f BudApicalDirectionVector = FMath::Lerp(CurrentPointApicalDirectionVector,
 						NextPointApicalDirectionVector, DistanceAlpha);
 
+					ConditionAttrs.UpAlignment.Add(UpAlignmentNormalized(BudApicalDirectionVector));
+
+					const float CurrentLightAvailable = NormalizedPointCaches.Light[CurrentPointIndex];
+					const float NextLightAvailable = NormalizedPointCaches.Light[NextPointIndex];
+					const float BudLightReceived = FMath::Lerp(CurrentLightAvailable, NextLightAvailable, DistanceAlpha);
+					ConditionAttrs.Light.Add(BudLightReceived);
+
+					const float CurrentHealth = NormalizedPointCaches.Health[CurrentPointIndex];
+					const float NextHealth = NormalizedPointCaches.Health[NextPointIndex];
+					const float BudHealth = FMath::Lerp(CurrentHealth, NextHealth, DistanceAlpha);
+					ConditionAttrs.Health.Add(BudHealth);
+
+					const float CurrentPScale = NormalizedPointCaches.PScales[CurrentPointIndex];
+					const float NextPScale = NormalizedPointCaches.PScales[NextPointIndex];
+					const float PointPScale = FMath::Lerp(CurrentPScale, NextPScale, DistanceAlpha);
+					ConditionAttrs.Scale.Add(PointPScale);
+
+					const bool bIsTip =
+						(BranchPointIndex == BranchPoints.Num() - 2) &&
+						FMath::IsNearlyEqual(DistanceAlpha, 1.0f, 1e-3f);
+					ConditionAttrs.Tip.Add(Tip(bIsTip));
+
 					int32 BudsToAdd = 1;
 
 					// Compute Up Vector
@@ -963,6 +1053,12 @@ void FPVGrafter::DistributeGraftWithHormoneBasedSettings(FManagedArrayCollection
 								AttachmentPointsPointIndices.Add(LastPointIndexAdded);
 								AttachmentPointsScales.Add(LastScaleAdded);
 								AttachmentPointsLengthsFromRoot.Add(LastLengthFromRootAdded);
+
+								ConditionAttrs.Light.Add(ConditionAttrs.Light.Last());
+								ConditionAttrs.Health.Add(ConditionAttrs.Health.Last());
+								ConditionAttrs.UpAlignment.Add(ConditionAttrs.UpAlignment.Last());
+								ConditionAttrs.Scale.Add(ConditionAttrs.Scale.Last());
+								ConditionAttrs.Tip.Add(ConditionAttrs.Tip.Last());
 							}
 						}
 
@@ -1004,8 +1100,30 @@ void FPVGrafter::DistributeGraftWithHormoneBasedSettings(FManagedArrayCollection
 				.bIsTipInstance = bIsTipInstance
 			};
 
-			check(GraftSources.Num() > 0);
-			const FManagedArrayCollection& PickedSourceGraft = GraftSources[RandomStream.RandRange(0, GraftSources.Num() - 1)];
+			check(GraftSourcesExtracted.Num() > 0);
+
+			ConditionPicker::FPointConditionSample Sample;
+			Sample.Light = ConditionAttrs.Light[Idx];
+			Sample.UpAlignment = ConditionAttrs.UpAlignment[Idx];
+			Sample.Health = ConditionAttrs.Health[Idx];
+			Sample.Tip = ConditionAttrs.Tip[Idx];
+			Sample.Scale = ConditionAttrs.Scale[Idx];
+
+			const int32 PickedCandidateIndex = ConditionPicker::PickCandidateIndex(
+				GraftSourcesExtracted.Num(),
+				DistributionConditions,
+				Sample,
+				RandomStream,
+				GetCandidateValue,
+				IsValidCandidate);
+
+			if (PickedCandidateIndex == INDEX_NONE)
+			{
+				continue;
+			}
+
+			const FManagedArrayCollection& PickedSourceGraft = GraftSourcesExtracted[PickedCandidateIndex];
+
 			SetGraft(OutCollection, SourceSkeletonCollection, PickedSourceGraft, Attributes);
 		}
 	}
@@ -1226,11 +1344,12 @@ void FPVGrafter::ComputeRotationAnglesForGraft(const FVector3f& ApicalDirection,
 
 void FPVGrafter::DistributeGraftWithParametricSettings(FManagedArrayCollection& OutCollection,
                                                        const FManagedArrayCollection& SourceSkeletonCollection,
-                                                       FManagedArrayCollection& SourceGraftCollection,
+                                                       TArray<FManagedArrayCollection>& SourceGraftCollections,
                                                        const FParametricSettings::FSpacingSettings& SpacingSettings,
                                                        const FParametricSettings::FAngleSettings& AngleSettings,
                                                        const FParametricSettings::FScaleSettings& ScaleSettings,
-                                                       const FMiscSettings& MiscSettings)
+                                                       const FMiscSettings& MiscSettings,
+                                                       const FPVFoliageDistributionConditions& DistributionConditions)
 {
 	TRACE_CPUPROFILER_EVENT_SCOPE(PVGrafter::DistributeGraftWithHormoneBasedSettings);
 
@@ -1249,6 +1368,13 @@ void FPVGrafter::DistributeGraftWithParametricSettings(FManagedArrayCollection&
 	auto BranchPointsAttribute = PV::FBranchPointsAttribute::GetAttribute(OutCollection);
 	auto PointPositionAttribute = PV::FPointPositionAttribute::GetAttribute(OutCollection);
 	auto BudDirectionAttribute = PV::FBudDirectionAttribute::GetAttribute(OutCollection);
+	
+	const int32 NumOfBranches = BranchPointsAttribute.Num();
+	check(NumOfBranches > 0);
+	if (NumOfBranches == 0)
+	{
+		return;
+	}
 
 	// Only need to do this temporarily since legacy skeleton has inconsistent data in children
 	// TODO: Remove this once support for legacy skeletons is droppped
@@ -1265,9 +1391,35 @@ void FPVGrafter::DistributeGraftWithParametricSettings(FManagedArrayCollection&
 			BranchNumberAttribute,
 			BranchChildrenAttribute
 		});
-	
+
+	FNormalizedPointCaches NormalizedPointCaches;
+	FBuildOptions BuildOptions;
+	BuildOptions.bComputePScales = DistributionConditions.IsActiveCondition(EFoliageDistributionCondition::Scale);
+	BuildOptions.bComputeLight = DistributionConditions.IsActiveCondition(EFoliageDistributionCondition::Light);
+	BuildOptions.bComputeHealth = DistributionConditions.IsActiveCondition(EFoliageDistributionCondition::Health);
+
+	BuildNormalizedPointCaches(
+		SourceSkeletonCollection,
+		BuildOptions,
+		NormalizedPointCaches);
+
 	// We extract all individual plants from the graft source as separate collections for simplicity
-	TArray<FManagedArrayCollection> GraftSources = ExtractIndividualPlants(SourceGraftCollection);
+	TArray<FManagedArrayCollection> GraftSourcesExtracted;
+	for (auto& Source : SourceGraftCollections)
+	{
+		GraftSourcesExtracted.Append(ExtractIndividualPlants(Source));
+	}
+
+	check(GraftSourcesExtracted.Num() > 0);
+
+	TArray<PV::Facades::FGraftData> CandidateGraftData;
+	CandidateGraftData.SetNum(GraftSourcesExtracted.Num());
+
+	for (int32 i = 0; i < GraftSourcesExtracted.Num(); ++i)
+	{
+		const PV::Facades::FMetaInfoFacade MetaInfoFacade(GraftSourcesExtracted[i]);
+		CandidateGraftData[i] = MetaInfoFacade.GetGraftAttributes();
+	}
 
 	const TManagedArray<int32>& SourceSkeletonBudNumbers = PointFacadeSourceSkeleton.GetBudNumbersAttribute();
 	check(SourceSkeletonBudNumbers.Num() > 0);
@@ -1285,11 +1437,45 @@ void FPVGrafter::DistributeGraftWithParametricSettings(FManagedArrayCollection&
 		MaxBranchLength = FMath::Max(BranchLength, MaxBranchLength);
 		AverageBranchLength = AverageBranchLength + BranchLength;
 	}
+	
+	check(MaxBranchLength > KINDA_SMALL_NUMBER);
+	if (MaxBranchLength < KINDA_SMALL_NUMBER)
+	{
+		return;
+	}
 
 	AverageBranchLength = AverageBranchLength / BranchIndicesToBranchLengths.Num();
 	AverageBranchLength = FMath::Lerp(AverageBranchLength, MaxBranchLength, 0.33f);
 
-	const int32 NumOfBranches = BranchParentNumberAttribute.Num();
+	auto IsValidCandidate = [&](int32 CandidateIndex) -> bool
+		{
+			return CandidateGraftData.IsValidIndex(CandidateIndex);
+		};
+
+	auto GetCandidateValue = [&](int32 CandidateIndex, EFoliageDistributionCondition Condition) -> float
+		{
+			const auto& [Light, UpAlignment, Health, Tip] = CandidateGraftData[CandidateIndex];
+
+			switch (Condition)
+			{
+			case EFoliageDistributionCondition::Light:
+				return Light;
+			case EFoliageDistributionCondition::UpAlignment:
+				return UpAlignment;
+			case EFoliageDistributionCondition::Health:
+				return Health;
+			case EFoliageDistributionCondition::Tip:
+				return Tip
+					? 1.0f
+					: 0.0f;
+			case EFoliageDistributionCondition::Scale:
+				return 0.0f;
+
+			default:
+				return 0.0f;
+			}
+		};
+
 	for (int32 BranchIndex = 0; BranchIndex < NumOfBranches; BranchIndex++)
 	{
 		// Check generation limit
@@ -1372,7 +1558,7 @@ void FPVGrafter::DistributeGraftWithParametricSettings(FManagedArrayCollection&
 					PointIndexA,
 					PointIndexB, PointAlpha);
 			}
-			else if (SpacingSettings.SpacingBasis == FParametricSettings::EDistributionBasis::Branch)
+			else // SpacingBasis is Branch
 			{
 				BranchGradient = SpacingSettings.SpacingRamp->GetRichCurveConst()->Eval(BaseValue);
 				BranchGradient = FMath::GetMappedRangeValueClamped(
@@ -1400,22 +1586,27 @@ void FPVGrafter::DistributeGraftWithParametricSettings(FManagedArrayCollection&
 					FVector2f(0.0f, 1.0f),
 					FVector2f(LFRofRoot, LFRofTip),
 					BranchGradient);
-
-				const float PlantGradient0 = 1.0f - PointPlantGradientAttribute[PointIndexA];
-				const float PlantGradient1 = 1.0f - PointPlantGradientAttribute[PointIndexB];
-				PlantGradient = FMath::Lerp(PlantGradient0, PlantGradient1, PointAlpha);
 			}
 
-			float ScaleRampValue = ScaleSettings.ScaleRamp->GetRichCurveConst()->Eval(BranchGradient);
+			const float PlantGradient0 = 1.0f - PointPlantGradientAttribute[PointIndexA];
+			const float PlantGradient1 = 1.0f - PointPlantGradientAttribute[PointIndexB];
+			PlantGradient = FMath::Lerp(PlantGradient0, PlantGradient1, PointAlpha);
+
+			const float ScaleRampLookUp =
+				(ScaleSettings.ScaleRampBasis == FParametricSettings::EDistributionBasis::Branch)
+				? BranchGradient
+				: PlantGradient;
 
-			if (SpacingSettings.SpacingBasis == FParametricSettings::EDistributionBasis::Branch)
+			float ScaleRampValue = ScaleSettings.ScaleRamp->GetRichCurveConst()->Eval(ScaleRampLookUp);
+
+			if (ScaleSettings.ScaleRampBasis == FParametricSettings::EDistributionBasis::Branch)
 			{
 				ScaleRampValue = FMath::GetMappedRangeValueClamped(
 					FVector2f(0.0f, 1.0f),
 					FVector2f(0.0f, LengthRatio),
 					ScaleRampValue);
 			}
-			
+
 			const float RandomScaleFactor = RandomStream.FRandRange(
 				ScaleSettings.RandomizeScaleAngleMinimum,
 				ScaleSettings.RandomizeScaleAngleMaximum);
@@ -1432,7 +1623,7 @@ void FPVGrafter::DistributeGraftWithParametricSettings(FManagedArrayCollection&
 
 			const float AxilRampValue = AngleSettings.AxilAngleRamp->GetRichCurveConst()->Eval(BaseValue);
 			float AxilAngleInDegrees = AngleSettings.AxilAngle * AxilRampValue;
-			
+
 			AxilAngleInDegrees += RandomStream.FRandRange(AngleSettings.RandomizeAxilAngleMinimum, AngleSettings.RandomizeAxilAngleMaximum);
 			AxilAngleInDegrees = FMath::Clamp(AxilAngleInDegrees, -90.0f, 90.0f);
 
@@ -1503,8 +1694,40 @@ void FPVGrafter::DistributeGraftWithParametricSettings(FManagedArrayCollection&
 				.bIsTipInstance = bIsTipInstance
 			};
 
-			check(GraftSources.Num() > 0);
-			const FManagedArrayCollection& PickedSourceGraft = GraftSources[RandomStream.RandRange(0, GraftSources.Num() - 1)]; 
+			ConditionPicker::FPointConditionSample Sample;
+			Sample.UpAlignment = UpAlignmentNormalized(ApicalDirection);
+			Sample.Tip = Tip(bIsTipInstance);
+
+			if (PointIndexA != INDEX_NONE && PointIndexB != INDEX_NONE)
+			{
+				Sample.Light = FMath::Lerp(NormalizedPointCaches.Light[PointIndexA], NormalizedPointCaches.Light[PointIndexB], PointAlpha);
+				Sample.Health = FMath::Lerp(NormalizedPointCaches.Health[PointIndexA], NormalizedPointCaches.Health[PointIndexB], PointAlpha);
+				Sample.Scale = FMath::Lerp(NormalizedPointCaches.PScales[PointIndexA], NormalizedPointCaches.PScales[PointIndexB], PointAlpha);
+			}
+			else
+			{
+				Sample.Light = 0.0f;
+				Sample.Health = 0.0f;
+				Sample.Scale = 0.0f;
+			}
+
+			check(GraftSourcesExtracted.Num() > 0);
+
+			const int32 PickedCandidateIndex = ConditionPicker::PickCandidateIndex(
+				GraftSourcesExtracted.Num(),
+				DistributionConditions,
+				Sample,
+				RandomStream,
+				GetCandidateValue,
+				IsValidCandidate);
+
+			if (PickedCandidateIndex == INDEX_NONE)
+			{
+				continue;
+			}
+
+			const FManagedArrayCollection& PickedSourceGraft = GraftSourcesExtracted[PickedCandidateIndex];
+
 			SetGraft(OutCollection, SourceSkeletonCollection, PickedSourceGraft, Attributes);
 		}
 	}
diff --git a/Engine/Plugins/Experimental/ProceduralVegetationEditor/Source/ProceduralVegetation/Private/Implementations/PVGrafter.h b/Engine/Plugins/Experimental/ProceduralVegetationEditor/Source/ProceduralVegetation/Private/Implementations/PVGrafter.h
--- a/Engine/Plugins/Experimental/ProceduralVegetationEditor/Source/ProceduralVegetation/Private/Implementations/PVGrafter.h
+++ b/Engine/Plugins/Experimental/ProceduralVegetationEditor/Source/ProceduralVegetation/Private/Implementations/PVGrafter.h
@@ -7,6 +7,7 @@
 #include "GeometryCollection/ManagedArrayCollection.h"
 #include "Helpers/PVPhyllotaxyHelper.h"
 
+struct FPVFoliageDistributionConditions;
 struct FPVFloatRamp;
 
 struct FPVGrafter
@@ -137,12 +138,13 @@ struct FPVGrafter
 
 	static void DistributeGraftWithHormoneBasedSettings(FManagedArrayCollection& OutCollection,
 	                                                    const FManagedArrayCollection& SourceSkeletonCollection,
-	                                                    FManagedArrayCollection& SourceGraftCollection,
+	                                                    TArray<FManagedArrayCollection>& SourceGraftCollections,
 	                                                    const FHormoneBasedSettings::FDistributionSettings& DistributionSettings,
 	                                                    const FHormoneBasedSettings::FPhyllotaxySettings& PhyllotaxySettings,
 	                                                    const FHormoneBasedSettings::FScaleSettings& ScaleSettings,
 	                                                    const FHormoneBasedSettings::FVectorSettings& VectorSettings,
-	                                                    const FMiscSettings& MiscSettings);
+	                                                    const FMiscSettings& MiscSettings,
+	                                                    const FPVFoliageDistributionConditions& DistributionConditions);
 
 	static TMap<int32, float> ComputeBranchLengths(
 		PV::FPointLengthFromRootAttributeConstView PointLengthFromRootAttribute,
@@ -180,9 +182,10 @@ struct FPVGrafter
 
 	static void DistributeGraftWithParametricSettings(FManagedArrayCollection& OutCollection,
 	                                                  const FManagedArrayCollection& SourceSkeletonCollection,
-	                                                  FManagedArrayCollection& SourceGraftCollection,
+	                                                  TArray<FManagedArrayCollection>& SourceGraftCollections,
 	                                                  const FParametricSettings::FSpacingSettings& SpacingSettings,
 	                                                  const FParametricSettings::FAngleSettings& AngleSettings,
 	                                                  const FParametricSettings::FScaleSettings& ScaleSettings,
-	                                                  const FMiscSettings& MiscSettings);
+	                                                  const FMiscSettings& MiscSettings,
+	                                                  const FPVFoliageDistributionConditions& DistributionConditions);
 };
diff --git a/Engine/Plugins/Experimental/ProceduralVegetationEditor/Source/ProceduralVegetation/Private/Nodes/PVFoliageDistributorSettings.h b/Engine/Plugins/Experimental/ProceduralVegetationEditor/Source/ProceduralVegetation/Private/Nodes/PVFoliageDistributorSettings.h
--- a/Engine/Plugins/Experimental/ProceduralVegetationEditor/Source/ProceduralVegetation/Private/Nodes/PVFoliageDistributorSettings.h
+++ b/Engine/Plugins/Experimental/ProceduralVegetationEditor/Source/ProceduralVegetation/Private/Nodes/PVFoliageDistributorSettings.h
@@ -5,6 +5,7 @@
 #include "Utils/PVFloatRamp.h"
  
 #include "DataTypes/PVData.h"
+#include "Helpers/PVFoliageHelper.h"
 #include "Helpers/PVPhyllotaxyHelper.h"
 
 #include "Implementations/PVFoliage.h"
diff --git a/Engine/Plugins/Experimental/ProceduralVegetationEditor/Source/ProceduralVegetation/Private/Nodes/PVFoliagePaletteSettings.cpp b/Engine/Plugins/Experimental/ProceduralVegetationEditor/Source/ProceduralVegetation/Private/Nodes/PVFoliagePaletteSettings.cpp
--- a/Engine/Plugins/Experimental/ProceduralVegetationEditor/Source/ProceduralVegetation/Private/Nodes/PVFoliagePaletteSettings.cpp
+++ b/Engine/Plugins/Experimental/ProceduralVegetationEditor/Source/ProceduralVegetation/Private/Nodes/PVFoliagePaletteSettings.cpp
@@ -15,6 +15,7 @@
 #include "Helpers/PCGHelpers.h"
 #include "Helpers/PCGSettingsHelpers.h"
 #include "Helpers/PVAnalyticsHelper.h"
+#include "Helpers/PVFoliageHelper.h"
 #include "Implementations/PVFoliage.h"
  
 #define LOCTEXT_NAMESPACE "PVFoliagePaletteSettings"
diff --git a/Engine/Plugins/Experimental/ProceduralVegetationEditor/Source/ProceduralVegetation/Private/Nodes/PVGrafterPaletteSettings.cpp b/Engine/Plugins/Experimental/ProceduralVegetationEditor/Source/ProceduralVegetation/Private/Nodes/PVGrafterPaletteSettings.cpp
new file mode 100644
--- /dev/null
+++ b/Engine/Plugins/Experimental/ProceduralVegetationEditor/Source/ProceduralVegetation/Private/Nodes/PVGrafterPaletteSettings.cpp
@@ -0,0 +1,124 @@
+// Copyright Epic Games, Inc. All Rights Reserved.
+
+#include "PVGrafterPaletteSettings.h"
+
+#include "ProceduralVegetationModule.h"
+#include "DataTypes/PVGrafterPaletteData.h"
+#include "Data/PCGBasePointData.h"
+#include "DataTypes/PVFoliageData.h"
+#include "DataTypes/PVGrowthData.h"
+#include "Facades/PVMetaInfoFacade.h"
+#include "GeometryCollection/ManagedArrayCollection.h"
+#include "Helpers/PCGSettingsHelpers.h"
+#include "Helpers/PVFoliageHelper.h"
+
+#define LOCTEXT_NAMESPACE "PVGrafterPaletteSettings"
+
+#if WITH_EDITOR
+FText UPVGrafterPaletteSettings::GetDefaultNodeTitle() const
+{
+	return LOCTEXT("NodeTitle", "Grafter Palette");
+}
+
+FText UPVGrafterPaletteSettings::GetNodeTooltipText() const
+{
+	return LOCTEXT("NodeTooltip",
+		"Allows the user to visualize and graft skeletons. "
+		"The user can add or remove graft skeletons."
+		"\n\nPress Ctrl + L to lock/unlock node output"
+	);
+}
+#endif
+
+TArray<FPCGPinProperties> UPVGrafterPaletteSettings::InputPinProperties() const
+{
+	TArray<FPCGPinProperties> Properties;
+
+	if (GrafterInfos.Num() > 0)
+	{
+		for (int32 i = 0; i < GrafterInfos.Num(); i++)
+		{
+			FName PinName = FName(*FString::Printf(TEXT("Graft %d"), i + 1));
+			FPCGPinProperties& Pin = Properties.Emplace_GetRef(PinName, GetInputPinTypeIdentifier());
+			Pin.bAllowMultipleData = false;
+		}
+	}
+
+	return Properties;
+}
+
+FPCGDataTypeIdentifier UPVGrafterPaletteSettings::GetOutputPinTypeIdentifier() const
+{
+	return FPCGDataTypeIdentifier{FPVDataTypeInfoGrafterPalette::AsId()};
+}
+
+FPCGDataTypeIdentifier UPVGrafterPaletteSettings::GetInputPinTypeIdentifier() const
+{
+	return FPCGDataTypeIdentifier{FPVDataTypeInfoGrowth::AsId()};;
+}
+
+FPCGElementPtr UPVGrafterPaletteSettings::CreateElement() const
+{
+	return MakeShared<FPVGrafterPaletteElement>();
+}
+
+bool FPVGrafterPaletteElement::ExecuteInternal(FPCGContext* InContext) const
+{
+	TRACE_CPUPROFILER_EVENT_SCOPE(FPVGrafterPaletteElement::Execute);
+
+	check(InContext);
+
+	const UPVGrafterPaletteSettings* Settings = InContext->GetInputSettings<UPVGrafterPaletteSettings>();
+	check(Settings);
+
+	TArray<TObjectPtr<UPVGrowthData>> GrowthDataElements;
+
+	const TArray<FPCGPinProperties> InputPins = Settings->AllInputPinProperties();
+	const int32 NumOfInputPins = InputPins.Num();
+	for (int32 PinIndex = 0; PinIndex < NumOfInputPins; PinIndex++)
+	{
+		const FName PinName = InputPins[PinIndex].Label;
+		const TArray<FPCGTaggedData> Inputs = InContext->InputData.GetInputsByPin(PinName);
+
+		if (Inputs.IsEmpty())
+		{
+			continue;
+		}
+
+		const UPVGrowthData* InGrowthData = Cast<UPVGrowthData>(Inputs[0].Data);
+		if (!InGrowthData)
+		{
+			PCGLog::InputOutput::LogInvalidInputDataError(InContext);
+			return true;
+		}
+
+		UPVGrowthData* NewGrowthDataElement = FPCGContext::NewObject_AnyThread<UPVGrowthData>(InContext);
+
+		FManagedArrayCollection OutCopy;
+		InGrowthData->GetCollection().CopyTo(&OutCopy);
+
+		if (Settings->GrafterInfos.IsValidIndex(PinIndex))
+		{
+			PV::Facades::FMetaInfoFacade Facade(OutCopy);
+			const auto& Attributes = Settings->GrafterInfos[PinIndex];
+			Facade.SetGraftAttributes({
+				.Light = Attributes.Light,
+				.UpAlignment = Attributes.UpAlignment,
+				.Health = Attributes.Health,
+				.Tip = Attributes.Tip
+			});
+		}
+
+		NewGrowthDataElement->Initialize(MoveTemp(OutCopy));
+		GrowthDataElements.Add(NewGrowthDataElement);
+	}
+
+	UPVGrafterPaletteData* OutPalette = FPCGContext::NewObject_AnyThread<UPVGrafterPaletteData>(InContext);
+	OutPalette->Initialize(MoveTemp(GrowthDataElements));
+
+	InContext->OutputData.TaggedData.Emplace(OutPalette);
+
+	return true;
+}
+
+#undef LOCTEXT_NAMESPACE
diff --git a/Engine/Plugins/Experimental/ProceduralVegetationEditor/Source/ProceduralVegetation/Private/Nodes/PVGrafterPaletteSettings.h b/Engine/Plugins/Experimental/ProceduralVegetationEditor/Source/ProceduralVegetation/Private/Nodes/PVGrafterPaletteSettings.h
new file mode 100644
--- /dev/null
+++ b/Engine/Plugins/Experimental/ProceduralVegetationEditor/Source/ProceduralVegetation/Private/Nodes/PVGrafterPaletteSettings.h
@@ -0,0 +1,44 @@
+// Copyright Epic Games, Inc. All Rights Reserved.
+#pragma once
+ 
+#include "CoreMinimal.h"
+#include "Nodes/PVBaseSettings.h"
+#include "DataTypes/PVData.h"
+#include "PVGrafterPaletteSettings.generated.h"
+
+
+UCLASS(BlueprintType, ClassGroup = (Procedural))
+class UPVGrafterPaletteSettings : public UPVBaseSettings
+{
+	GENERATED_BODY()
+
+public:
+#if WITH_EDITOR
+	virtual FName GetDefaultNodeName() const override { return FName(TEXT("Grafter Palette")); }
+	virtual FText GetDefaultNodeTitle() const override;
+	virtual FText GetNodeTooltipText() const override;
+	virtual FLinearColor GetNodeTitleColor() const override { return FLinearColor::Green; }
+ 
+	virtual TArray<EPVRenderType> GetDefaultRenderType() const override { return TArray<EPVRenderType>{}; }
+	virtual TArray<EPVRenderType> GetSupportedRenderTypes() const override { return TArray<EPVRenderType>{}; }
+#endif
+ 
+protected:
+	virtual TArray<FPCGPinProperties> InputPinProperties() const override;
+	virtual FPCGDataTypeIdentifier GetOutputPinTypeIdentifier() const override;
+	virtual FPCGDataTypeIdentifier GetInputPinTypeIdentifier() const override;
+	virtual FPCGElementPtr CreateElement() const override;
+ 
+public:
+	UPROPERTY(EditAnywhere, Category = "Graft Skeletons", Meta = (ShowOnlyInnerProperties, FullyExpand="true"))
+	TArray<struct FPVFoliageAttributes> GrafterInfos;
+};
+ 
+class FPVGrafterPaletteElement : public FPVBaseElement
+{
+protected:
+	virtual bool ExecuteInternal(FPCGContext* InContext) const override;
+ 
+public:
+	virtual bool CanExecuteOnlyOnMainThread(FPCGContext* Context) const override { return true; }
+};
\ No newline at end of file
diff --git a/Engine/Plugins/Experimental/ProceduralVegetationEditor/Source/ProceduralVegetation/Private/Nodes/PVGrafterSettings.cpp b/Engine/Plugins/Experimental/ProceduralVegetationEditor/Source/ProceduralVegetation/Private/Nodes/PVGrafterSettings.cpp
--- a/Engine/Plugins/Experimental/ProceduralVegetationEditor/Source/ProceduralVegetation/Private/Nodes/PVGrafterSettings.cpp
+++ b/Engine/Plugins/Experimental/ProceduralVegetationEditor/Source/ProceduralVegetation/Private/Nodes/PVGrafterSettings.cpp
@@ -5,6 +5,7 @@
 #include "ProceduralVegetationModule.h"
 #include "DataTypes/PVGrowthData.h"
 #include "ProceduralVegetationPreset.h"
+#include "DataTypes/PVGrafterPaletteData.h"
 #include "Facades/PVPointFacade.h"
 #include "Helpers/PCGHelpers.h"
 #include "Implementations/PVGrafter.h"
@@ -68,7 +69,7 @@ FText UPVGrafterSettings::GetDefaultNodeTitle() const
 FText UPVGrafterSettings::GetNodeTooltipText() const
 {
 	return LOCTEXT("NodeTooltip",
-		"This allows the user to distribute a skeleton onto another one. "
+		"This allows the user to distribute a collection of grafts (from Grafter Palette) onto a plant skeleton."
 		"\n\nPress Ctrl + L to lock/unlock node output"
 	);
 }
@@ -78,12 +79,12 @@ TArray<FPCGPinProperties> UPVGrafterSettings::InputPinProperties() const
 {
 	TArray<FPCGPinProperties> Properties;
 
-	FPCGPinProperties& Pin1 = Properties.Emplace_GetRef(PVGrafterInputPins::Skeleton, GetInputPinTypeIdentifier());
+	FPCGPinProperties& Pin1 = Properties.Emplace_GetRef(PVGrafterInputPins::Skeleton, GetSkeletonPinTypeIdentifier());
 	Pin1.SetRequiredPin();
 	Pin1.SetAllowMultipleConnections(false);
 	Pin1.bAllowMultipleData = false;
 
-	FPCGPinProperties& Pin2 = Properties.Emplace_GetRef(PVGrafterInputPins::Graft, GetInputPinTypeIdentifier());
+	FPCGPinProperties& Pin2 = Properties.Emplace_GetRef(PVGrafterInputPins::Graft, GetGraftPinTypeIdentifier());
 	Pin2.SetRequiredPin();
 	Pin2.SetAllowMultipleConnections(false);
 	Pin2.bAllowMultipleData = false;
@@ -96,6 +97,16 @@ FPCGDataTypeIdentifier UPVGrafterSettings::GetOutputPinTypeIdentifier() const
 	return FPCGDataTypeIdentifier{FPVDataTypeInfoGrowth::AsId()};
 }
 
+FPCGDataTypeIdentifier UPVGrafterSettings::GetSkeletonPinTypeIdentifier() const
+{
+	return FPCGDataTypeIdentifier{ FPVDataTypeInfoGrowth::AsId() };
+}
+
+FPCGDataTypeIdentifier UPVGrafterSettings::GetGraftPinTypeIdentifier() const
+{
+	return FPCGDataTypeIdentifier{ FPVDataTypeInfoGrafterPalette::AsId() };
+}
+
 FPCGElementPtr UPVGrafterSettings::CreateElement() const
 {
 	return MakeShared<FPVGrafterElement>();
@@ -146,15 +157,32 @@ bool FPVGrafterElement::ExecuteInternal(FPCGContext* InContext) const
 	}
 
 	const UPVGrowthData* SkeletonInputData = Cast<UPVGrowthData>(SkeletonPinInputs[0].Data);
-	const UPVGrowthData* GraftInputData = Cast<UPVGrowthData>(GraftPinInputs[0].Data);
-	if (!SkeletonInputData || !GraftInputData)
+	const UPVGrafterPaletteData* GraftPaletteInputData = Cast<UPVGrafterPaletteData>(GraftPinInputs[0].Data);
+	if (!SkeletonInputData || !GraftPaletteInputData)
 	{
 		PCGLog::InputOutput::LogInvalidInputDataError(InContext);
 		return true;
 	}
 
 	FManagedArrayCollection SourceSkeletonCollection = SkeletonInputData->GetCollection();
-	FManagedArrayCollection SourceGraftCollection = GraftInputData->GetCollection();
+	const TArray<TObjectPtr<UPVGrowthData>>& GraftPaletteElements = GraftPaletteInputData->GetGrowthDataElements();
+	if (GraftPaletteElements.IsEmpty())
+	{
+		return true;
+	}
+	
+	TArray<FManagedArrayCollection> SourceGraftCollections;
+	for (int32 i = 0; i < GraftPaletteElements.Num(); ++i)
+	{
+		const UPVGrowthData* GraftElement = GraftPaletteElements[i];
+		if (!GraftElement)
+		{
+			continue;
+		}
+
+		FManagedArrayCollection SourceGraft = GraftElement->GetCollection();
+		SourceGraftCollections.Add(SourceGraft);
+	}
 
 	FManagedArrayCollection OutCollection;
 	SourceSkeletonCollection.CopyTo(&OutCollection);
@@ -166,11 +194,6 @@ bool FPVGrafterElement::ExecuteInternal(FPCGContext* InContext) const
 		PCGLog::LogWarningOnGraph(LOCTEXT("FoliageDataRemovedFromSourceSkeleton", "Foliage data removed from source skeleton"), InContext);
 	}
 
-	if (FPVGrafter::RemoveFoliageDataIfPresent(SourceGraftCollection))
-	{
-		PCGLog::LogWarningOnGraph(LOCTEXT("FoliageDataRemovedFromGraft", "Foliage data removed from source graft"), InContext);
-	}
-
 	if (Settings->Mode == EPVDistributionSettingsMode::HormoneBasedSettings)
 	{
 		const FPVGrafter::FHormoneBasedSettings::FDistributionSettings DistributionSettings = {
@@ -216,12 +239,13 @@ bool FPVGrafterElement::ExecuteInternal(FPCGContext* InContext) const
 		FPVGrafter::DistributeGraftWithHormoneBasedSettings(
 			OutCollection,
 			SourceSkeletonCollection,
-			SourceGraftCollection,
+			SourceGraftCollections,
 			DistributionSettings,
 			PhyllotaxySettings,
 			ScaleSettings,
 			VectorSettings,
-			MiscSettings);
+			MiscSettings,
+			Settings->Conditions);
 	}
 	else
 	{
@@ -267,11 +291,12 @@ bool FPVGrafterElement::ExecuteInternal(FPCGContext* InContext) const
 		FPVGrafter::DistributeGraftWithParametricSettings(
 			OutCollection,
 			SourceSkeletonCollection,
-			SourceGraftCollection,
+			SourceGraftCollections,
 			SpacingSettings,
 			AngleSettings,
 			ScaleSettings,
-			MiscSettings);
+			MiscSettings,
+			Settings->Conditions);
 	}
 
 	OutManagedArrayCollectionData->Initialize(MoveTemp(OutCollection));
diff --git a/Engine/Plugins/Experimental/ProceduralVegetationEditor/Source/ProceduralVegetation/Private/Nodes/PVGrafterSettings.h b/Engine/Plugins/Experimental/ProceduralVegetationEditor/Source/ProceduralVegetation/Private/Nodes/PVGrafterSettings.h
--- a/Engine/Plugins/Experimental/ProceduralVegetationEditor/Source/ProceduralVegetation/Private/Nodes/PVGrafterSettings.h
+++ b/Engine/Plugins/Experimental/ProceduralVegetationEditor/Source/ProceduralVegetation/Private/Nodes/PVGrafterSettings.h
@@ -2,6 +2,7 @@
 #pragma once
 
 #include "CoreMinimal.h"
+#include "Helpers/PVFoliageHelper.h"
 #include "Helpers/PVPhyllotaxyHelper.h"
 #include "Utils/PVFloatRamp.h"
 #include "Nodes/PVBaseSettings.h"
@@ -270,6 +271,8 @@ public:
 protected:
 	virtual TArray<FPCGPinProperties> InputPinProperties() const override;
 	virtual FPCGDataTypeIdentifier GetOutputPinTypeIdentifier() const override;
+	virtual FPCGDataTypeIdentifier GetSkeletonPinTypeIdentifier() const;
+	virtual FPCGDataTypeIdentifier GetGraftPinTypeIdentifier() const;
 	virtual FPCGElementPtr CreateElement() const override;
 #if WITH_EDITOR
 	virtual bool CanEditChange(const FProperty* InProperty) const override;
@@ -289,6 +292,9 @@ private:
 
 	UPROPERTY(EditAnywhere, Category="Misc Settings", meta=(PCG_Overridable))
 	int32 RandomSeed = 123456;
+	
+	UPROPERTY(EditAnywhere, Category="ConditionSettings", meta = (PCG_Overridable, ShowOnlyInnerProperties))
+	FPVFoliageDistributionConditions Conditions;
 };
 
 class FPVGrafterElement : public FPVBaseElement
diff --git a/Engine/Plugins/Experimental/ProceduralVegetationEditor/Source/ProceduralVegetation/Public/DataTypes/PVGrafterPaletteData.h b/Engine/Plugins/Experimental/ProceduralVegetationEditor/Source/ProceduralVegetation/Public/DataTypes/PVGrafterPaletteData.h
new file mode 100644
--- /dev/null
+++ b/Engine/Plugins/Experimental/ProceduralVegetationEditor/Source/ProceduralVegetation/Public/DataTypes/PVGrafterPaletteData.h
@@ -0,0 +1,48 @@
+// Copyright Epic Games, Inc. All Rights Reserved.
+
+#pragma once
+
+#include "PVData.h"
+#include "PVGrowthData.h"
+#include "PVGrafterPaletteData.generated.h"
+
+USTRUCT()
+struct FPVDataTypeInfoGrafterPalette : public FPCGDataTypeInfo
+{
+	GENERATED_BODY()
+
+	PCG_DECLARE_TYPE_INFO(PROCEDURALVEGETATION_API);
+
+#if WITH_EDITOR
+	// If the type should be hidden to the user.
+	virtual bool Hidden() const override { return true; };
+#endif // WITH_EDITOR
+};
+
+UCLASS(BlueprintType, ClassGroup = (Procedural))
+class UPVGrafterPaletteData : public UPCGData
+{
+	GENERATED_BODY()
+
+public:
+	//~ Begin UPCGData interface
+	PCG_ASSIGN_DEFAULT_TYPE_INFO(FPVDataTypeInfoGrafterPalette)
+	//~ End UPCGData interface
+
+	void Initialize(TArray<TObjectPtr<UPVGrowthData>>&& InGrowthDataElements);
+	
+	// ~Begin UPCGData interface
+	virtual void AddToCrc(FArchiveCrc32& Ar, bool bFullDataCrc) const override;
+	// ~End UPCGData interface
+	
+	const TArray<TObjectPtr<UPVGrowthData>>& GetGrowthDataElements() const { return GrowthDataElements; }
+	TArray<TObjectPtr<UPVGrowthData>>& GetGrowthDataElements() { return GrowthDataElements; }
+	
+protected:
+	// ~Begin UPCGData interface
+	virtual bool SupportsFullDataCrc() const override { return true; }
+	// ~End UPCGData interface
+	
+	UPROPERTY()
+	TArray<TObjectPtr<UPVGrowthData>> GrowthDataElements;
+};
\ No newline at end of file
diff --git a/Engine/Plugins/Experimental/ProceduralVegetationEditor/Source/ProceduralVegetation/Public/Facades/PVAttributesNames.h b/Engine/Plugins/Experimental/ProceduralVegetationEditor/Source/ProceduralVegetation/Public/Facades/PVAttributesNames.h
--- a/Engine/Plugins/Experimental/ProceduralVegetationEditor/Source/ProceduralVegetation/Public/Facades/PVAttributesNames.h
+++ b/Engine/Plugins/Experimental/ProceduralVegetationEditor/Source/ProceduralVegetation/Public/Facades/PVAttributesNames.h
@@ -138,6 +138,11 @@ namespace PV::AttributeNames
 	inline constexpr TCHAR LeafGrowthStr[] = TEXT("leafGrowth");
 	inline constexpr TCHAR AbscissionSenescenseStr[] = TEXT("abscissionSenescense");
 	inline constexpr TCHAR LateralElongationStr[] = TEXT("lateralElongation");
+	inline constexpr TCHAR GraftLightStr[] = TEXT("GraftLight");
+	inline constexpr TCHAR GraftScaleStr[] = TEXT("GraftScale");
+	inline constexpr TCHAR GraftTipStr[] = TEXT("GraftTip");
+	inline constexpr TCHAR GraftUpAlignmentStr[] = TEXT("GraftUpAlignment");
+	inline constexpr TCHAR GraftHealthStr[] = TEXT("GraftHealth");
 
 	inline static const FName LeafPhyllotaxy = FName(LeafPhyllotaxyStr);
 	inline static const FName FoliagePath = FName(FoliagePathStr);
@@ -145,6 +150,11 @@ namespace PV::AttributeNames
 	inline static const FName LeafGrowth = FName(LeafGrowthStr);
 	inline static const FName AbscissionSenescense = FName(AbscissionSenescenseStr);
 	inline static const FName LateralElongation = FName(LateralElongationStr);
+	inline static const FName GraftLight = FName(GraftLightStr);
+	inline static const FName GraftScale = FName(GraftScaleStr);
+	inline static const FName GraftTip = FName(GraftTipStr);
+	inline static const FName GraftUpAlignment = FName(GraftUpAlignmentStr);
+	inline static const FName GraftHealth = FName(GraftHealthStr);
 
 	//PlantProfiles Attributes
 	inline constexpr TCHAR ProfilePointsStr[] = TEXT("ProfilePoints");
diff --git a/Engine/Plugins/Experimental/ProceduralVegetationEditor/Source/ProceduralVegetation/Public/Facades/PVMetaInfoFacade.h b/Engine/Plugins/Experimental/ProceduralVegetationEditor/Source/ProceduralVegetation/Public/Facades/PVMetaInfoFacade.h
--- a/Engine/Plugins/Experimental/ProceduralVegetationEditor/Source/ProceduralVegetation/Public/Facades/PVMetaInfoFacade.h
+++ b/Engine/Plugins/Experimental/ProceduralVegetationEditor/Source/ProceduralVegetation/Public/Facades/PVMetaInfoFacade.h
@@ -7,6 +7,14 @@
 
 namespace PV::Facades
 {
+	struct PROCEDURALVEGETATION_API FGraftData
+	{
+		float Light;
+		float UpAlignment;
+		float Health;
+		bool Tip;
+	};
+	
 	class PROCEDURALVEGETATION_API FMetaInfoFacade
 	{
 	public:
@@ -19,6 +27,9 @@ namespace PV::Facades
 		void CreateGuid(const FString& InPath);
 		FGuid GetGuid() const;
 		
+		void SetGraftAttributes(const FGraftData& InGraftData);
+		FGraftData GetGraftAttributes() const;
+		
 		const TArray<float>& GetLeafGrowth() const;
 		const TArray<float>& GetAbscissionSenescense() const;
 		const TArray<float>& GetLateralElongation() const;
@@ -33,5 +44,9 @@ namespace PV::Facades
 		TManagedArrayAccessor<TArray<float>> LeafGrowthAttribute;
 		TManagedArrayAccessor<TArray<float>> AbscissionSenescenseAttribute;
 		TManagedArrayAccessor<TArray<float>> LateralElongationAttribute;
+		TManagedArrayAccessor<float> GraftLightAttribute;
+		TManagedArrayAccessor<float> GraftUpAlignmentAttribute;
+		TManagedArrayAccessor<float> GraftHealthAttribute;
+		TManagedArrayAccessor<float> GraftTipAttribute;
 	};
 }
diff --git a/Engine/Plugins/Experimental/ProceduralVegetationEditor/Source/ProceduralVegetationEditor/Private/Helpers/PVExportHelper.cpp b/Engine/Plugins/Experimental/ProceduralVegetationEditor/Source/ProceduralVegetationEditor/Private/Helpers/PVExportHelper.cpp
--- a/Engine/Plugins/Experimental/ProceduralVegetationEditor/Source/ProceduralVegetationEditor/Private/Helpers/PVExportHelper.cpp
+++ b/Engine/Plugins/Experimental/ProceduralVegetationEditor/Source/ProceduralVegetationEditor/Private/Helpers/PVExportHelper.cpp
@@ -56,6 +56,7 @@
 #include "PhysicsEngine/PhysicsAsset.h"
 #include "PhysicsEngine/PhysicsConstraintTemplate.h"
 #include "PhysicsEngine/SkeletalBodySetup.h"
+#include "ProceduralVegetation/Private/Helpers/PVFoliageHelper.h"
 #include "ProceduralVegetation/Private/Implementations/PVFoliage.h"
  
 #include "UObject/Package.h"
diff --git a/Engine/Plugins/Experimental/ProceduralVegetationEditor/Source/ProceduralVegetationEditor/Private/Visualizations/PVDataVisualization.cpp b/Engine/Plugins/Experimental/ProceduralVegetationEditor/Source/ProceduralVegetationEditor/Private/Visualizations/PVDataVisualization.cpp
--- a/Engine/Plugins/Experimental/ProceduralVegetationEditor/Source/ProceduralVegetationEditor/Private/Visualizations/PVDataVisualization.cpp
+++ b/Engine/Plugins/Experimental/ProceduralVegetationEditor/Source/ProceduralVegetationEditor/Private/Visualizations/PVDataVisualization.cpp
@@ -38,7 +38,8 @@
  
 #include "Nodes/PVBaseSettings.h"
 #include "Nodes/PVObjectInteractionSettings.h"
- 
+#include "ProceduralVegetation/Private/Helpers/PVFoliageHelper.h"
+
 #include "Rendering/SkeletalMeshRenderData.h"
  
 #include "Visualizations/PVSkeletonVisualizerComponent.h"
