[PVE] Fixed compile error

#Synced-CL 51491700

diff --git a/Engine/Plugins/Experimental/ProceduralVegetationEditor/Source/ProceduralVegetation/Private/Helpers/PVUtilities.cpp b/Engine/Plugins/Experimental/ProceduralVegetationEditor/Source/ProceduralVegetation/Private/Helpers/PVUtilities.cpp
--- a/Engine/Plugins/Experimental/ProceduralVegetationEditor/Source/ProceduralVegetation/Private/Helpers/PVUtilities.cpp
+++ b/Engine/Plugins/Experimental/ProceduralVegetationEditor/Source/ProceduralVegetation/Private/Helpers/PVUtilities.cpp
@@ -3,6 +3,7 @@
 #include "Helpers/PVUtilities.h"
 #include "ProceduralVegetationModule.h"
 #include "ProceduralVegetationPreset.h"
+#include "PVFoliageHelper.h"
 #include "AssetRegistry/AssetRegistryModule.h"
 #include "Components/InstancedSkinnedMeshComponent.h"
 #include "Engine/StaticMesh.h"
