[PVE] Texture Setup node
- Added node for Texture setup, the purpose of the node is to create trim sheet style texture to pack multiple generations in one texture.
- Each generation can contains multiple channels per generation.
- Also have a bake textures button to save the textures to the content browser.
- The Texture menipulation is done in compute shaders
- First step is to do the dilation, the border pixel at the right and left are dilated to a given range
- Than other operations are done in a second pass
- To add the button we needed Detail Customization.
- Added the visualization for the node
- Added visualization mesh and material in content
- Added support to add the Material in the Texture setup node, this material along with other data gets propogated as output and can be connected to mesher
- Added a new type for the Texture setup Data

#virtualized
#Synced-CL 50970602

diff --git a/Engine/Plugins/Experimental/ProceduralVegetationEditor/Shaders/PVTextureSetupCS.usf b/Engine/Plugins/Experimental/ProceduralVegetationEditor/Shaders/PVTextureSetupCS.usf
new file mode 100644
--- /dev/null
+++ b/Engine/Plugins/Experimental/ProceduralVegetationEditor/Shaders/PVTextureSetupCS.usf
@@ -0,0 +1,49 @@
+// Copyright Epic Games, Inc. All Rights Reserved.
+#include "/Engine/Public/Platform.ush"
+
+float OffsetX;
+float TileX;
+float TileY;
+float ScaleX;
+
+Texture2D<float4> Source;
+SamplerState SourceSampler;
+Texture2D<float4> SourceB;
+SamplerState SourceSamplerB;
+RWTexture2D<float4> Target;
+uint2 TargetSize;
+
+[numthreads(8,8,1)]
+void DilationCS(uint3 DTid : SV_DispatchThreadID, uint3 GTid : SV_GroupThreadID, uint3 Gid : SV_GroupID)
+{
+	if (DTid.x >= TargetSize.x || DTid.y >= TargetSize.y) return;
+	
+	float2 UV = (float2(DTid.xy)) / float2(TargetSize.xy);
+	UV.x = UV.x - OffsetX;
+	UV.x = UV.x * TileX; 
+	UV.y = UV.y * TileY;
+	
+	if (UV.x < 0) UV.x = 0;
+	else if (UV.x > ScaleX * TileX) UV.x = 1.0;
+	else UV.x = UV.x / ScaleX;
+	float4 Color = Source.SampleLevel(SourceSampler, UV, 0);
+	
+	Target[DTid.xy] = Color;
+}
+
+[numthreads(8,8,1)]
+void ScaleCS(uint3 DTid : SV_DispatchThreadID, uint3 GTid : SV_GroupThreadID, uint3 Gid : SV_GroupID)
+{
+	if (DTid.x >= TargetSize.x || DTid.y >= TargetSize.y) return;
+	
+	float2 UV = (float2(DTid.xy)) / float2(TargetSize.xy);
+	//UV.x = UV.x + OffsetX;
+	float2 OutTid = float2(DTid.x + OffsetX * TargetSize.x, DTid.y);
+		
+	if (UV.x <= (ScaleX))
+	{
+		UV.x = ( UV.x / ScaleX );
+		float4 Color = Source.SampleLevel(SourceSampler, UV, 0);
+		Target[OutTid] = Color;
+	}
+}
diff --git a/Engine/Plugins/Experimental/ProceduralVegetationEditor/Source/ProceduralVegetation/Private/DataTypes/PVTextureSetupData.cpp b/Engine/Plugins/Experimental/ProceduralVegetationEditor/Source/ProceduralVegetation/Private/DataTypes/PVTextureSetupData.cpp
new file mode 100644
--- /dev/null
+++ b/Engine/Plugins/Experimental/ProceduralVegetationEditor/Source/ProceduralVegetation/Private/DataTypes/PVTextureSetupData.cpp
@@ -0,0 +1,35 @@
+// Copyright Epic Games, Inc. All Rights Reserved.
+
+#include "DataTypes/PVTextureSetupData.h"
+#include "Serialization/ArchiveCrc32.h"
+
+bool FPVTextureSetupInfo::IsOverflowing() const
+{
+	for (auto Generation : GenerationUVs)
+	{
+		if (Generation.GetOverflowValue() > 0)
+		{
+			return true;
+		}
+	}
+
+	return false;
+}
+
+PCG_DEFINE_TYPE_INFO(FPVDataTypeInfoTextureSetup, UPVTextureSetupData);
+
+void UPVTextureSetupData::AddToCrc(FArchiveCrc32& Ar, bool bFullDataCrc) const
+{
+	Super::AddToCrc(Ar, bFullDataCrc);
+		
+	if (bFullDataCrc)
+	{
+		FString ClassName = StaticClass()->GetPathName();
+		Ar << ClassName;
+		Ar << TextureSetupInfo;
+	}
+	else
+	{
+		AddUIDToCrc(Ar);
+	}
+}
\ No newline at end of file
diff --git a/Engine/Plugins/Experimental/ProceduralVegetationEditor/Source/ProceduralVegetation/Private/Implementations/PVMaterialSettings.cpp b/Engine/Plugins/Experimental/ProceduralVegetationEditor/Source/ProceduralVegetation/Private/Implementations/PVMaterialSettings.cpp
--- a/Engine/Plugins/Experimental/ProceduralVegetationEditor/Source/ProceduralVegetation/Private/Implementations/PVMaterialSettings.cpp
+++ b/Engine/Plugins/Experimental/ProceduralVegetationEditor/Source/ProceduralVegetation/Private/Implementations/PVMaterialSettings.cpp
@@ -38,10 +38,10 @@ void FPVMaterialSettings::ApplyMaterialSettings(FManagedArrayCollection& Collect
 	}
 }
 
-void FPVMaterialSettings::SetMaterial(FString TrunkMaterial, FVector2f UMinMax, int32 Index)
+void FPVMaterialSettings::SetMaterial(UMaterialInterface* Material, FVector2f UMinMax, int32 Index)
 {
 	FTrunkGenerationMaterialSetup MaterialSetup;
-	MaterialSetup.Material = LoadObject<UMaterialInterface>(nullptr, TrunkMaterial, {}, LOAD_NoWarn | LOAD_Quiet);
+	MaterialSetup.Material = Material;
 	MaterialSetup.URange = FFloatRange(UMinMax.X, UMinMax.Y);
 
 	if (!MaterialSetups.IsValidIndex(Index))
@@ -55,6 +55,12 @@ void FPVMaterialSettings::SetMaterial(FString TrunkMaterial, FVector2f UMinMax,
 	}
 }
 
+void FPVMaterialSettings::SetMaterial(FString TrunkMaterial, FVector2f UMinMax, int32 Index)
+{
+	UMaterialInterface* Material = LoadObject<UMaterialInterface>(nullptr, TrunkMaterial, {}, LOAD_NoWarn | LOAD_Quiet);
+	SetMaterial(Material, UMinMax, Index);
+}
+
 void FPVMaterialSettings::SetBranchUVMaterial(FManagedArrayCollection& Collection, TArray<int32>& UsedMaterials, TArray<FBranchUVData>& BranchUVData) const
 {
 	PV::Facades::FPointFacade PointFacade = PV::Facades::FPointFacade(Collection);
diff --git a/Engine/Plugins/Experimental/ProceduralVegetationEditor/Source/ProceduralVegetation/Private/Implementations/PVMaterialSettings.h b/Engine/Plugins/Experimental/ProceduralVegetationEditor/Source/ProceduralVegetation/Private/Implementations/PVMaterialSettings.h
--- a/Engine/Plugins/Experimental/ProceduralVegetationEditor/Source/ProceduralVegetation/Private/Implementations/PVMaterialSettings.h
+++ b/Engine/Plugins/Experimental/ProceduralVegetationEditor/Source/ProceduralVegetation/Private/Implementations/PVMaterialSettings.h
@@ -96,6 +96,8 @@ struct FPVMaterialSettings
 	 calculates the U offset and write to the Point Group in collection
 	 */
 	void ApplyMaterialSettings(FManagedArrayCollection& Collection) const;
+	
+	void SetMaterial(UMaterialInterface* Material, FVector2f UMinMax, int32 Index);
 
 	void SetMaterial(FString TrunkMaterial, FVector2f UMinMax, int32 Index);
 
diff --git a/Engine/Plugins/Experimental/ProceduralVegetationEditor/Source/ProceduralVegetation/Private/Implementations/PVMeshBuilder.h b/Engine/Plugins/Experimental/ProceduralVegetationEditor/Source/ProceduralVegetation/Private/Implementations/PVMeshBuilder.h
--- a/Engine/Plugins/Experimental/ProceduralVegetationEditor/Source/ProceduralVegetation/Private/Implementations/PVMeshBuilder.h
+++ b/Engine/Plugins/Experimental/ProceduralVegetationEditor/Source/ProceduralVegetation/Private/Implementations/PVMeshBuilder.h
@@ -32,12 +32,10 @@ struct FPVMeshBuilderParams
 	GENERATED_BODY()
 
 	UPROPERTY(EditAnywhere, Category="Materials",
-		meta=(PCG_NotOverridable, PinHiddenByDefault, InlineEditConditionToggle, Tooltip=
-			"Toggles between default and custom material settings.\n\nDefault material setting is defined in Procedural Vegetation Preset.\n\nCustom material settings allow for editing."
-		))
+		meta=(PCG_NotOverridable, PinHiddenByDefault, InlineEditConditionToggle))
 	bool bOverrideMaterial = false;
 
-	UPROPERTY(EditAnywhere, Category="Materials", meta = (PCG_NotOverridable, EditCondition="bOverrideMaterial"))
+	UPROPERTY(EditAnywhere, Category="Materials" , meta = (PCG_NotOverridable))
 	FPVMaterialSettings MaterialSettings;
 
 	UPROPERTY(EditAnywhere, Category="Mesh Settings", DisplayName="Point Removal",
diff --git a/Engine/Plugins/Experimental/ProceduralVegetationEditor/Source/ProceduralVegetation/Private/Implementations/PVTextureSetup.cpp b/Engine/Plugins/Experimental/ProceduralVegetationEditor/Source/ProceduralVegetation/Private/Implementations/PVTextureSetup.cpp
new file mode 100644
--- /dev/null
+++ b/Engine/Plugins/Experimental/ProceduralVegetationEditor/Source/ProceduralVegetation/Private/Implementations/PVTextureSetup.cpp
@@ -0,0 +1,389 @@
+// Copyright Epic Games, Inc. All Rights Reserved.
+
+#include "PVTextureSetup.h"
+#include "ImageUtils.h"
+#include "PVTextureSetupCS.h"
+#include "RenderGraphBuilder.h"
+#include "RenderGraphUtils.h"
+#include "RHIStaticStates.h"
+#include "TextureResource.h"
+#include "AssetRegistry/AssetRegistryModule.h"
+#include "Engine/Texture2D.h"
+#include "Engine/TextureRenderTarget2D.h"
+#include "Materials/MaterialInstanceConstant.h"
+#include "Misc/PackageName.h"
+#include "UObject/SavePackage.h"
+#include "Params//PVTextureSetupParams.h"
+
+#define LOCTEXT_NAMESPACE "PVTextureSetup"
+
+IMPLEMENT_GLOBAL_SHADER(FPVTextureSetupDilationCS, "/Plugins/Experimental/ProceduralVegetationEditor/PVTextureSetupCS.usf", "DilationCS", SF_Compute);
+IMPLEMENT_GLOBAL_SHADER(FPVTextureSetupScaleCS, "/Plugins/Experimental/ProceduralVegetationEditor/PVTextureSetupCS.usf", "ScaleCS", SF_Compute);
+
+TArray<FString> FPVTextureSetupImplementation::GetChannels(const TArray<FPVTextureGenerationParams>& Array)
+{
+	TArray<FString> Channels;
+
+	for (auto Generation : Array)
+	{
+		for (auto Channel : Generation.GetChannelNames())
+		{
+			Channels.AddUnique(Channel);
+		}
+	}
+	return Channels;
+}
+
+UTexture* FPVTextureSetupImplementation::CreateTexture(int Width, int Height, bool bCanCreateUAV, FName Name)
+{
+	UTextureRenderTarget2D* RenderTargetTexture = NewObject<UTextureRenderTarget2D>(GetTransientPackage(), Name);
+	RenderTargetTexture->ClearColor = FColor::Black;
+	RenderTargetTexture->bCanCreateUAV = bCanCreateUAV;
+	RenderTargetTexture->bSupportsUAV = bCanCreateUAV;
+	RenderTargetTexture->InitAutoFormat(Width, Height);
+	RenderTargetTexture->UpdateResourceImmediate(true);
+	
+	return RenderTargetTexture;
+}
+
+TMap<FString, TObjectPtr<UTexture>> FPVTextureSetupImplementation::CreateTextures(int32 Resolution,const TArray<FString>& ChannelNames, const FString& AssetNamePrefix, bool bVirtualTextureStreaming)
+{
+	TMap<FString, TObjectPtr<UTexture>> TextureSet;
+	for (const auto& ChannelName : ChannelNames)
+	{
+		FString Name = AssetNamePrefix + ChannelName;
+		UTexture* Texture = CreateTexture(Resolution, Resolution, true, FName(Name));
+		Texture->VirtualTextureStreaming = bVirtualTextureStreaming;
+		TextureSet.Add(ChannelName, Texture);
+	}
+	return TextureSet;
+}
+
+FRHITexture* FPVTextureSetupImplementation::GetRHITexture(UTexture* Texture)
+{
+	if (UTextureRenderTarget2D* RenderTarget = Cast<UTextureRenderTarget2D>(Texture))
+	{
+		const FTextureRenderTargetResource* Resource = RenderTarget->GetRenderTargetResource();//>GetResource();
+		if (!Resource)
+		{
+			return nullptr;
+		}
+	
+		return Resource->GetRenderTargetTexture();
+	}
+
+	const FTextureResource* Resource = Texture->GetResource();
+	if (!Resource)
+	{
+		return nullptr;
+	}
+
+	return Resource->GetTexture2DRHI();
+}
+
+FRDGTextureSRVRef FPVTextureSetupImplementation::GetTextureSRVRef(FRDGBuilder& GraphBuilder,UTexture* InputTexture,const TCHAR* Name)
+{
+	FRHITexture* RHITexture = GetRHITexture(InputTexture);
+	TRefCountPtr<IPooledRenderTarget> InputRenderTarget = CreateRenderTarget(RHITexture, Name);
+	FRDGTextureRef InputTextureRef = GraphBuilder.RegisterExternalTexture(InputRenderTarget);
+	
+	FRDGTextureSRVRef InputTextureSRV = GraphBuilder.CreateSRV(FRDGTextureSRVDesc::CreateForMipLevel(InputTextureRef, 0));
+	return InputTextureSRV;
+}
+
+FRDGTextureUAVRef FPVTextureSetupImplementation::GetTextureUAVRef(FRDGBuilder& GraphBuilder,UTexture* Texture,const TCHAR* Name)
+{
+	FRHITexture* RHITexture = GetRHITexture(Texture);
+	
+	TRefCountPtr<IPooledRenderTarget> RenderTarget = CreateRenderTarget(RHITexture, Name);
+	
+	FRDGTextureRef TextureRDGRef = GraphBuilder.RegisterExternalTexture(RenderTarget);
+	FRDGTextureUAVDesc Descriptor(TextureRDGRef);
+	Descriptor.MipLevel = 0;
+	
+	FRDGTextureUAVRef TextureUAV = GraphBuilder.CreateUAV(Descriptor);
+	return TextureUAV;
+}
+
+void FPVTextureSetupImplementation::DilationPass(FRDGBuilder& GraphBuilder, float DilationFactor, float TileX, float TileY, UTexture* InputTexture, UTexture* OutputTexture)
+{
+	float ScaleX = 1.0 - (DilationFactor * 2.0);
+	
+	FRDGTextureSRVRef InputTextureSRV = GetTextureSRVRef(GraphBuilder, InputTexture, TEXT("PVDilationInputTexture"));
+	FRDGTextureUAVRef RWOutputTextureUAV = GetTextureUAVRef(GraphBuilder, OutputTexture, TEXT("PVDilationOutputTexture"));
+		
+	FPVTextureSetupDilationCS::FParameters* Params = GraphBuilder.AllocParameters<FPVTextureSetupDilationCS::FParameters>();
+	Params->OffsetX = DilationFactor;
+	Params->ScaleX = ScaleX;
+	Params->TileX = TileX;
+	Params->TileY = TileY;
+	Params->Source = InputTextureSRV;
+	Params->SourceSampler = TStaticSamplerState<SF_Bilinear, AM_Clamp, AM_Clamp>::GetRHI();
+	Params->Target = RWOutputTextureUAV;
+	Params->TargetSize = FIntPoint(OutputTexture->GetSurfaceWidth(), OutputTexture->GetSurfaceHeight());
+	TShaderMapRef<FPVTextureSetupDilationCS> PVTextureSetupDilationCS(GetGlobalShaderMap(GMaxRHIFeatureLevel));
+
+	const FIntVector GroupCount = FComputeShaderUtils::GetGroupCount(Params->TargetSize, FIntPoint(8,8));
+	//FIntVector DispatchCount(FMath::DivideAndRoundUp(Params->TargetSize.X, 8),FMath::DivideAndRoundUp(Params->TargetSize.Y, 8),1);
+	FComputeShaderUtils::AddPass(GraphBuilder, RDG_EVENT_NAME("PVDilation"), PVTextureSetupDilationCS, Params, GroupCount);
+	
+}
+
+void FPVTextureSetupImplementation::ScaleAndCombineTexturePass(FRDGBuilder& GraphBuilder, UTexture* TargetTexture, UTexture* InputTexture, const float WidthRatio, const float OffsetX, const FPVTextureGenerationParams& Generation)
+{
+	FRDGTextureSRVRef InputTextureSRV = GetTextureSRVRef(GraphBuilder, InputTexture, TEXT("PVDilationInputTexture"));
+	FRDGTextureUAVRef RWOutputTextureUAV = GetTextureUAVRef(GraphBuilder, Cast<UTextureRenderTarget2D>(TargetTexture), TEXT("PVTrimSheetTexture"));
+		
+	FPVTextureSetupScaleCS::FParameters* Params = GraphBuilder.AllocParameters<FPVTextureSetupScaleCS::FParameters>();
+	Params->OffsetX = OffsetX;
+	Params->ScaleX = Generation.XScale / WidthRatio;
+	Params->TileX = Generation.BaseTileX;
+	Params->TileY = Generation.BaseTileY;
+	Params->Source = InputTextureSRV;
+	Params->SourceSampler = TStaticSamplerState<SF_Bilinear>::GetRHI();
+	Params->Target = RWOutputTextureUAV;
+	Params->TargetSize = FIntPoint(TargetTexture->GetSurfaceWidth(), TargetTexture->GetSurfaceHeight());
+	TShaderMapRef<FPVTextureSetupScaleCS> PVTextureSetupScaleCS(GetGlobalShaderMap(GMaxRHIFeatureLevel));
+
+	const FIntVector GroupCount = FComputeShaderUtils::GetGroupCount(Params->TargetSize, FIntPoint(8,8));
+	//FIntVector DispatchCount(FMath::DivideAndRoundUp(Params->TargetSize.X, 8),FMath::DivideAndRoundUp(Params->TargetSize.Y, 8),1);
+	FComputeShaderUtils::AddPass(GraphBuilder, RDG_EVENT_NAME("PVScale"), PVTextureSetupScaleCS, Params, GroupCount);
+}
+
+void FPVTextureSetupImplementation::AddChannelTextureToTextureSet(const float InOffsetX, const FPVTextureGenerationParams& Generation, const FPVTextureChannelParams& Channel, UTexture* Texture)
+{
+	UTexture* ChannelTexture = Channel.Texture;
+	UTexture* DilatedTexture = CreateTexture(ChannelTexture->GetSurfaceWidth(), ChannelTexture->GetSurfaceHeight(), true, NAME_None);
+	
+	ENQUEUE_RENDER_COMMAND(RunPVScaleAndCombinePass)
+	([ &Generation, &InOffsetX, ChannelTexture, DilatedTexture, Texture](FRHICommandListImmediate& RHICmdList)
+	{
+		FRDGBuilder GraphBuilder(RHICmdList);
+		float WidthRatio = ChannelTexture->GetSurfaceWidth() > 0 ? static_cast<float>(ChannelTexture->GetSurfaceHeight() / static_cast<float>(ChannelTexture->GetSurfaceWidth())) : 1.0;
+		DilationPass(GraphBuilder, Generation.Dilation, Generation.BaseTileX, Generation.BaseTileY, ChannelTexture, DilatedTexture);
+		ScaleAndCombineTexturePass(GraphBuilder, Texture, DilatedTexture, WidthRatio, InOffsetX, Generation);
+		GraphBuilder.Execute();
+	});
+
+	FlushRenderingCommands();
+}
+
+void FPVTextureSetupImplementation::AddGenerationIntoChannel(const FString& ChannelName, UTexture* Texture, const FPVTextureSetupParams& Params)
+{
+	float XOffset = 0.0f;
+	for (auto Generation : Params.Generations)
+	{
+		//Get the Channel with name
+		if (const FPVTextureChannelParams* Channel = Generation.GetChannel(ChannelName))
+		{
+			if (Channel->Texture != nullptr)
+			{
+				AddChannelTextureToTextureSet(XOffset, Generation, *Channel, Texture);
+				XOffset += Generation.XScale / Channel->WidthRatio;
+			}
+		}
+	}
+}
+
+void FPVTextureSetupImplementation::AddGenerationIntoChannel(const TMap<FString, TObjectPtr<UTexture>>& Map,
+	const FPVTextureSetupParams& Params)
+{
+	for (auto Pair : Map)
+	{
+		AddGenerationIntoChannel(Pair.Key, Pair.Value, Params);
+	}
+}
+
+
+FString FPVTextureSetupImplementation::NormalizeString(FString InString)
+{
+	InString = InString.ToLower();
+	InString.ReplaceInline(TEXT(" "), TEXT(""));
+	InString.ReplaceInline(TEXT("_"), TEXT(""));
+	InString.ReplaceInline(TEXT("-"), TEXT(""));
+	return InString;
+}
+
+EPVTextureChannel FPVTextureSetupImplementation::GetTextureChannelFromString(const FString& InString)
+{
+	FString NormalizedString = NormalizeString(InString);
+
+	static TArray<FString> BaseColorNames = {"basecolor", "albedo", "alb", "diffuse", "col", "color"};
+	static TArray<FString> AONames = {"ao", "ambientocclusion", "occlusion", "occ"};
+	static TArray<FString> NormalNames = {"normal", "norm", "nor"};
+	static TArray<FString> DisplacementNames = {"displacement", "disp", "height"};
+	static TArray<FString> BumpNames = {"bump", "bum"};
+	static TArray<FString> CavityNames = {"cavity", "cav"};
+	static TArray<FString> GlossNames = {"gloss", "gls"};
+	static TArray<FString> RoughnessNames = {"roughness", "rough"};
+	static TArray<FString> SpecularNames = {"specular", "spec"};
+	
+	if (BaseColorNames.Contains(NormalizedString))
+	{
+		return EPVTextureChannel::BaseColor;
+	}
+	else if (AONames.Contains(NormalizedString))
+	{
+		return EPVTextureChannel::AO;
+	}
+	else if (NormalNames.Contains(NormalizedString))
+	{
+		return EPVTextureChannel::Normal;
+	}
+	else if (DisplacementNames.Contains(NormalizedString))
+	{
+		return EPVTextureChannel::Displacement;
+	}
+	else if (BumpNames.Contains(NormalizedString))
+	{
+		return EPVTextureChannel::Bump;
+	}
+	else if (CavityNames.Contains(NormalizedString))
+	{
+		return EPVTextureChannel::Cavity;
+	}
+	else if (GlossNames.Contains(NormalizedString))
+	{
+		return EPVTextureChannel::Gloss;
+	}
+	else if (RoughnessNames.Contains(NormalizedString))
+	{
+		return EPVTextureChannel::Roughness;
+	}
+	else if (SpecularNames.Contains(NormalizedString))
+	{
+		return EPVTextureChannel::Specular;
+	}
+
+	return EPVTextureChannel::Custom;
+}
+
+TObjectPtr<UTexture> FPVTextureSetupImplementation::GetChannelTextureForMaterialParam(FPVTextureSetupInfo& TextureSetupInfo, const FString MaterialParamName)
+{
+	EPVTextureChannel Channel = GetTextureChannelFromString(MaterialParamName);
+	FString ChannelName = StaticEnum<EPVTextureChannel>()->GetNameStringByValue((int64)Channel);
+
+	if (TextureSetupInfo.Channels.Contains(ChannelName))
+	{
+		return TextureSetupInfo.Channels[ChannelName];
+	}
+	else if (TextureSetupInfo.Channels.Contains(MaterialParamName))
+	{
+		return TextureSetupInfo.Channels[MaterialParamName];
+	}
+
+	return nullptr;
+}
+
+void FPVTextureSetupImplementation::SetupMaterial(TObjectPtr<UMaterialInterface> Material, FPVTextureSetupInfo& TextureSetupInfo)
+{
+	TObjectPtr<UMaterialInstanceConstant> MaterialInstance = Cast<UMaterialInstanceConstant>(Material);
+	if (MaterialInstance)
+	{
+#if WITH_EDITOR
+		MaterialInstance->Modify();
+		TArray<FMaterialParameterInfo> Infos;
+		TArray<FGuid> Guids;
+		MaterialInstance->GetAllTextureParameterInfo(Infos, Guids);
+
+		for (auto Info : Infos)
+		{
+			UTexture* Texture = GetChannelTextureForMaterialParam(TextureSetupInfo, Info.Name.ToString());
+			if (Texture)
+			{
+				MaterialInstance->SetTextureParameterValueEditorOnly(Info, Texture);
+			}
+		}
+
+		bool bDirty = MaterialInstance->MarkPackageDirty();
+		FPropertyChangedEvent EmptyPropertyUpdateStruct(NULL);
+		MaterialInstance->PostEditChangeProperty(EmptyPropertyUpdateStruct);
+#endif
+	}
+	
+	TextureSetupInfo.Material = Material;
+}
+
+void FPVTextureSetupImplementation::CreateTrimSheetTexture( FPVTextureSetupParams& Params, FPVTextureSetupInfo& TextureSetupInfo)
+{
+	TextureSetupInfo.Channels.Empty();
+	TextureSetupInfo.GenerationUVs.Empty();
+	//Create one texture per Channel with given Resolution
+	int32 Resolution = static_cast<int32>(Params.Resolution);
+	TArray<FString> Channels = GetChannels(Params.Generations);
+
+	TMap<FString, TObjectPtr<UTexture>> TextureSet = CreateTextures(Resolution, Channels, Params.AssetNamePrefix, Params.bVirtualTexture);
+	AddGenerationIntoChannel(TextureSet, Params);
+
+	float PreviousOffsetX = 0.0f;
+	for (auto Generation : Params.Generations)
+	{
+		float Scale = Generation.XScale / Params.GetWidthRatio(Generation);
+		FPVGenerationUVRange UV;
+		float DilationFactor = (Generation.Dilation * Scale);
+		UV.OffsetXStart = PreviousOffsetX + DilationFactor;
+		UV.OffsetXEnd = PreviousOffsetX + (Scale - DilationFactor);
+		UV.DilationFactor = DilationFactor;
+		TextureSetupInfo.GenerationUVs.Add(UV);
+		PreviousOffsetX += Scale;
+	}
+	
+	TextureSetupInfo.Channels = MoveTemp(TextureSet);
+	TextureSetupInfo.Resolution = Resolution;
+
+	SetupMaterial(Params.Material, TextureSetupInfo);
+
+	TextureSetupInfo.PreviewChannelName = StaticEnum<EPVTextureChannel>()->GetNameStringByValue((int64)Params.PreviewChannel);
+}
+
+UTexture2D* FPVTextureSetupImplementation::ConvertRenderTargetToTexture2D(UTextureRenderTarget2D* InRenderTarget, const FString& Path)
+{
+	if (!InRenderTarget)
+	{
+		return nullptr;
+	}
+	
+#if WITH_EDITOR
+	FString AssetName = InRenderTarget->GetName();
+	FString PackageName = Path / AssetName;
+	FString ObjectPath = PackageName + TEXT(".") + AssetName; 
+
+	UTexture2D* Texture = LoadObject<UTexture2D>(nullptr, ObjectPath);
+	UPackage* Package = nullptr;
+
+	if (Texture)
+	{
+		Package = Texture->GetPackage();
+	}
+	else
+	{
+		Package = CreatePackage(*PackageName);
+		Texture = NewObject<UTexture2D>(Package, FName(AssetName), RF_Public | RF_Standalone);
+		FAssetRegistryModule::AssetCreated(Texture);
+	}
+	
+	if (!InRenderTarget->UpdateTexture(Texture))
+	{
+		return Texture;
+	}
+	
+	Texture->Modify();
+	Texture->VirtualTextureStreaming = InRenderTarget->VirtualTextureStreaming;
+	bool bDirty = Texture->MarkPackageDirty();
+	Texture->PostEditChange();
+	Texture->UpdateResource();
+	
+	const FString PackageFileName = FPackageName::LongPackageNameToFilename(PackageName, FPackageName::GetAssetPackageExtension());
+	FSavePackageArgs SavePackageArgs;
+	SavePackageArgs.TopLevelFlags = RF_Public | RF_Standalone;
+
+	const bool bSave = UPackage::SavePackage(Package, Texture, *PackageFileName, SavePackageArgs);
+
+	return Texture;
+#endif
+
+	return nullptr;
+}
+
+#undef LOCTEXT_NAMESPACE
+
diff --git a/Engine/Plugins/Experimental/ProceduralVegetationEditor/Source/ProceduralVegetation/Private/Implementations/PVTextureSetup.h b/Engine/Plugins/Experimental/ProceduralVegetationEditor/Source/ProceduralVegetation/Private/Implementations/PVTextureSetup.h
new file mode 100644
--- /dev/null
+++ b/Engine/Plugins/Experimental/ProceduralVegetationEditor/Source/ProceduralVegetation/Private/Implementations/PVTextureSetup.h
@@ -0,0 +1,38 @@
+// Copyright Epic Games, Inc. All Rights Reserved.
+
+#pragma once
+
+#include "CoreMinimal.h"
+#include "RenderGraphBuilder.h"
+#include "TextureResource.h"
+#include "DataTypes/PVTextureSetupData.h"
+#include "Params/PVTextureSetupParams.h"
+
+class UTexture2D;
+class UTexture;
+class UTextureRenderTarget2D;
+class UMaterialInterface;
+
+struct FPVTextureSetupImplementation
+{
+	static TArray<FString> GetChannels(const TArray<FPVTextureGenerationParams>& Array);
+	static UTexture* CreateTexture(int Width, int Height, bool bCanCreateUAV, FName Name);
+	static TMap<FString, TObjectPtr<UTexture>> CreateTextures(int32 Resolution, const TArray<FString>& ChannelNames, const FString& AssetNamePrefix, bool bVirtualTextureStreaming);
+	static FRHITexture* GetRHITexture(UTexture* Texture);
+	static FRDGTextureSRVRef GetTextureSRVRef(FRDGBuilder& GraphBuilder, UTexture* InputTexture, const TCHAR* Name);
+	static FRDGTextureUAVRef GetTextureUAVRef(FRDGBuilder& GraphBuilder, UTexture* Texture, const TCHAR* Name);
+
+	static void DilationPass(FRDGBuilder& GraphBuilder, float DilationFactor, float TileX, float TileY, UTexture* InputTexture, UTexture* OutputTexture);
+	static void ScaleAndCombineTexturePass(FRDGBuilder& GraphBuilder, UTexture* TargetTexture, UTexture* InputTexture, const float WidthRatio, const float OffsetX,
+	                            const FPVTextureGenerationParams& Generation);
+	static void AddChannelTextureToTextureSet(const float InOffsetX, const FPVTextureGenerationParams& Generation, const FPVTextureChannelParams& Channel, UTexture* Texture);
+	static void AddGenerationIntoChannel(const FString& ChannelName, UTexture* Texture, const FPVTextureSetupParams& Params);
+	static void AddGenerationIntoChannel(const TMap<FString, TObjectPtr<UTexture>>& Map, const FPVTextureSetupParams& Params);
+	
+	static FString NormalizeString(FString InString);
+	static EPVTextureChannel GetTextureChannelFromString(const FString& InString);
+	static TObjectPtr<UTexture>  GetChannelTextureForMaterialParam(FPVTextureSetupInfo& TextureSetupInfo, FString MaterialParamName);
+	static void SetupMaterial(TObjectPtr<UMaterialInterface> Material, FPVTextureSetupInfo& TextureSetupInfo);
+	static void CreateTrimSheetTexture( FPVTextureSetupParams& Params, FPVTextureSetupInfo& TextureSetupInfo);
+	static UTexture2D* ConvertRenderTargetToTexture2D(UTextureRenderTarget2D* InRenderTarget, const FString& Path);
+};
diff --git a/Engine/Plugins/Experimental/ProceduralVegetationEditor/Source/ProceduralVegetation/Private/Implementations/PVTextureSetupCS.h b/Engine/Plugins/Experimental/ProceduralVegetationEditor/Source/ProceduralVegetation/Private/Implementations/PVTextureSetupCS.h
new file mode 100644
--- /dev/null
+++ b/Engine/Plugins/Experimental/ProceduralVegetationEditor/Source/ProceduralVegetation/Private/Implementations/PVTextureSetupCS.h
@@ -0,0 +1,43 @@
+// Copyright Epic Games, Inc. All Rights Reserved.
+
+#pragma once
+
+#include "CoreMinimal.h"
+#include "GlobalShader.h"
+#include "ShaderParameterStruct.h"
+
+class FPVTextureSetupDilationCS : public FGlobalShader
+{
+	DECLARE_GLOBAL_SHADER(FPVTextureSetupDilationCS);
+	SHADER_USE_PARAMETER_STRUCT(FPVTextureSetupDilationCS, FGlobalShader);
+
+	BEGIN_SHADER_PARAMETER_STRUCT(FParameters, )
+	SHADER_PARAMETER(float, OffsetX)
+	SHADER_PARAMETER(float, ScaleX)
+	SHADER_PARAMETER(float, TileX)
+	SHADER_PARAMETER(float, TileY)
+	SHADER_PARAMETER(FIntPoint, TargetSize)
+	SHADER_PARAMETER_RDG_TEXTURE_SRV(Texture2D<float4>, Source)
+	SHADER_PARAMETER_SAMPLER(SamplerState, SourceSampler)
+	SHADER_PARAMETER_RDG_TEXTURE_UAV(RWTexture2D<float4>, Target)
+	END_SHADER_PARAMETER_STRUCT()
+};
+
+class FPVTextureSetupScaleCS : public FGlobalShader
+{
+	DECLARE_GLOBAL_SHADER(FPVTextureSetupScaleCS);
+	SHADER_USE_PARAMETER_STRUCT(FPVTextureSetupScaleCS, FGlobalShader);
+
+	BEGIN_SHADER_PARAMETER_STRUCT(FParameters, )
+	SHADER_PARAMETER(float, OffsetX)
+	SHADER_PARAMETER(float, ScaleX)
+	SHADER_PARAMETER(float, TileX)
+	SHADER_PARAMETER(float, TileY)
+	SHADER_PARAMETER(FIntPoint, TargetSize)
+	SHADER_PARAMETER_RDG_TEXTURE_SRV(Texture2D<float4>, Source)
+	SHADER_PARAMETER_SAMPLER(SamplerState, SourceSampler)
+	SHADER_PARAMETER_RDG_TEXTURE_UAV(RWTexture2D<float4>, Target)
+	END_SHADER_PARAMETER_STRUCT()
+};
+
+
diff --git a/Engine/Plugins/Experimental/ProceduralVegetationEditor/Source/ProceduralVegetation/Private/Nodes/PVMeshBuilderSettings.cpp b/Engine/Plugins/Experimental/ProceduralVegetationEditor/Source/ProceduralVegetation/Private/Nodes/PVMeshBuilderSettings.cpp
--- a/Engine/Plugins/Experimental/ProceduralVegetationEditor/Source/ProceduralVegetation/Private/Nodes/PVMeshBuilderSettings.cpp
+++ b/Engine/Plugins/Experimental/ProceduralVegetationEditor/Source/ProceduralVegetation/Private/Nodes/PVMeshBuilderSettings.cpp
@@ -1,12 +1,13 @@
 ﻿// Copyright Epic Games, Inc. All Rights Reserved.
 
 #include "PVMeshBuilderSettings.h"
-
+#include "ProceduralVegetationModule.h"
 #include "PVCommon.h"
 #include "DataTypes/PVData.h"
 #include "DataTypes/PVGrowthData.h"
 #include "DataTypes/PVMeshData.h"
 #include "Data/PCGBasePointData.h"
+#include "DataTypes/PVTextureSetupData.h"
 #include "DataTypes/PVPlantProfileData.h"
 #include "Facades/PVBranchFacade.h"
 #include "Facades/PVProfileFacade.h"
@@ -14,6 +15,8 @@
 #include "Helpers/PCGSettingsHelpers.h"
 #include "Helpers/PVAnalyticsHelper.h"
 #include "Implementations/PVMeshBuilder.h"
+#include "Implementations/PVTextureSetup.h"
+#include "Materials/MaterialInstanceDynamic.h"
 #include "Materials/MaterialInterface.h"
 
 #define LOCTEXT_NAMESPACE "PVMeshBuilderSettings"
@@ -141,6 +144,11 @@ TArray<FPCGPinProperties> UPVMeshBuilderSettings::InputPinProperties() const
 	PlantProfilePin.SetAllowMultipleConnections(false);
 	PlantProfilePin.SetAdvancedPin();
 
+	FPCGPinProperties& PinTextureSetup = Properties.Emplace_GetRef(TEXT("Texture Setup"), FPCGDataTypeIdentifier{ FPVDataTypeInfoTextureSetup::AsId()});
+	PinTextureSetup.SetAllowMultipleConnections(false);
+	PinTextureSetup.SetOverrideOrUserParamPin();
+	PinTextureSetup.bAllowMultipleData = false;
+	
 	return Properties;
 }
 
@@ -167,6 +175,7 @@ bool FPVMeshBuilderElement::ExecuteInternal(FPCGContext* InContext) const
 	}
 
 	const TArray<FPCGTaggedData>& Inputs = InContext->InputData.GetInputsByPin(PCGPinConstants::DefaultInputLabel);
+	const TArray<FPCGTaggedData>& InputTextureSetups = InContext->InputData.GetInputsByPin(TEXT("Texture Setup"));
 	if (!Inputs.IsEmpty())
 	{
 		if (const UPVGrowthData* InputData = Cast<UPVGrowthData>(Inputs[0].Data))
@@ -180,11 +189,24 @@ bool FPVMeshBuilderElement::ExecuteInternal(FPCGContext* InContext) const
 				const FString TrunkMaterial = BranchFacade.GetTrunkMaterialPath();
 				const TArray<FVector2f> TrunkURanges = BranchFacade.GetTrunkURange();
 
-				int32 Generation = 0;
-				for (const FVector2f& URange : TrunkURanges)
+				if (InputTextureSetups.IsValidIndex(0) && Cast<UPVTextureSetupData>(InputTextureSetups[0].Data))
+				{
+					const UPVTextureSetupData* TextureSetup = Cast<UPVTextureSetupData>(InputTextureSetups[0].Data);
+					int32 Generation = 0;
+					for (const FPVGenerationUVRange& GenerationUVRange : TextureSetup->TextureSetupInfo.GenerationUVs)
+					{
+						MaterialSettings.SetMaterial(TextureSetup->TextureSetupInfo.Material, FVector2f(GenerationUVRange.OffsetXStart, GenerationUVRange.OffsetXEnd), Generation);
+						Generation++;
+					}
+				}
+				else
 				{
-					MaterialSettings.SetMaterial(TrunkMaterial, URange, Generation);
-					Generation++;
+					int32 Generation = 0;
+					for (const FVector2f& URange : TrunkURanges)
+					{
+						MaterialSettings.SetMaterial(TrunkMaterial, URange, Generation);
+						Generation++;
+					}
 				}
 
 				MutableSettings->DefaultMaterialSettings = MaterialSettings;
diff --git a/Engine/Plugins/Experimental/ProceduralVegetationEditor/Source/ProceduralVegetation/Private/Nodes/PVTextureSetupSettings.cpp b/Engine/Plugins/Experimental/ProceduralVegetationEditor/Source/ProceduralVegetation/Private/Nodes/PVTextureSetupSettings.cpp
new file mode 100644
--- /dev/null
+++ b/Engine/Plugins/Experimental/ProceduralVegetationEditor/Source/ProceduralVegetation/Private/Nodes/PVTextureSetupSettings.cpp
@@ -0,0 +1,136 @@
+﻿// Copyright Epic Games, Inc. All Rights Reserved.
+ 
+#include "Nodes/PVTextureSetupSettings.h"
+#include "DataTypes/PVGrowthData.h"
+#include "Engine/TextureRenderTarget2D.h"
+#include "Engine/Texture2D.h"
+#include "Engine/Texture.h"
+#include "Implementations/PVTextureSetup.h"
+
+class UTexture2D;
+
+UPVTextureSetupSettings::UPVTextureSetupSettings()
+{
+	TextureSetupParams.ContentBrowserFolder.Path = "/Game/ProceduralVegetation/";
+}
+
+#if WITH_EDITOR
+void UPVTextureSetupSettings::PostEditChangeChainProperty(struct FPropertyChangedChainEvent& PropertyChangedEvent)
+{
+	Super::PostEditChangeChainProperty(PropertyChangedEvent);
+	
+	if (!ensure(PropertyChangedEvent.PropertyChain.GetActiveNode()))
+	{
+		return;
+	}
+	
+	auto Node = (PropertyChangedEvent.PropertyChain.GetActiveNode()->GetPrevNode());
+	
+	if (Node && Node->GetValue()->GetName() == GET_MEMBER_NAME_CHECKED(FPVTextureChannelParams, TextureFilePath))
+	{
+		int GenerationIndex = PropertyChangedEvent.GetArrayIndex(GET_MEMBER_NAME_CHECKED(FPVTextureSetupParams, Generations).ToString());
+		int ChannelIndex = PropertyChangedEvent.GetArrayIndex(GET_MEMBER_NAME_CHECKED(FPVTextureGenerationParams, Channels).ToString());
+
+		TextureSetupParams.LoadGenerationChannelTexture(GenerationIndex, ChannelIndex);
+	}
+
+	if (PropertyChangedEvent.ChangeType == EPropertyChangeType::ValueSet)
+	{
+		FPVTextureSetupImplementation::CreateTrimSheetTexture(TextureSetupParams, TextureSetupInfo);
+		bTexturesBaked = false;
+	}
+}
+
+void UPVTextureSetupSettings::PostLoad()
+{
+	Super::PostLoad();
+
+	int GenerationIndex = 0;
+	for (auto Generation : TextureSetupParams.Generations)
+	{
+		int ChannelIndex = 0;
+		for (auto Channel : Generation.Channels)
+		{
+			TextureSetupParams.LoadGenerationChannelTexture(GenerationIndex, ChannelIndex);
+			ChannelIndex++;
+		}
+		GenerationIndex++;
+	}
+
+	FPVTextureSetupImplementation::CreateTrimSheetTexture(TextureSetupParams, TextureSetupInfo);
+}
+
+void UPVTextureSetupSettings::PostEditUndo()
+{
+	Super::PostEditUndo();
+	FPVTextureSetupImplementation::CreateTrimSheetTexture(TextureSetupParams, TextureSetupInfo);
+}
+
+#endif
+
+TArray<FPCGPinProperties> UPVTextureSetupSettings::InputPinProperties() const
+{
+	return TArray<FPCGPinProperties>();
+}
+ 
+FPCGElementPtr UPVTextureSetupSettings::CreateElement() const
+{
+	return MakeShared<FPVTextureSetupElement>();
+}
+ 
+FPCGDataTypeIdentifier UPVTextureSetupSettings::GetOutputPinTypeIdentifier() const
+{
+	return FPCGDataTypeIdentifier{ FPVDataTypeInfoTextureSetup::AsId() };
+}
+
+void UPVTextureSetupSettings::BakeTextures()
+{
+#if WITH_EDITOR
+	for (auto& Channel : TextureSetupInfo.Channels)
+	{
+		if (UTexture* ConvertedTexture = Cast<UTexture>(FPVTextureSetupImplementation::ConvertRenderTargetToTexture2D(Cast<UTextureRenderTarget2D>(Channel.Value),
+				TextureSetupParams.ContentBrowserFolder.Path)))
+		{
+			Channel.Value = ConvertedTexture;
+		}
+	}
+
+	FPVTextureSetupImplementation::SetupMaterial(TextureSetupParams.Material, TextureSetupInfo);
+	bTexturesBaked = true;
+#else
+	UE_LOG(LogTemp, Display, TEXT("Baking Textures is not supported"));
+#endif
+}
+
+bool FPVTextureSetupElement::ExecuteInternal(FPCGContext* InContext) const
+{
+	TRACE_CPUPROFILER_EVENT_SCOPE(FPVImporterElement::Execute);
+ 
+	check(InContext);
+ 
+	const UPVTextureSetupSettings* Settings = InContext->GetInputSettings<UPVTextureSetupSettings>();
+	check(Settings);
+	
+	UPVTextureSetupData* OutTextureSetupData = FPCGContext::NewObject_AnyThread<UPVTextureSetupData>(InContext);
+	
+	OutTextureSetupData->TextureSetupInfo = Settings->TextureSetupInfo;
+	InContext->OutputData.TaggedData.Emplace(OutTextureSetupData);
+
+	int GenerationIndex = 0;
+	for (auto Generation : Settings->TextureSetupParams.Generations)
+	{
+		if (!Generation.IsWidthRatioSame())
+		{
+			PCGLog::LogWarningOnGraph(FText::FromString(FString::Printf(TEXT("Generations %i dose not have the same width to height ratio for all channels. "
+																	"\nMake sure all channels in a generation have same width height ratio"), GenerationIndex)), InContext);
+		}
+		GenerationIndex++;
+	}
+	
+	if (Settings->TextureSetupInfo.IsOverflowing())
+	{
+		PCGLog::LogWarningOnGraph(FText::FromString(TEXT("Generations are not fitting in the texture, consider scaling down the XScale values in generation settings.")), InContext);
+	}
+	
+	return true;
+}
diff --git a/Engine/Plugins/Experimental/ProceduralVegetationEditor/Source/ProceduralVegetation/Private/Params/PVTextureSetupParams.cpp b/Engine/Plugins/Experimental/ProceduralVegetationEditor/Source/ProceduralVegetation/Private/Params/PVTextureSetupParams.cpp
new file mode 100644
--- /dev/null
+++ b/Engine/Plugins/Experimental/ProceduralVegetationEditor/Source/ProceduralVegetation/Private/Params/PVTextureSetupParams.cpp
@@ -0,0 +1,135 @@
+// Copyright Epic Games, Inc. All Rights Reserved.
+
+#include "Params/PVTextureSetupParams.h"
+
+#include "ImageUtils.h"
+#include "HAL/PlatformFileManager.h"
+#include "Engine/Texture2D.h"
+
+void FPVTextureChannelParams::LoadTexture()
+{
+	if (FPlatformFileManager::Get().GetPlatformFile().FileExists(* TextureFilePath.FilePath))
+	{
+		if (UTexture2D* OriginalTexture = FImageUtils::ImportFileAsTexture2D(* TextureFilePath.FilePath))
+		{
+			Texture = Cast<UTexture>(OriginalTexture);
+			Texture->UpdateResource();
+
+			if (ensure(Texture->GetSurfaceWidth() > 0))
+			{
+				WidthRatio = Texture->GetSurfaceHeight() / (float) Texture->GetSurfaceWidth();
+			}
+			else
+			{
+				ResetTexture();
+			}
+		}
+		else
+		{
+			ResetTexture();
+		}
+	}
+	else
+	{
+		ResetTexture();
+	}
+}
+
+void FPVTextureChannelParams::ResetTexture()
+{
+	Texture = nullptr;
+	WidthRatio = 1.0f;
+}
+
+FString FPVTextureChannelParams::GetChannelName() const
+{
+	if (EPVTextureChannel::Custom == Channel)
+	{
+		return ChannelName;
+	}
+	else
+	{
+		return  StaticEnum<EPVTextureChannel>()->GetNameStringByValue((int64)Channel);
+	}
+}
+
+void FPVTextureGenerationParams::LoadChannelTexture(int ChannelIndex)
+{
+	if (Channels.IsValidIndex(ChannelIndex))
+	{
+		FPVTextureChannelParams& Channel = Channels[ChannelIndex];
+		Channel.LoadTexture();
+	}
+}
+
+TArray<FString> FPVTextureGenerationParams::GetChannelNames()
+{
+	TArray<FString> ChannelNames;
+	for (const auto& Channel : Channels)
+	{
+		ChannelNames.Add(Channel.GetChannelName());
+	}
+
+	return ChannelNames;
+}
+
+const FPVTextureChannelParams* FPVTextureGenerationParams::GetChannel(const FString& Name)
+{
+	FPVTextureChannelParams* FoundChannel = nullptr;
+	for (auto& Channel : Channels)
+	{
+		if (Channel.GetChannelName() == Name)
+		{
+			FoundChannel = &Channel;
+		}
+	}
+
+	return FoundChannel;
+}
+
+float FPVTextureGenerationParams::GetWidthRatio() const
+{
+	if (Channels.IsValidIndex(0))
+	{
+		return Channels[0].WidthRatio;
+	}
+
+	return 1.0;
+}
+
+bool FPVTextureGenerationParams::IsWidthRatioSame() const
+{
+	if (!Channels.IsValidIndex(0))
+	{
+		return true;
+	}
+	
+	float PreviousWidthRatio = Channels[0].WidthRatio;
+	for (int Channel = 1; Channel < Channels.Num(); ++Channel)
+	{
+		if (!FMath::IsNearlyEqual(PreviousWidthRatio,Channels[Channel].WidthRatio))
+		{
+			return false;
+		}
+	}
+	return true;
+}
+
+void FPVTextureSetupParams::LoadGenerationChannelTexture(int GenerationIndex, int ChannelIndex)
+{
+	if (Generations.IsValidIndex(GenerationIndex))
+	{
+		FPVTextureGenerationParams& Generation = Generations[GenerationIndex];
+		Generation.LoadChannelTexture(ChannelIndex);
+	}
+}
+
+float FPVTextureSetupParams::GetWidthRatio(FPVTextureGenerationParams Generation) const
+{
+	return Generation.GetWidthRatio();
+}
+
+bool FPVTextureSetupParams::IsWidthRatioSame(FPVTextureGenerationParams Generation) const
+{
+	return Generation.IsWidthRatioSame();
+}
diff --git a/Engine/Plugins/Experimental/ProceduralVegetationEditor/Source/ProceduralVegetation/ProceduralVegetation.Build.cs b/Engine/Plugins/Experimental/ProceduralVegetationEditor/Source/ProceduralVegetation/ProceduralVegetation.Build.cs
--- a/Engine/Plugins/Experimental/ProceduralVegetationEditor/Source/ProceduralVegetation/ProceduralVegetation.Build.cs
+++ b/Engine/Plugins/Experimental/ProceduralVegetationEditor/Source/ProceduralVegetation/ProceduralVegetation.Build.cs
@@ -33,7 +33,7 @@ namespace UnrealBuildTool.Rules
 					"StaticMeshDescription",
 					"RHI",
 					"Projects",
-					"ImageCore",
+					"ImageCore"
 				}
 			);
 		}
diff --git a/Engine/Plugins/Experimental/ProceduralVegetationEditor/Source/ProceduralVegetation/Public/DataTypes/PVData.h b/Engine/Plugins/Experimental/ProceduralVegetationEditor/Source/ProceduralVegetation/Public/DataTypes/PVData.h
--- a/Engine/Plugins/Experimental/ProceduralVegetationEditor/Source/ProceduralVegetation/Public/DataTypes/PVData.h
+++ b/Engine/Plugins/Experimental/ProceduralVegetationEditor/Source/ProceduralVegetation/Public/DataTypes/PVData.h
@@ -23,6 +23,7 @@ enum class EPVRenderType : uint8
 	FoliageGrid 		= 1 << 4,
 	FoliageAttachments	= 1 << 5,
 	Seed				= 1 << 6,
+	Textures			= 1 << 7,
 
 	// Add new renter type above
 	Count	
diff --git a/Engine/Plugins/Experimental/ProceduralVegetationEditor/Source/ProceduralVegetation/Public/DataTypes/PVTextureSetupData.h b/Engine/Plugins/Experimental/ProceduralVegetationEditor/Source/ProceduralVegetation/Public/DataTypes/PVTextureSetupData.h
new file mode 100644
--- /dev/null
+++ b/Engine/Plugins/Experimental/ProceduralVegetationEditor/Source/ProceduralVegetation/Public/DataTypes/PVTextureSetupData.h
@@ -0,0 +1,93 @@
+// Copyright Epic Games, Inc. All Rights Reserved.
+
+#pragma once
+
+#include "PVData.h"
+#include "PVTextureSetupData.generated.h"
+
+class UMaterialInterface;
+class UTexture;
+
+UENUM()
+enum class EPVTextureChannel : uint8
+{
+	BaseColor,
+	AO,
+	Normal,
+	Displacement,
+	Bump,
+	Cavity,
+	Gloss,
+	Roughness,
+	Specular,
+	Custom,
+};
+
+USTRUCT()
+struct FPVGenerationUVRange
+{
+	GENERATED_BODY()
+	
+	UPROPERTY()
+	float OffsetXStart = 0.0f;
+
+	UPROPERTY()
+	float OffsetXEnd = 1.0f;
+
+	UPROPERTY()
+	float DilationFactor = 0.0f;
+
+	inline float GetOverflowValue() const {return (OffsetXEnd + DilationFactor) - 1.0;}
+};
+
+USTRUCT()
+struct FPVTextureSetupInfo
+{
+	GENERATED_BODY()
+
+	UPROPERTY()
+	TMap<FString, TObjectPtr<UTexture>> Channels;
+
+	UPROPERTY()
+	TArray<FPVGenerationUVRange> GenerationUVs;
+
+	UPROPERTY()
+	TObjectPtr<UMaterialInterface> Material;
+
+	UPROPERTY()
+	int Resolution = 1024;
+
+	UPROPERTY()
+	FString PreviewChannelName;
+
+	bool IsOverflowing() const;
+};
+
+USTRUCT()
+struct FPVDataTypeInfoTextureSetup : public FPCGDataTypeInfo
+{
+	GENERATED_BODY()
+
+	PCG_DECLARE_TYPE_INFO(PROCEDURALVEGETATION_API);
+
+#if WITH_EDITOR
+	// If the type should be hidden to the user.
+	virtual bool Hidden() const override { return true; };
+#endif // WITH_EDITOR
+};
+
+UCLASS(BlueprintType, ClassGroup = (Procedural))
+class UPVTextureSetupData : public UPVData
+{
+	GENERATED_BODY()
+
+public:
+	UPROPERTY()
+	FPVTextureSetupInfo TextureSetupInfo;
+	//~ Begin UPCGData interface
+	PCG_ASSIGN_DEFAULT_TYPE_INFO(FPVDataTypeInfoTextureSetup)
+
+	virtual void AddToCrc(FArchiveCrc32& Ar, bool bFullDataCrc) const override;
+	
+	//~ End UPCGData interface
+};
\ No newline at end of file
diff --git a/Engine/Plugins/Experimental/ProceduralVegetationEditor/Source/ProceduralVegetation/Public/Nodes/PVTextureSetupSettings.h b/Engine/Plugins/Experimental/ProceduralVegetationEditor/Source/ProceduralVegetation/Public/Nodes/PVTextureSetupSettings.h
new file mode 100644
--- /dev/null
+++ b/Engine/Plugins/Experimental/ProceduralVegetationEditor/Source/ProceduralVegetation/Public/Nodes/PVTextureSetupSettings.h
@@ -0,0 +1,55 @@
+﻿// Copyright Epic Games, Inc. All Rights Reserved.
+#pragma once
+
+#include "CoreMinimal.h"
+#include "Params/PVTextureSetupParams.h"
+#include "Nodes/PVBaseSettings.h"
+#include "PVTextureSetupSettings.generated.h"
+
+UCLASS(BlueprintType, ClassGroup = (Procedural))
+class UPVTextureSetupSettings : public UPVBaseSettings
+{
+	GENERATED_BODY()
+
+public:
+	UPVTextureSetupSettings();
+
+#if WITH_EDITOR
+	virtual FName GetDefaultNodeName() const override { return FName(TEXT("TextureSetup")); }
+	virtual FText GetDefaultNodeTitle() const override { return NSLOCTEXT("PVTextureSetupSettings", "NodeTitle", "Texture Setup"); }
+	virtual FText GetNodeTooltipText() const override { return NSLOCTEXT("PVTextureSetupSettings", "NodeTooltip", "Create Trim Sheet textures and also output the trim sheet UV details"); }
+	virtual FLinearColor GetNodeTitleColor() const override { return FLinearColor(FColor::Black); }
+
+	virtual TArray<EPVRenderType> GetDefaultRenderType() const override { return TArray{ EPVRenderType::Textures }; }
+	virtual void PostEditChangeChainProperty(struct FPropertyChangedChainEvent& PropertyChangedEvent) override;
+	virtual void PostLoad() override;
+	virtual void PostEditUndo() override;
+#endif
+
+	PROCEDURALVEGETATION_API bool GetTextureBaked() const { return bTexturesBaked; }
+
+protected:
+	virtual TArray<FPCGPinProperties> InputPinProperties() const override;
+	virtual FPCGDataTypeIdentifier GetOutputPinTypeIdentifier() const override;
+	virtual FPCGElementPtr CreateElement() const override;
+
+public:
+	UFUNCTION(Category=NoCategory)
+	PROCEDURALVEGETATION_API void BakeTextures();
+	
+	UPROPERTY(EditAnywhere, Category=TextureSetup, meta=(RelativeToGameDir, ShowOnlyInnerProperties))
+	FPVTextureSetupParams TextureSetupParams;
+
+	UPROPERTY()
+	FPVTextureSetupInfo TextureSetupInfo;
+
+private:
+	UPROPERTY()
+	bool bTexturesBaked = false;
+};
+
+class FPVTextureSetupElement : public FPVBaseElement
+{
+protected:
+	virtual bool ExecuteInternal(FPCGContext* InContext) const override;
+};
diff --git a/Engine/Plugins/Experimental/ProceduralVegetationEditor/Source/ProceduralVegetation/Public/Params/PVTextureSetupParams.h b/Engine/Plugins/Experimental/ProceduralVegetationEditor/Source/ProceduralVegetation/Public/Params/PVTextureSetupParams.h
new file mode 100644
--- /dev/null
+++ b/Engine/Plugins/Experimental/ProceduralVegetationEditor/Source/ProceduralVegetation/Public/Params/PVTextureSetupParams.h
@@ -0,0 +1,113 @@
+// Copyright Epic Games, Inc. All Rights Reserved.
+#pragma once
+
+#include "CoreMinimal.h"
+#include "RenderGraphBuilder.h"
+#include "TextureResource.h"
+#include "DataTypes/PVTextureSetupData.h"
+#include "Materials/MaterialInstance.h"
+#include "PVTextureSetupParams.generated.h"
+
+UENUM()
+enum class EPVTextureResolution 
+{
+	Resolution512 = 512 UMETA(DisplayName = "512 x 512"),
+	Resolution1024 = 1024 UMETA(DisplayName = "1024 x 1024"),
+	Resolution2048 = 2048 UMETA(DisplayName = "2048 x 2048"),
+	Resolution4096 = 4096 UMETA(DisplayName = "4096 x 4096"),
+	Resolution8192 = 8192 UMETA(DisplayName = "8192 x 8192"),
+	Resolution16384 = 16384 UMETA(DisplayName = "16384 x 16384"),
+};
+
+USTRUCT()
+struct FPVTextureChannelParams
+{
+	GENERATED_BODY()
+
+	UPROPERTY(EditAnywhere, Category=Channel)
+	EPVTextureChannel Channel = EPVTextureChannel::BaseColor;
+
+	UPROPERTY(EditAnywhere, Category=Channel, meta=(EditCondition="Channel == EPVTextureChannel::Custom"))
+	FString ChannelName;
+
+	UPROPERTY(EditAnywhere, Category= Channel, meta = (AllowPrivateAccess = true, FilePathFilter = "Image Files (*.png;*.PNG;*.jpg;*.jpeg;*.exr;)|*.png;*.PNG;*.jpg;*.jpeg;*.exr;", RelativeToGameDir))
+	FFilePath TextureFilePath;
+
+	UPROPERTY(VisibleAnywhere, Category=Channel)
+	TObjectPtr<UTexture> Texture;
+
+	float WidthRatio = 1.0;
+
+	void LoadTexture();
+	
+	void ResetTexture();
+
+	FString GetChannelName() const;
+};
+
+USTRUCT()
+struct FPVTextureGenerationParams
+{
+	GENERATED_BODY()
+
+	UPROPERTY(EditAnywhere, Category= Generation)
+	FDirectoryPath BakeTextureFolder;
+
+	UPROPERTY(EditAnywhere, Category= Generation)
+	float Dilation = 0.1f;
+
+	UPROPERTY(EditAnywhere, Category= Generation)
+	float XScale = 1.0f;
+
+	UPROPERTY(EditAnywhere, Category= Generation)
+	float BaseTileX = 1.0f;
+
+	UPROPERTY(EditAnywhere, Category= Generation)
+	float BaseTileY = 1.0f;
+
+	UPROPERTY( EditAnywhere, Category= Generation)
+	TArray<FPVTextureChannelParams> Channels;
+
+	void LoadChannelTexture(int ChannelIndex);
+
+	TArray<FString> GetChannelNames();
+
+	const FPVTextureChannelParams* GetChannel(const FString& Name);
+
+	float GetWidthRatio() const;
+	
+	bool IsWidthRatioSame() const;
+};
+
+USTRUCT()
+struct FPVTextureSetupParams
+{
+	GENERATED_BODY()
+
+	UPROPERTY(EditAnywhere, Category= NoCategory)
+	FString AssetNamePrefix;
+
+	UPROPERTY(EditAnywhere, Category= NoCategory)
+	FDirectoryPath ContentBrowserFolder;
+
+	UPROPERTY( EditAnywhere, Category= NoCategory)
+	bool bVirtualTexture = false;
+
+	UPROPERTY( EditAnywhere, Category= NoCategory)
+	TObjectPtr<UMaterialInstance> Material;
+
+	UPROPERTY(EditAnywhere, Category=TextureSetup)
+	EPVTextureChannel PreviewChannel = EPVTextureChannel::BaseColor;
+
+	UPROPERTY(EditAnywhere, Category=TextureSetup)
+	EPVTextureResolution Resolution = EPVTextureResolution::Resolution1024;
+
+	UPROPERTY(EditAnywhere, Category= TextureSetup)
+	TArray<FPVTextureGenerationParams> Generations;
+
+	void LoadGenerationChannelTexture(int GenerationIndex, int ChannelIndex);
+
+	float GetWidthRatio(FPVTextureGenerationParams Generation) const;
+	
+	bool IsWidthRatioSame(FPVTextureGenerationParams Generation) const;
+};
diff --git a/Engine/Plugins/Experimental/ProceduralVegetationEditor/Source/ProceduralVegetationEditor/Private/Customizations/PVTextureSetupSettingsCustomization.cpp b/Engine/Plugins/Experimental/ProceduralVegetationEditor/Source/ProceduralVegetationEditor/Private/Customizations/PVTextureSetupSettingsCustomization.cpp
new file mode 100644
--- /dev/null
+++ b/Engine/Plugins/Experimental/ProceduralVegetationEditor/Source/ProceduralVegetationEditor/Private/Customizations/PVTextureSetupSettingsCustomization.cpp
@@ -0,0 +1,95 @@
+// Copyright Epic Games, Inc. All Rights Reserved.
+
+#include "PVTextureSetupSettingsCustomization.h"
+
+#include "ClassIconFinder.h"
+#include "DetailWidgetRow.h"
+#include "IDetailChildrenBuilder.h"
+#include "PropertyHandle.h"
+
+#include "ActorFactories/ActorFactoryBasicShape.h"
+
+#include "AssetRegistry/IAssetRegistry.h"
+
+#include "Nodes/PVTextureSetupSettings.h"
+#include "Widgets/Input/SButton.h"
+#include "Widgets/Text/STextBlock.h"
+
+#define LOCTEXT_NAMESPACE "PVTextureSetupSettingsCustomization"
+
+TSharedRef<IDetailCustomization> FPVTextureSetupSettingsCustomization::MakeInstance()
+{
+	return MakeShareable(new FPVTextureSetupSettingsCustomization());
+}
+
+void FPVTextureSetupSettingsCustomization::CustomizeDetails(IDetailLayoutBuilder& DetailBuilder)
+{
+	SelectedSettings.Empty();
+	
+	const FName CategoryName("Baked Textures");
+	IDetailCategoryBuilder& Category = DetailBuilder.EditCategory(CategoryName);
+
+	FDetailWidgetRow& NewRow = Category.AddCustomRow(FText::GetEmpty());
+	
+	TArray<TWeakObjectPtr<UObject>> ObjectsBeingCustomized;
+	DetailBuilder.GetObjectsBeingCustomized(ObjectsBeingCustomized);
+
+	for (const TWeakObjectPtr<UObject>& Object : ObjectsBeingCustomized)
+	{
+		UPVTextureSetupSettings* Settings = Cast<UPVTextureSetupSettings>(Object.Get());
+		if (ensure(Settings))
+		{
+			SelectedSettings.Add(Settings);
+		}
+	}
+
+	if (SelectedSettings.Num() == 1)
+	{
+		NewRow.ValueContent()
+		.MaxDesiredWidth(120.f)
+		[
+			SNew(SHorizontalBox)
+			+ SHorizontalBox::Slot()
+			.AutoWidth()
+			.Padding(4.0f)
+			[
+				SNew(SButton)
+				.OnClicked(this, &FPVTextureSetupSettingsCustomization::OnBakeTexturesClicked)
+				.IsEnabled_Raw(this, &FPVTextureSetupSettingsCustomization::BakeTexturesButtonEnabled)
+				[
+					SNew(STextBlock)
+					.Font(IDetailLayoutBuilder::GetDetailFont())
+					.Text(LOCTEXT("BakeTexturesButton", "Bake Textures"))
+				]
+			]
+		];
+	}
+}
+
+bool FPVTextureSetupSettingsCustomization::BakeTexturesButtonEnabled() const
+{
+	UPVTextureSetupSettings* Settings = !SelectedSettings.IsEmpty() ? SelectedSettings[0].Get() : nullptr;
+	if (Settings)
+	{
+		return !Settings->GetTextureBaked();
+	}
+	
+	return true;
+}
+
+FReply FPVTextureSetupSettingsCustomization::OnBakeTexturesClicked()
+{
+	UPVTextureSetupSettings* Settings = !SelectedSettings.IsEmpty() ? SelectedSettings[0].Get() : nullptr;
+
+	if (!Settings || !IsValid(Settings))
+	{
+		return FReply::Handled();
+	}
+
+	Settings->BakeTextures();
+
+	return FReply::Handled();
+}
+
+#undef LOCTEXT_NAMESPACE
+
diff --git a/Engine/Plugins/Experimental/ProceduralVegetationEditor/Source/ProceduralVegetationEditor/Private/Customizations/PVTextureSetupSettingsCustomization.h b/Engine/Plugins/Experimental/ProceduralVegetationEditor/Source/ProceduralVegetationEditor/Private/Customizations/PVTextureSetupSettingsCustomization.h
new file mode 100644
--- /dev/null
+++ b/Engine/Plugins/Experimental/ProceduralVegetationEditor/Source/ProceduralVegetationEditor/Private/Customizations/PVTextureSetupSettingsCustomization.h
@@ -0,0 +1,23 @@
+// Copyright Epic Games, Inc. All Rights Reserved.
+
+#pragma once
+
+#include "DetailLayoutBuilder.h"
+#include "IDetailCustomization.h"
+#include "IDetailPropertyRow.h"
+#include "Nodes/PVTextureSetupSettings.h"
+
+class FPVTextureSetupSettingsCustomization : public IDetailCustomization
+{
+public:
+	static TSharedRef<IDetailCustomization> MakeInstance();
+
+	virtual void CustomizeDetails(IDetailLayoutBuilder& DetailBuilder) override;
+
+private:
+	bool BakeTexturesButtonEnabled() const;
+	FReply OnBakeTexturesClicked();
+
+private:
+	TArray<TWeakObjectPtr<UPVTextureSetupSettings>> SelectedSettings;
+};
diff --git a/Engine/Plugins/Experimental/ProceduralVegetationEditor/Source/ProceduralVegetationEditor/Private/ProceduralVegetationEditorModule.cpp b/Engine/Plugins/Experimental/ProceduralVegetationEditor/Source/ProceduralVegetationEditor/Private/ProceduralVegetationEditorModule.cpp
--- a/Engine/Plugins/Experimental/ProceduralVegetationEditor/Source/ProceduralVegetationEditor/Private/ProceduralVegetationEditorModule.cpp
+++ b/Engine/Plugins/Experimental/ProceduralVegetationEditor/Source/ProceduralVegetationEditor/Private/ProceduralVegetationEditorModule.cpp
@@ -29,6 +29,7 @@
 #include "DataTypes/PVGrowthData.h"
 #include "DataTypes/PVMeshData.h"
 #include "Params/PVExportParams.h"
+#include "Nodes/PVTextureSetupSettings.h"
 
 #include "Engine/SkeletalMesh.h"
 #include "Engine/StaticMesh.h"
@@ -41,6 +42,7 @@
 
 #include "Customizations/PVColliderParamsCustomization.h"
 #include "Customizations/PVLoopDebugStepperCustomization.h"
+#include "Customizations/PVTextureSetupSettingsCustomization.h"
 
 #include "Helpers/PVUtilities.h"
 
@@ -74,6 +76,8 @@ void FProceduralVegetationEditorModule::StartupModule()
 	PropertyEditor.RegisterCustomPropertyTypeLayout(FPVFloatRamp::StaticStruct()->GetFName(), FOnGetPropertyTypeCustomizationInstance::CreateStatic(&FPVFloatRampCustomization::MakeInstance));
 	PropertyEditor.RegisterCustomPropertyTypeLayout(FLoopDebugStepper::StaticStruct()->GetFName(), FOnGetPropertyTypeCustomizationInstance::CreateStatic(&FPVLoopDebugStepperCustomization::MakeInstance));
 
+	PropertyEditor.RegisterCustomClassLayout(UPVTextureSetupSettings::StaticClass()->GetFName(), FOnGetDetailCustomizationInstance::CreateStatic(&FPVTextureSetupSettingsCustomization::MakeInstance));
+	
 	FMessageLogModule& MessageLogModule = FModuleManager::LoadModuleChecked<FMessageLogModule>("MessageLog");
 	MessageLogModule.RegisterLogListing(PVEditor::MessageLogName, LOCTEXT("PVEditorLogLabel", "Procedural Vegetation Editor"));
 }
@@ -89,8 +93,11 @@ void FProceduralVegetationEditorModule::ShutdownModule()
 		if (FPropertyEditorModule* PropertyEditor = FModuleManager::GetModulePtr<FPropertyEditorModule>("PropertyEditor"))
 		{
 			PropertyEditor->UnregisterCustomPropertyTypeLayout(FPVExportParams::StaticStruct()->GetFName());
+			PropertyEditor->UnregisterCustomPropertyTypeLayout(FPVColliderParams::StaticStruct()->GetFName());
 			PropertyEditor->UnregisterCustomPropertyTypeLayout(FPVFloatRamp::StaticStruct()->GetFName());
 			PropertyEditor->UnregisterCustomPropertyTypeLayout(FLoopDebugStepper::StaticStruct()->GetFName());
+			PropertyEditor->UnregisterCustomClassLayout(UPVTextureSetupSettings::StaticClass()->GetFName());
+	
 		}
 	}
 
diff --git a/Engine/Plugins/Experimental/ProceduralVegetationEditor/Source/ProceduralVegetationEditor/Private/SPVEditorViewport.cpp b/Engine/Plugins/Experimental/ProceduralVegetationEditor/Source/ProceduralVegetationEditor/Private/SPVEditorViewport.cpp
--- a/Engine/Plugins/Experimental/ProceduralVegetationEditor/Source/ProceduralVegetationEditor/Private/SPVEditorViewport.cpp
+++ b/Engine/Plugins/Experimental/ProceduralVegetationEditor/Source/ProceduralVegetationEditor/Private/SPVEditorViewport.cpp
@@ -279,7 +279,7 @@ void SPVEditorViewport::PopulateViewportOverlays(TSharedRef<SOverlay> Overlay)
 				+ SVerticalBox::Slot()
 				.FillHeight(1.f)
 				[
-					SNew(SListView<TSharedPtr<FText>>)
+					SAssignNew(LegendListView, SListView<TSharedPtr<FText>>)
 					.ListItemsSource(&LegendData)
 					.ListViewStyle(&GetTransparentViewStyle())
 					.OnGenerateRow_Lambda(
@@ -562,6 +562,11 @@ void SPVEditorViewport::OnResetScene()
 
 	FoliageComponents->Empty();
 	LegendData.Empty();
+
+	if (LegendListView)
+	{
+		LegendListView->RequestListRefresh();
+	}
 }
 
 FToolMenuEntry SPVEditorViewport::CreateSettingsToolbarMenu()
@@ -944,6 +949,11 @@ void SPVEditorViewport::PopulateLegendOverlayText(const TArrayView<TPair<FColor,
 	{
 		LegendData.Add(MakeShared<FText>(FText::Format(LegendStyle, FText::FromString(Color.ToHex()), Text)));
 	}
+
+	if (LegendListView)
+	{
+		LegendListView->RequestListRefresh();
+	}
 }
 
 void SPVEditorViewport::AddInstancedFoliageComponent(TObjectPtr<UMeshComponent> Component)
diff --git a/Engine/Plugins/Experimental/ProceduralVegetationEditor/Source/ProceduralVegetationEditor/Private/SPVEditorViewport.h b/Engine/Plugins/Experimental/ProceduralVegetationEditor/Source/ProceduralVegetationEditor/Private/SPVEditorViewport.h
--- a/Engine/Plugins/Experimental/ProceduralVegetationEditor/Source/ProceduralVegetationEditor/Private/SPVEditorViewport.h
+++ b/Engine/Plugins/Experimental/ProceduralVegetationEditor/Source/ProceduralVegetationEditor/Private/SPVEditorViewport.h
@@ -8,6 +8,7 @@
 
 #include "Widgets/AssetEditorViewport/SPCGEditorViewport.h"
 #include "Widgets/Text/SRichTextBlock.h"
+#include "Widgets/Views/SListView.h"
 
 #include "SPVEditorViewport.generated.h"
 
@@ -148,4 +149,5 @@ private:
 	ESkeletonVisualizationModes DebugVisualizationMode = ESkeletonVisualizationModes::None;
 
 	TSharedPtr<FSlateBrush> PreviewNodeBackgroundBrush;
+	TSharedPtr<SListView<TSharedPtr<FText>>> LegendListView;
 };
diff --git a/Engine/Plugins/Experimental/ProceduralVegetationEditor/Source/ProceduralVegetationEditor/Private/Visualizations/PVDataVisualization.cpp b/Engine/Plugins/Experimental/ProceduralVegetationEditor/Source/ProceduralVegetationEditor/Private/Visualizations/PVDataVisualization.cpp
--- a/Engine/Plugins/Experimental/ProceduralVegetationEditor/Source/ProceduralVegetationEditor/Private/Visualizations/PVDataVisualization.cpp
+++ b/Engine/Plugins/Experimental/ProceduralVegetationEditor/Source/ProceduralVegetationEditor/Private/Visualizations/PVDataVisualization.cpp
@@ -21,7 +21,8 @@
 #include "Components/TextRenderComponent.h"
  
 #include "DataTypes/PVData.h"
- 
+#include "DataTypes/PVTextureSetupData.h"
+
 #include "DebugVisualization/PVDebugVisualizer.h"
  
 #include "Engine/SkeletalMesh.h"
@@ -35,7 +36,8 @@
 #include "Helpers/PVUtilities.h"
  
 #include "Materials/Material.h"
- 
+#include "Materials/MaterialInstanceDynamic.h"
+
 #include "Nodes/PVBaseSettings.h"
 #include "Nodes/PVObjectInteractionSettings.h"
 #include "ProceduralVegetation/Private/Helpers/PVFoliageHelper.h"
@@ -65,6 +67,8 @@ void FPVDataVisualization::RegisterVisualizations()
 	RenderMap.Add(EPVRenderType::FoliageGrid, FPVRenderCallback::CreateRaw(this, &FPVDataVisualization::FoliageGridRenderer));
 	RenderMap.Add(EPVRenderType::FoliageAttachments, FPVRenderCallback::CreateRaw(this, &FPVDataVisualization::FoliageAttachmentPointsRenderer));
 	RenderMap.Add(EPVRenderType::Seed, FPVRenderCallback::CreateRaw(this, &FPVDataVisualization::SeedRenderer));
+	
+	RenderMap.Add(EPVRenderType::Textures, FPVRenderCallback::CreateRaw(this, &FPVDataVisualization::TexturesRenderer));
 }
  
 void FPVDataVisualization::ExecuteDebugDisplay(FPCGContext* Context, const UPCGSettingsInterface* SettingsInterface, const UPCGData* Data,
@@ -112,11 +116,11 @@ FPCGSetupSceneFunc FPVDataVisualization::GetViewportSetupFunc(const UPCGSettings
 			if (RenderMap.Contains(RenderType))
 			{
 				InOutParams.Scene->GetLineBatcher()->Flush();
-				RenderMap[RenderType].Execute(Collection, InOutParams, Bounds);
+				RenderMap[RenderType].Execute(WeakData.Get(), Collection, InOutParams, Bounds);
 
 				for (const FPVVisualizationCollection& VisualizationCollection : WeakData->GetVisualizationCollections())
 				{
-					RenderMap[RenderType].Execute(VisualizationCollection.Collection, InOutParams, Bounds);
+					RenderMap[RenderType].Execute(WeakData.Get(), VisualizationCollection.Collection, InOutParams, Bounds);
 				}
 			}
 		}
@@ -149,7 +153,7 @@ FPCGSetupSceneFunc FPVDataVisualization::GetViewportSetupFunc(const UPCGSettings
 	};
 }
  
-void FPVDataVisualization::MeshRenderer(const FManagedArrayCollection& InCollection, FPCGSceneSetupParams& InOutParams,
+void FPVDataVisualization::MeshRenderer(const UPVData* Data, const FManagedArrayCollection& InCollection, FPCGSceneSetupParams& InOutParams,
                                                   FBoxSphereBounds& OutBounds)
 {
 	TRACE_CPUPROFILER_EVENT_SCOPE(FPVDataVisualization::MeshRenderer);
@@ -186,7 +190,7 @@ void FPVDataVisualization::MeshRenderer(const FManagedArrayCollection& InCollect
 	}
 }
  
-void FPVDataVisualization::SkeletonRenderer(const FManagedArrayCollection& InCollection, FPCGSceneSetupParams& InOutParams,
+void FPVDataVisualization::SkeletonRenderer(const UPVData* Data, const FManagedArrayCollection& InCollection, FPCGSceneSetupParams& InOutParams,
                                                       FBoxSphereBounds& OutBounds)
 {
 	TRACE_CPUPROFILER_EVENT_SCOPE(FPVDataVisualization::SkeletonRenderer);
@@ -207,7 +211,7 @@ void FPVDataVisualization::SkeletonRenderer(const FManagedArrayCollection& InCol
 	OutBounds = OutBounds + VisualizerComponent->CalcLocalBounds();
 }
  
-void FPVDataVisualization::FoliageRenderer(const FManagedArrayCollection& InCollection, FPCGSceneSetupParams& InOutParams,
+void FPVDataVisualization::FoliageRenderer(const UPVData* Data, const FManagedArrayCollection& InCollection, FPCGSceneSetupParams& InOutParams,
                                                      FBoxSphereBounds& OutBounds)
 {
 	TRACE_CPUPROFILER_EVENT_SCOPE(FPVDataVisualization::FoliageRenderer);
@@ -235,7 +239,7 @@ void FPVDataVisualization::FoliageRenderer(const FManagedArrayCollection& InColl
 	}
 }
  
-void FPVDataVisualization::BonesRenderer(const FManagedArrayCollection& InCollection, FPCGSceneSetupParams& InOutParams,
+void FPVDataVisualization::BonesRenderer(const UPVData* Data, const FManagedArrayCollection& InCollection, FPCGSceneSetupParams& InOutParams,
 	FBoxSphereBounds& OutBounds)
 {
 	TRACE_CPUPROFILER_EVENT_SCOPE(FPVDataVisualization::BonesRenderer);
@@ -289,7 +293,7 @@ TPair<int32, int32> GetMeshTrisAndVerts(const UObject* FoliageObject)
 	return {NumTriangles, NumVertices};
 }
  
-void FPVDataVisualization::FoliageGridRenderer(const FManagedArrayCollection& InCollection, FPCGSceneSetupParams& InOutParams,
+void FPVDataVisualization::FoliageGridRenderer(const UPVData* Data, const FManagedArrayCollection& InCollection, FPCGSceneSetupParams& InOutParams,
 	FBoxSphereBounds& OutBounds)
 {
 	const PV::Facades::FFoliageFacade FoliageFacade(InCollection);
@@ -408,7 +412,7 @@ void FPVDataVisualization::FoliageGridRenderer(const FManagedArrayCollection& In
 	}
 }
  
-void FPVDataVisualization::FoliageAttachmentPointsRenderer(const FManagedArrayCollection& InCollection, FPCGSceneSetupParams& InOutParams,
+void FPVDataVisualization::FoliageAttachmentPointsRenderer(const UPVData* Data, const FManagedArrayCollection& InCollection, FPCGSceneSetupParams& InOutParams,
 	FBoxSphereBounds& OutBounds)
 {
 	const PV::Facades::FFoliageFacade FoliageFacade(InCollection);
@@ -469,7 +473,7 @@ void FPVDataVisualization::FoliageAttachmentPointsRenderer(const FManagedArrayCo
 	OutBounds = OutBounds + Bounds;
 }
 
-void FPVDataVisualization::SeedRenderer(const FManagedArrayCollection& InCollection, FPCGSceneSetupParams& InOutParams,
+void FPVDataVisualization::SeedRenderer(const UPVData* Data, const FManagedArrayCollection& InCollection, FPCGSceneSetupParams& InOutParams,
 	FBoxSphereBounds& OutBounds)
 {
 	const PV::Facades::FPointFacade PointFacade(InCollection);
@@ -522,5 +526,99 @@ void FPVDataVisualization::SeedRenderer(const FManagedArrayCollection& InCollect
 	
 	OutBounds = OutBounds + Bounds;
 }
+
+void FPVDataVisualization::TexturesRenderer(const UPVData* Data, const FManagedArrayCollection& InCollection, FPCGSceneSetupParams& InOutParams,
+	FBoxSphereBounds& OutBounds)
+{
+	const UPVTextureSetupData* TextureData = Cast<UPVTextureSetupData>(Data);
+	if (!ensure(TextureData))
+	{
+		UE_LOG(LogProceduralVegetationDataVisualization, Warning, TEXT("Cannot visualize Textures, Texture Data is not valid."));
+		return;
+	}
+
+	const TObjectPtr<UStaticMeshComponent> StaticMeshComponent = NewObject<UStaticMeshComponent>(GetTransientPackage(), NAME_None, RF_Transient);
+	StaticMeshComponent->SetStaticMesh(LoadObject<UStaticMesh>(nullptr, TEXT("/Engine/BasicShapes/Plane.Plane")));
+
+	UMaterialInterface* Material = LoadObject<UMaterialInterface>(nullptr, TEXT("/ProceduralVegetationEditor/DefaultAssets/Visualization/Materials/M_TextureSetup.M_TextureSetup"));
+
+	if (Material)
+	{
+		UMaterialInstanceDynamic* DynamicMaterial = UMaterialInstanceDynamic::Create(Material, nullptr);
+
+		if (DynamicMaterial)
+		{
+			FString ChannelToPreview = TextureData->TextureSetupInfo.PreviewChannelName;
+			if (TextureData->TextureSetupInfo.Channels.Contains(ChannelToPreview))
+			{
+				DynamicMaterial->SetTextureParameterValue(TEXT("PreviewTexture"), TextureData->TextureSetupInfo.Channels[ChannelToPreview]);
+			}
+	
+			StaticMeshComponent->SetMaterial(0, DynamicMaterial);
+		}
+	}
+
+	FVector Position = FVector(0,0,550);
+	FBox Bounds;
+	Bounds += StaticMeshComponent->GetBounds().GetBox();
+	Bounds += Position;
+	
+	InOutParams.ManagedResources.Add(StaticMeshComponent);
+	InOutParams.Scene->AddComponent(StaticMeshComponent, FTransform(
+				FQuat::MakeFromEuler(FVector(90, 0,90)),
+				static_cast<FVector>(Position),
+				FVector::One() * 10
+			));
+
+	const TObjectPtr<UPVLineBatchComponent> LineBatchComponent = NewObject<UPVLineBatchComponent>(GetTransientPackage(), NAME_None, RF_Transient);
+
+	TArray<TPair<FColor, FText>> LegendItems;
+	
+	float Generation = 0;
+	float PreviousEnd = 0;
+	constexpr float Scale = 1000;
+	const int NumGenerations = TextureData->TextureSetupInfo.GenerationUVs.Num();
+	
+	for (auto GenerationUV : TextureData->TextureSetupInfo.GenerationUVs)
+	{
+		float DilationFactor = (GenerationUV.OffsetXStart * Scale) - PreviousEnd;
+		const FLinearColor Color = PV::Utilities::GetRandomHueColor(NumGenerations > 0 ? Generation / static_cast<float>(NumGenerations) : Generation);
+		LineBatchComponent->AddLine(
+			FVector(0, PreviousEnd - (Scale/ 2.0), Scale + 100 ),
+			FVector(0, ((GenerationUV.OffsetXEnd * Scale) + DilationFactor) - (Scale/ 2.0), Scale + 100),
+			Color,
+			SDPG_World,
+			EPointDrawSettings::None
+		);
+
+		LineBatchComponent->AddLine(
+			FVector(0, PreviousEnd - (Scale/ 2.0), 0 ),
+			FVector(0, ((GenerationUV.OffsetXEnd * Scale) + DilationFactor) - (Scale/ 2.0), 0),
+			Color,
+			SDPG_World,
+			EPointDrawSettings::None
+		);
+
+		float OverflowValue = GenerationUV.GetOverflowValue();
+		int Pixels = static_cast<int>( OverflowValue * static_cast<float>(TextureData->TextureSetupInfo.Resolution));
+		FString LegendString = OverflowValue > 0 ? FString::Printf(TEXT("Generation %i Overflow by %i pixels."), static_cast<int>(Generation + 1), Pixels)
+		: FString::Printf(TEXT("Generation %i"), static_cast<int>(Generation + 1));
+
+		FText LegendText = FText::FromString(LegendString);
+		
+		LegendItems.Add({Color.ToFColor(true), LegendText });
+		
+		PreviousEnd = (GenerationUV.OffsetXEnd * Scale) + DilationFactor ;
+		Generation++;
+	}
+
+	InOutParams.ManagedResources.Add(LineBatchComponent);
+	InOutParams.Scene->AddComponent(LineBatchComponent, FTransform::Identity);
+	
+	TSharedPtr<SPVEditorViewport> ViewportWidget = StaticCastSharedPtr<SPVEditorViewport>(InOutParams.EditorViewportClient->GetEditorViewportWidget());
+    	ViewportWidget->PopulateLegendOverlayText(LegendItems);
+	
+	OutBounds = OutBounds + Bounds;
+}
  
 #undef LOCTEXT_NAMESPACE
\ No newline at end of file
diff --git a/Engine/Plugins/Experimental/ProceduralVegetationEditor/Source/ProceduralVegetationEditor/Private/Visualizations/PVDataVisualization.h b/Engine/Plugins/Experimental/ProceduralVegetationEditor/Source/ProceduralVegetationEditor/Private/Visualizations/PVDataVisualization.h
--- a/Engine/Plugins/Experimental/ProceduralVegetationEditor/Source/ProceduralVegetationEditor/Private/Visualizations/PVDataVisualization.h
+++ b/Engine/Plugins/Experimental/ProceduralVegetationEditor/Source/ProceduralVegetationEditor/Private/Visualizations/PVDataVisualization.h
@@ -12,7 +12,7 @@ class AActor;
 class UPCGData;
 struct FPCGContext;
 
-DECLARE_DELEGATE_ThreeParams(FPVRenderCallback, const FManagedArrayCollection&, FPCGSceneSetupParams&, FBoxSphereBounds&);
+DECLARE_DELEGATE_FourParams(FPVRenderCallback, const UPVData*, const FManagedArrayCollection&, FPCGSceneSetupParams&, FBoxSphereBounds&);
 
 class FPVDataVisualization : public IPCGSpatialDataVisualization
 {
@@ -29,13 +29,14 @@ protected:
 	
 	void RegisterVisualizations();
 
-	void MeshRenderer(const FManagedArrayCollection& InCollection, FPCGSceneSetupParams& InOutParams, FBoxSphereBounds& OutBounds);
-	void SkeletonRenderer(const FManagedArrayCollection& InCollection, FPCGSceneSetupParams& InOutParams, FBoxSphereBounds& OutBounds);
-	void FoliageRenderer(const FManagedArrayCollection& InCollection, FPCGSceneSetupParams& InOutParams, FBoxSphereBounds& OutBounds);
-	void BonesRenderer(const FManagedArrayCollection& InCollection, FPCGSceneSetupParams& InOutParams, FBoxSphereBounds& OutBounds);
-	void FoliageGridRenderer(const FManagedArrayCollection& InCollection, FPCGSceneSetupParams& InOutParams, FBoxSphereBounds& OutBounds);
-	void FoliageAttachmentPointsRenderer(const FManagedArrayCollection& InCollection, FPCGSceneSetupParams& InOutParams, FBoxSphereBounds& OutBounds);
-	void SeedRenderer(const FManagedArrayCollection& InCollection, FPCGSceneSetupParams& InOutParams, FBoxSphereBounds& OutBounds);
+	void MeshRenderer(const UPVData* Data, const FManagedArrayCollection& InCollection, FPCGSceneSetupParams& InOutParams, FBoxSphereBounds& OutBounds);
+	void SkeletonRenderer(const UPVData* Data, const FManagedArrayCollection& InCollection, FPCGSceneSetupParams& InOutParams, FBoxSphereBounds& OutBounds);
+	void FoliageRenderer(const UPVData* Data, const FManagedArrayCollection& InCollection, FPCGSceneSetupParams& InOutParams, FBoxSphereBounds& OutBounds);
+	void BonesRenderer(const UPVData* Data, const FManagedArrayCollection& InCollection, FPCGSceneSetupParams& InOutParams, FBoxSphereBounds& OutBounds);
+	void FoliageGridRenderer(const UPVData* Data, const FManagedArrayCollection& InCollection, FPCGSceneSetupParams& InOutParams, FBoxSphereBounds& OutBounds);
+	void FoliageAttachmentPointsRenderer(const UPVData* Data, const FManagedArrayCollection& InCollection, FPCGSceneSetupParams& InOutParams, FBoxSphereBounds& OutBounds);
+	void SeedRenderer(const UPVData* Data, const FManagedArrayCollection& InCollection, FPCGSceneSetupParams& InOutParams, FBoxSphereBounds& OutBounds);
+	void TexturesRenderer(const UPVData* Data, const FManagedArrayCollection& InCollection, FPCGSceneSetupParams& InOutParams, FBoxSphereBounds& OutBounds);
 };
 
 DECLARE_LOG_CATEGORY_EXTERN(LogProceduralVegetationDataVisualization, Log, All);
\ No newline at end of file
