[PVE] Fix PVS studio warnings for grafter implementation

#Synced-CL 51487316

diff --git a/Engine/Plugins/Experimental/ProceduralVegetationEditor/Source/ProceduralVegetation/Private/Implementations/PVGrafter.cpp b/Engine/Plugins/Experimental/ProceduralVegetationEditor/Source/ProceduralVegetation/Private/Implementations/PVGrafter.cpp
--- a/Engine/Plugins/Experimental/ProceduralVegetationEditor/Source/ProceduralVegetation/Private/Implementations/PVGrafter.cpp
+++ b/Engine/Plugins/Experimental/ProceduralVegetationEditor/Source/ProceduralVegetation/Private/Implementations/PVGrafter.cpp
@@ -772,9 +772,9 @@ void FPVGrafter::DistributeGraftWithHormoneBasedSettings(FManagedArrayCollection
 			return CandidateGraftData.IsValidIndex(CandidateIndex);
 		};
 
-	auto GetCandidateValue = [&](int32 CandidateIndex, EFoliageDistributionCondition Condition) -> float
+	auto GetCandidateValue = [CandidateGraftData](int32 CandidateIndex, EFoliageDistributionCondition Condition) -> float
 		{
-			const auto& [Light, UpAlignment, Health, Tip] = CandidateGraftData[CandidateIndex];
+			const auto& [Light, UpAlignment, Health, bTip] = CandidateGraftData[CandidateIndex];
 
 			switch (Condition)
 			{
@@ -785,7 +785,7 @@ void FPVGrafter::DistributeGraftWithHormoneBasedSettings(FManagedArrayCollection
 			case EFoliageDistributionCondition::Health:
 				return Health;
 			case EFoliageDistributionCondition::Tip:
-				return Tip
+				return bTip
 					? 1.0f
 					: 0.0f;
 			case EFoliageDistributionCondition::Scale:
@@ -1452,9 +1452,9 @@ void FPVGrafter::DistributeGraftWithParametricSettings(FManagedArrayCollection&
 			return CandidateGraftData.IsValidIndex(CandidateIndex);
 		};
 
-	auto GetCandidateValue = [&](int32 CandidateIndex, EFoliageDistributionCondition Condition) -> float
+	auto GetCandidateValue = [CandidateGraftData](int32 CandidateIndex, EFoliageDistributionCondition Condition) -> float
 		{
-			const auto& [Light, UpAlignment, Health, Tip] = CandidateGraftData[CandidateIndex];
+			const auto& [Light, UpAlignment, Health, bTip] = CandidateGraftData[CandidateIndex];
 
 			switch (Condition)
 			{
@@ -1465,7 +1465,7 @@ void FPVGrafter::DistributeGraftWithParametricSettings(FManagedArrayCollection&
 			case EFoliageDistributionCondition::Health:
 				return Health;
 			case EFoliageDistributionCondition::Tip:
-				return Tip
+				return bTip
 					? 1.0f
 					: 0.0f;
 			case EFoliageDistributionCondition::Scale:
