[PVE] Fix for crash when Axillary Parent growth is tweaked

#Synced-CL 50970602

diff --git a/Engine/Plugins/Experimental/ProceduralVegetationEditor/Source/ProceduralVegetation/Private/Implementations/PVGrower.cpp b/Engine/Plugins/Experimental/ProceduralVegetationEditor/Source/ProceduralVegetation/Private/Implementations/PVGrower.cpp
--- a/Engine/Plugins/Experimental/ProceduralVegetationEditor/Source/ProceduralVegetation/Private/Implementations/PVGrower.cpp
+++ b/Engine/Plugins/Experimental/ProceduralVegetationEditor/Source/ProceduralVegetation/Private/Implementations/PVGrower.cpp
@@ -415,7 +415,7 @@ void FPVGrower::InitializeGrowth(const FPVGrowerParams& InGrowerParams, const FP
 	UPVGrowerPrimitive* Primitive = std::move(InitializePrimitiveFromSeed(SeedPoint, Skeleton));
 	//Skeleton->Primitives.Add(Primitive);
 	Skeleton->AddPrimitive(Primitive);
-	Skeleton->AddPoint(Point, Primitive, false);
+	Skeleton->AddPoint(Point, Primitive);
 	Skeleton->MinLengthFromRoot = SeedPoint.LengthFromRoot;
 
 	//Calculate MinRootDist from seed points
@@ -527,13 +527,9 @@ UPVGrowerPoint* FPVGrower::InitializePointFromSeed(const FPVSeedPoint& SeedPoint
 
 UPVGrowerPrimitive* FPVGrower::InitializePrimitiveFromSeed(const FPVSeedPoint& SeedPoint, UPVGrowerData* Skeleton)
 {
-	TArray<int32> BranchBuds;
-	BranchBuds.Add(1);
-	
 	UPVGrowerPrimitive* Primitive = NewObject<UPVGrowerPrimitive>(Skeleton,UPVGrowerPrimitive::StaticClass(),NAME_None, RF_Transactional);
 	Primitive->BranchNumber = SeedPoint.BranchNumber;
 	Primitive->PlantNumber = SeedPoint.PlantNumber;
-	Primitive->BranchBuds = BranchBuds;
 	Primitive->BranchParentNumber = 0;
 	Primitive->BranchSourceBudNumber = 1;
 	Primitive->Parents = {0};
