Auto-toggle bExposeToLibrary for PVE nodes which should be hidden in DebugMode

#Synced-CL 51495698

diff --git a/Engine/Plugins/Experimental/ProceduralVegetationEditor/Source/ProceduralVegetation/Private/Helpers/PVUtilities.cpp b/Engine/Plugins/Experimental/ProceduralVegetationEditor/Source/ProceduralVegetation/Private/Helpers/PVUtilities.cpp
--- a/Engine/Plugins/Experimental/ProceduralVegetationEditor/Source/ProceduralVegetation/Private/Helpers/PVUtilities.cpp
+++ b/Engine/Plugins/Experimental/ProceduralVegetationEditor/Source/ProceduralVegetation/Private/Helpers/PVUtilities.cpp
@@ -30,9 +30,22 @@ namespace PV::Utilities
 		TEXT("Enables debug mode for the Procedural Vegetation Editor"), 
 		FConsoleVariableDelegate::CreateLambda([](IConsoleVariable* ConsoleVariable)
 			{
-				UPVImporterSettings::StaticClass()->GetDefaultObject<UPVImporterSettings>()->bExposeToLibrary = ConsoleVariable->GetBool();
-				UPVPresetLoaderSettings::StaticClass()->GetDefaultObject<UPVPresetLoaderSettings>()->bExposeToLibrary = ConsoleVariable->GetBool();
-				UProceduralVegetationPreset::ShowHideInternalProperties(ConsoleVariable->GetBool());
+				const bool bNewValue = ConsoleVariable->GetBool();
+
+				TArray<UClass*> PVSettingsClasses;
+				GetDerivedClasses(UPVBaseSettings::StaticClass(), PVSettingsClasses);
+
+				for (UClass* PVSettingsClass : PVSettingsClasses)
+				{
+					UPVBaseSettings* BaseSettings = PVSettingsClass->GetDefaultObject<UPVBaseSettings>();
+					check(BaseSettings != nullptr);
+					if (BaseSettings->bOnlyExposeInDebugMode)
+					{
+						BaseSettings->bExposeToLibrary = bNewValue;
+					}
+				}
+			
+				UProceduralVegetationPreset::ShowHideInternalProperties(bNewValue);
 			})
 	);
 #endif
diff --git a/Engine/Plugins/Experimental/ProceduralVegetationEditor/Source/ProceduralVegetation/Private/Nodes/Development/PVGrowthDataDiffSettings.cpp b/Engine/Plugins/Experimental/ProceduralVegetationEditor/Source/ProceduralVegetation/Private/Nodes/Development/PVGrowthDataDiffSettings.cpp
--- a/Engine/Plugins/Experimental/ProceduralVegetationEditor/Source/ProceduralVegetation/Private/Nodes/Development/PVGrowthDataDiffSettings.cpp
+++ b/Engine/Plugins/Experimental/ProceduralVegetationEditor/Source/ProceduralVegetation/Private/Nodes/Development/PVGrowthDataDiffSettings.cpp
@@ -333,7 +333,7 @@ UPVGrowthDataDiffSettings::UPVGrowthDataDiffSettings()
 #if WITH_EDITOR
 	if (HasAnyFlags(RF_ClassDefaultObject))
 	{
-		bExposeToLibrary = PV::Utilities::DebugModeEnabled();
+		bOnlyExposeInDebugMode = true;
 	}
 #endif
 }
diff --git a/Engine/Plugins/Experimental/ProceduralVegetationEditor/Source/ProceduralVegetation/Private/Nodes/Development/PVGrowthDataJsonImporterSettings.cpp b/Engine/Plugins/Experimental/ProceduralVegetationEditor/Source/ProceduralVegetation/Private/Nodes/Development/PVGrowthDataJsonImporterSettings.cpp
--- a/Engine/Plugins/Experimental/ProceduralVegetationEditor/Source/ProceduralVegetation/Private/Nodes/Development/PVGrowthDataJsonImporterSettings.cpp
+++ b/Engine/Plugins/Experimental/ProceduralVegetationEditor/Source/ProceduralVegetation/Private/Nodes/Development/PVGrowthDataJsonImporterSettings.cpp
@@ -34,7 +34,7 @@ UPVGrowthDataJsonImporterSettings::UPVGrowthDataJsonImporterSettings()
 #if WITH_EDITOR
 	if (HasAnyFlags(RF_ClassDefaultObject))
 	{
-		bExposeToLibrary = PV::Utilities::DebugModeEnabled();
+		bOnlyExposeInDebugMode = true;
 	}
 #endif
 }
diff --git a/Engine/Plugins/Experimental/ProceduralVegetationEditor/Source/ProceduralVegetation/Private/Nodes/Development/PVMockTreeImporterSettings.cpp b/Engine/Plugins/Experimental/ProceduralVegetationEditor/Source/ProceduralVegetation/Private/Nodes/Development/PVMockTreeImporterSettings.cpp
--- a/Engine/Plugins/Experimental/ProceduralVegetationEditor/Source/ProceduralVegetation/Private/Nodes/Development/PVMockTreeImporterSettings.cpp
+++ b/Engine/Plugins/Experimental/ProceduralVegetationEditor/Source/ProceduralVegetation/Private/Nodes/Development/PVMockTreeImporterSettings.cpp
@@ -32,7 +32,7 @@ UPVMockTreeImporterSettings::UPVMockTreeImporterSettings()
 #if WITH_EDITOR
 	if (HasAnyFlags(RF_ClassDefaultObject))
 	{
-		bExposeToLibrary = PV::Utilities::DebugModeEnabled();
+		bOnlyExposeInDebugMode = true;
 	}
 #endif
 }
diff --git a/Engine/Plugins/Experimental/ProceduralVegetationEditor/Source/ProceduralVegetation/Private/Nodes/Development/PVRecomputeGrowthDataAttributesSettings.cpp b/Engine/Plugins/Experimental/ProceduralVegetationEditor/Source/ProceduralVegetation/Private/Nodes/Development/PVRecomputeGrowthDataAttributesSettings.cpp
--- a/Engine/Plugins/Experimental/ProceduralVegetationEditor/Source/ProceduralVegetation/Private/Nodes/Development/PVRecomputeGrowthDataAttributesSettings.cpp
+++ b/Engine/Plugins/Experimental/ProceduralVegetationEditor/Source/ProceduralVegetation/Private/Nodes/Development/PVRecomputeGrowthDataAttributesSettings.cpp
@@ -29,7 +29,7 @@ UPVRecomputeGrowthDataAttributesSettings::UPVRecomputeGrowthDataAttributesSettin
 #if WITH_EDITOR
 	if (HasAnyFlags(RF_ClassDefaultObject))
 	{
-		bExposeToLibrary = PV::Utilities::DebugModeEnabled();
+		bOnlyExposeInDebugMode = true;
 	}
 #endif
 }
diff --git a/Engine/Plugins/Experimental/ProceduralVegetationEditor/Source/ProceduralVegetation/Private/Nodes/PVImporterSettings.cpp b/Engine/Plugins/Experimental/ProceduralVegetationEditor/Source/ProceduralVegetation/Private/Nodes/PVImporterSettings.cpp
--- a/Engine/Plugins/Experimental/ProceduralVegetationEditor/Source/ProceduralVegetation/Private/Nodes/PVImporterSettings.cpp
+++ b/Engine/Plugins/Experimental/ProceduralVegetationEditor/Source/ProceduralVegetation/Private/Nodes/PVImporterSettings.cpp
@@ -20,7 +20,7 @@ UPVImporterSettings::UPVImporterSettings()
 #if WITH_EDITOR
 	if (HasAnyFlags(RF_ClassDefaultObject))
 	{
-		bExposeToLibrary = PV::Utilities::DebugModeEnabled();
+		bOnlyExposeInDebugMode = true;
 	}
 #endif
 }
diff --git a/Engine/Plugins/Experimental/ProceduralVegetationEditor/Source/ProceduralVegetation/Private/Nodes/PVPresetLoaderSettings.cpp b/Engine/Plugins/Experimental/ProceduralVegetationEditor/Source/ProceduralVegetation/Private/Nodes/PVPresetLoaderSettings.cpp
--- a/Engine/Plugins/Experimental/ProceduralVegetationEditor/Source/ProceduralVegetation/Private/Nodes/PVPresetLoaderSettings.cpp
+++ b/Engine/Plugins/Experimental/ProceduralVegetationEditor/Source/ProceduralVegetation/Private/Nodes/PVPresetLoaderSettings.cpp
@@ -40,12 +40,12 @@ void UPVPresetLoaderSettings::PostEditChangeProperty(FPropertyChangedEvent& Prop
 
 UPVPresetLoaderSettings::UPVPresetLoaderSettings()
 {
-	#if WITH_EDITOR
-    	if (HasAnyFlags(RF_ClassDefaultObject))
-    	{
-    		bExposeToLibrary = PV::Utilities::DebugModeEnabled();
-    	}
-    #endif
+#if WITH_EDITOR
+	if (HasAnyFlags(RF_ClassDefaultObject))
+	{
+		bOnlyExposeInDebugMode = true;
+	}
+#endif
 }
 
 void UPVPresetLoaderSettings::PostLoad()
diff --git a/Engine/Plugins/Experimental/ProceduralVegetationEditor/Source/ProceduralVegetation/Public/Nodes/PVBaseSettings.h b/Engine/Plugins/Experimental/ProceduralVegetationEditor/Source/ProceduralVegetation/Public/Nodes/PVBaseSettings.h
--- a/Engine/Plugins/Experimental/ProceduralVegetationEditor/Source/ProceduralVegetation/Public/Nodes/PVBaseSettings.h
+++ b/Engine/Plugins/Experimental/ProceduralVegetationEditor/Source/ProceduralVegetation/Public/Nodes/PVBaseSettings.h
@@ -39,6 +39,12 @@ protected:
 	virtual FPCGDataTypeIdentifier GetInputPinTypeIdentifier() const;
 	virtual FPCGDataTypeIdentifier GetOutputPinTypeIdentifier() const;
 
+public:
+#if WITH_EDITORONLY_DATA
+	UPROPERTY(BlueprintReadWrite, EditAnywhere, Category = AssetInfo, AssetRegistrySearchable, meta = (PCGNoHash))
+	bool bOnlyExposeInDebugMode = false;
+#endif
+
 protected:
 
 #if WITH_EDITORONLY_DATA
