[PVE] Fix issue where build is failing on Clang Static Analysis due to incomplete type

#Synced-CL 51523444

diff --git a/Engine/Plugins/Experimental/ProceduralVegetationEditor/Source/ProceduralVegetation/Private/Facades/PVFoliageFacade.cpp b/Engine/Plugins/Experimental/ProceduralVegetationEditor/Source/ProceduralVegetation/Private/Facades/PVFoliageFacade.cpp
--- a/Engine/Plugins/Experimental/ProceduralVegetationEditor/Source/ProceduralVegetation/Private/Facades/PVFoliageFacade.cpp
+++ b/Engine/Plugins/Experimental/ProceduralVegetationEditor/Source/ProceduralVegetation/Private/Facades/PVFoliageFacade.cpp
@@ -8,7 +8,6 @@
 #include "VisualizeTexture.h"
 #include "Chaos/Deformable/MuscleActivationConstraints.h"
 #include "Facades/PVAttributesNames.h"
-#include "Helpers/PVFoliageHelper.h"
 #include "Implementations/PVFoliage.h"
  
 namespace PV::Facades
diff --git a/Engine/Plugins/Experimental/ProceduralVegetationEditor/Source/ProceduralVegetation/Private/Helpers/PVUtilities.cpp b/Engine/Plugins/Experimental/ProceduralVegetationEditor/Source/ProceduralVegetation/Private/Helpers/PVUtilities.cpp
--- a/Engine/Plugins/Experimental/ProceduralVegetationEditor/Source/ProceduralVegetation/Private/Helpers/PVUtilities.cpp
+++ b/Engine/Plugins/Experimental/ProceduralVegetationEditor/Source/ProceduralVegetation/Private/Helpers/PVUtilities.cpp
@@ -3,7 +3,7 @@
 #include "Helpers/PVUtilities.h"
 #include "ProceduralVegetationModule.h"
 #include "ProceduralVegetationPreset.h"
-#include "PVFoliageHelper.h"
+#include "Helpers/PVFoliageHelper.h"
 #include "AssetRegistry/AssetRegistryModule.h"
 #include "Components/InstancedSkinnedMeshComponent.h"
 #include "Engine/StaticMesh.h"
diff --git a/Engine/Plugins/Experimental/ProceduralVegetationEditor/Source/ProceduralVegetation/Public/Facades/PVFoliageFacade.h b/Engine/Plugins/Experimental/ProceduralVegetationEditor/Source/ProceduralVegetation/Public/Facades/PVFoliageFacade.h
--- a/Engine/Plugins/Experimental/ProceduralVegetationEditor/Source/ProceduralVegetation/Public/Facades/PVFoliageFacade.h
+++ b/Engine/Plugins/Experimental/ProceduralVegetationEditor/Source/ProceduralVegetation/Public/Facades/PVFoliageFacade.h
@@ -5,8 +5,7 @@
 #include "CoreMinimal.h"
 #include "IShrinkableFacade.h"
 #include "GeometryCollection/ManagedArrayAccessor.h"
- 
-struct FPVFoliageInfo;
+#include "ProceduralVegetation/Private/Helpers/PVFoliageHelper.h"
  
 namespace PV::Facades
 {
