[PVE] Add parametric based distribution feature to the foliage distributor node

- Adds parametric settings to the foliage distributor node with a mode drop down to select between parametric and hormone based settings
- Refactors code so that parametric settings based code is reused across foliage distributor and grafter node implementation

#Synced-CL 51535444

diff --git a/Engine/Plugins/Experimental/ProceduralVegetationEditor/Source/ProceduralVegetation/Private/Helpers/PVParametricDistributionHelper.cpp b/Engine/Plugins/Experimental/ProceduralVegetationEditor/Source/ProceduralVegetation/Private/Helpers/PVParametricDistributionHelper.cpp
new file mode 100644
--- /dev/null
+++ b/Engine/Plugins/Experimental/ProceduralVegetationEditor/Source/ProceduralVegetation/Private/Helpers/PVParametricDistributionHelper.cpp
@@ -0,0 +1,508 @@
+// Copyright Epic Games, Inc. All Rights Reserved.
+
+#include "PVParametricDistributionHelper.h"
+
+#include "PVAttributesHelper.h"
+#include "Facades/PVBranchFacade.h"
+#include "Helpers/PVUtilities.h"
+
+void PV::ParametricDistributionHelper::ComputeAttachmentPoints(FManagedArrayCollection& OutCollection,
+                                                               const FNormalizedPointCaches& NormalizedPointCaches,
+                                                               const FParametricSettings::FSpacingSettings& SpacingSettings,
+                                                               const FParametricSettings::FAngleSettings& AngleSettings,
+                                                               const FParametricSettings::FScaleSettings& ScaleSettings,
+                                                               FRandomStream& RandomStream,
+                                                               TArray<FAttachmentPoint>& OutAttachmentPoints)
+{
+	auto PointPlantGradientAttribute = FPointPlantGradientAttribute::GetAttribute(OutCollection);
+	auto PointBranchGradientAttribute = FPointBranchGradientAttribute::GetAttribute(OutCollection);
+	auto BranchParentNumberAttribute = FBranchParentNumberAttribute::GetAttribute(OutCollection);
+	auto BranchChildrenAttribute = FBranchChildrenAttribute::GetAttribute(OutCollection);
+	auto BranchNumberAttribute = FBranchNumberAttribute::GetAttribute(OutCollection);
+	auto BranchPointsAttribute = FBranchPointsAttribute::GetAttribute(OutCollection);
+	auto PointPositionAttribute = FPointPositionAttribute::GetAttribute(OutCollection);
+	auto BudDirectionAttribute = FBudDirectionAttribute::GetAttribute(OutCollection);
+	auto PointLengthFromRootAttribute = FPointLengthFromRootAttribute::GetAttribute(OutCollection);
+
+	const int32 NumOfBranches = BranchPointsAttribute.Num();
+	check(NumOfBranches > 0);
+	if (NumOfBranches == 0)
+	{
+		return;
+	}
+
+	// Only need to do this temporarily since legacy skeleton has inconsistent data in children
+	// TODO: Remove this once support for legacy skeletons is dropped
+	PV::Facades::FBranchFacade BranchFacadeOut(OutCollection);
+	BranchFacadeOut.RecomputeBranchChildren();
+	BranchFacadeOut.RecomputeBranchParents();
+
+	AttributesHelper::ComputePointPlantAndBranchGradient(
+		{
+			PointPlantGradientAttribute,
+			PointBranchGradientAttribute,
+			BranchPointsAttribute,
+			BranchParentNumberAttribute,
+			BranchNumberAttribute,
+			BranchChildrenAttribute
+		});
+	
+	TMap<int32, float> BranchIndicesToBranchLengths = ComputeBranchLengths(PointLengthFromRootAttribute, BranchPointsAttribute);
+	float MaxBranchLength = -1.0f;
+	float AverageBranchLength = 0.0f;
+	
+	for (const auto IndexLengthPair : BranchIndicesToBranchLengths)
+	{
+		const float BranchLength = IndexLengthPair.Value;
+		MaxBranchLength = FMath::Max(BranchLength, MaxBranchLength);
+		AverageBranchLength = AverageBranchLength + BranchLength;
+	}
+
+	check(MaxBranchLength > KINDA_SMALL_NUMBER);
+	if (MaxBranchLength < KINDA_SMALL_NUMBER)
+	{
+		return;
+	}
+
+	AverageBranchLength = AverageBranchLength / BranchIndicesToBranchLengths.Num();
+	AverageBranchLength = FMath::Lerp(AverageBranchLength, MaxBranchLength, 0.33f);
+	
+	for (int32 BranchIndex = 0; BranchIndex < NumOfBranches; BranchIndex++)
+	{
+		// Check generation limit
+		const TArray<int32>& BranchParents = BranchFacadeOut.GetParents(BranchIndex);
+		if (SpacingSettings.LimitStartGeneration && BranchParents.Num() < SpacingSettings.StartGeneration)
+		{
+			continue;
+		}
+
+		if (SpacingSettings.LimitEndGeneration && BranchParents.Num() > SpacingSettings.EndGeneration)
+		{
+			continue;
+		}
+
+		const float LengthRatio = BranchIndicesToBranchLengths[BranchIndex] / MaxBranchLength;
+		int32 LoopNumber = SpacingSettings.BranchDensity;
+		if (SpacingSettings.SpacingBasis == FParametricSettings::EDistributionBasis::Branch)
+		{
+			LoopNumber = static_cast<int32>(LoopNumber * LengthRatio);
+		}
+		else if (SpacingSettings.SpacingBasis == FParametricSettings::EDistributionBasis::Plant)
+		{
+			LoopNumber = static_cast<int32>(LoopNumber * (AverageBranchLength / MaxBranchLength));
+		}
+
+		LoopNumber = FMath::Max(LoopNumber, 1);
+
+		const TArray<int32>& BranchPoints = BranchPointsAttribute[BranchIndex];
+
+		if (BranchPoints.Num() < 2)
+		{
+			continue;
+		}
+
+		TArray<int32> SortedPointIndices = BranchPoints;
+		SortedPointIndices.Sort([PointLengthFromRootAttribute](const int32& A, const int32& B)
+			{
+				return PointLengthFromRootAttribute[A] < PointLengthFromRootAttribute[B];
+			});
+
+		const int32 FirstPointIndex = SortedPointIndices[0];
+		const int32 TipPointIndex = SortedPointIndices.Last();
+		const float PlantGradientAtRoot = 1.0f - PointPlantGradientAttribute[FirstPointIndex];
+		const float LFRofRoot = PointLengthFromRootAttribute[FirstPointIndex];
+		const float LFRofTip = PointLengthFromRootAttribute[TipPointIndex];
+
+		for (int32 j = 0; j < LoopNumber; j++)
+		{
+			float BaseValue = static_cast<float>(j) / FMath::Max(LoopNumber - 1, 1);
+			float PlantGradient;
+			float BranchGradient = 0.0f;
+			float AttachmentPointLengthFromRoot = 0.0f;
+			FVector3f Pos(0.0f, 0.0f, 0.0f);
+			int32 PointIndexA = INDEX_NONE;
+			int32 PointIndexB = INDEX_NONE;
+			float PointAlpha = 0.0f;
+
+			if (SpacingSettings.SpacingBasis == FParametricSettings::EDistributionBasis::Plant)
+			{
+				PlantGradient = BaseValue;
+				PlantGradient = SpacingSettings.SpacingRamp->GetRichCurveConst()->Eval(PlantGradient);
+				PlantGradient = FMath::GetMappedRangeValueClamped(
+					FVector2f(0.0f, 1.0f),
+					FVector2f(SpacingSettings.RelativeStart, SpacingSettings.RelativeEnd),
+					PlantGradient);
+				if (PlantGradient < PlantGradientAtRoot)
+				{
+					continue;
+				}
+
+				Pos = SampleVectorByFloat(PlantGradient, SortedPointIndices, PointPlantGradientAttribute, PointPositionAttribute);
+				AttachmentPointLengthFromRoot = SampleFloatByFloat(PlantGradient, SortedPointIndices, PointPlantGradientAttribute,
+					PointLengthFromRootAttribute);
+				BranchGradient = FMath::GetMappedRangeValueClamped(
+					FVector2f(LFRofRoot, LFRofTip),
+					FVector2f(0.0f, 1.0f),
+					AttachmentPointLengthFromRoot);
+
+				FindPointByNormalizedLengthFromRoot(BranchGradient, SortedPointIndices, PointLengthFromRootAttribute, LFRofRoot, LFRofTip,
+					PointIndexA,
+					PointIndexB, PointAlpha);
+			}
+			else // SpacingBasis is Branch
+			{
+				BranchGradient = SpacingSettings.SpacingRamp->GetRichCurveConst()->Eval(BaseValue);
+				BranchGradient = FMath::GetMappedRangeValueClamped(
+					FVector2f(0.0f, 1.0f),
+					FVector2f(SpacingSettings.RelativeStart, SpacingSettings.RelativeEnd),
+					BranchGradient);
+
+				if (SpacingSettings.BranchDensity == 1)
+				{
+					BranchGradient = SpacingSettings.RelativeEnd;
+					BaseValue = 1.0f;
+				}
+
+				if (!FindPointByNormalizedLengthFromRoot(BranchGradient, SortedPointIndices, PointLengthFromRootAttribute, LFRofRoot, LFRofTip,
+					PointIndexA, PointIndexB, PointAlpha))
+				{
+					continue;
+				}
+
+				const FVector3f P0 = PointPositionAttribute[PointIndexA];
+				const FVector3f P1 = PointPositionAttribute[PointIndexB];
+				Pos = FMath::Lerp(P0, P1, PointAlpha);
+
+				AttachmentPointLengthFromRoot = FMath::GetMappedRangeValueClamped(
+					FVector2f(0.0f, 1.0f),
+					FVector2f(LFRofRoot, LFRofTip),
+					BranchGradient);
+			}
+
+			const float PlantGradient0 = 1.0f - PointPlantGradientAttribute[PointIndexA];
+			const float PlantGradient1 = 1.0f - PointPlantGradientAttribute[PointIndexB];
+			PlantGradient = FMath::Lerp(PlantGradient0, PlantGradient1, PointAlpha);
+
+			const float ScaleRampLookUp =
+				(ScaleSettings.ScaleRampBasis == FParametricSettings::EDistributionBasis::Branch)
+				? BranchGradient
+				: PlantGradient;
+
+			float ScaleRampValue = ScaleSettings.ScaleRamp->GetRichCurveConst()->Eval(ScaleRampLookUp);
+
+			if (ScaleSettings.ScaleRampBasis == FParametricSettings::EDistributionBasis::Branch)
+			{
+				ScaleRampValue = FMath::GetMappedRangeValueClamped(
+					FVector2f(0.0f, 1.0f),
+					FVector2f(0.0f, LengthRatio),
+					ScaleRampValue);
+			}
+
+			const float RandomScaleFactor = RandomStream.FRandRange(
+				ScaleSettings.RandomizeScaleAngleMinimum,
+				ScaleSettings.RandomizeScaleAngleMaximum);
+
+			const float Scale = ScaleRampValue * RandomScaleFactor * ScaleSettings.BaseScale;
+
+			const FVector3f BudDirectionApicalA = BudDirectionAttribute[PointIndexA].Apical;
+			const FVector3f BudDirectionApicalB = BudDirectionAttribute[PointIndexB].Apical;
+			const FVector3f BudDirectionUpA = BudDirectionAttribute[PointIndexA].UpVector;
+			const FVector3f BudDirectionUpB = BudDirectionAttribute[PointIndexB].UpVector;
+
+			const FVector3f ApicalDirection = FMath::Lerp(BudDirectionApicalA, BudDirectionApicalB, PointAlpha);
+			const FVector3f BudDirectionUp = FMath::Lerp(BudDirectionUpA, BudDirectionUpB, PointAlpha);
+
+			const float AxilRampValue = AngleSettings.AxilAngleRamp->GetRichCurveConst()->Eval(BaseValue);
+			float AxilAngleInDegrees = AngleSettings.AxilAngle * AxilRampValue;
+
+			AxilAngleInDegrees += RandomStream.FRandRange(AngleSettings.RandomizeAxilAngleMinimum, AngleSettings.RandomizeAxilAngleMaximum);
+			AxilAngleInDegrees = FMath::Clamp(AxilAngleInDegrees, -90.0f, 90.0f);
+
+			const float Twist = FMath::DegreesToRadians(AngleSettings.Rotation) * static_cast<float>(j);
+
+			FVector3f UpDirection, NormalDirection;
+			ComputeRotationAnglesForAttachment(
+				ApicalDirection,
+				BudDirectionUp,
+				FMath::DegreesToRadians(AxilAngleInDegrees),
+				Twist,
+				UpDirection,
+				NormalDirection);
+
+			const bool bIsTipInstance =
+				FMath::IsNearlyEqual(AttachmentPointLengthFromRoot, LFRofTip, 0.001f) &&
+				(BranchGradient >= 0.999f || FMath::IsNearlyEqual(BranchGradient, SpacingSettings.RelativeEnd, 0.001f));
+
+			if (bIsTipInstance)
+			{
+				const int32 SecondToLastPointIndex = SortedPointIndices[SortedPointIndices.Num() - 2];
+				const FVector3f TipPosition = PointPositionAttribute[TipPointIndex];
+				const FVector3f SecondToLastPointPosition = PointPositionAttribute[SecondToLastPointIndex];
+				FVector3f TipTangent = (TipPosition - SecondToLastPointPosition).GetSafeNormal();
+				if (TipTangent.IsNearlyZero())
+				{
+					TipTangent = ApicalDirection.GetSafeNormal();
+				}
+
+				const FVector3f OldUp = UpDirection.GetSafeNormal();
+				FQuat4f AlignQuat = FQuat4f::FindBetweenNormals(OldUp, TipTangent);
+
+				UpDirection = TipTangent;
+				NormalDirection = AlignQuat.RotateVector(NormalDirection).GetSafeNormal();
+
+				if (FMath::Abs(FVector3f::DotProduct(UpDirection, NormalDirection)) > 0.999f)
+				{
+					FVector Axis1, Axis2;
+					FVector(UpDirection).FindBestAxisVectors(Axis1, Axis2);
+					NormalDirection = FVector3f(Axis1);
+				}
+			}
+
+			ConditionPicker::FPointConditionSample Sample;
+			Sample.UpAlignment = UpAlignmentNormalized(ApicalDirection);
+			Sample.Tip = Tip(bIsTipInstance);
+
+			if (PointIndexA != INDEX_NONE && PointIndexB != INDEX_NONE)
+			{
+				Sample.Light = FMath::Lerp(NormalizedPointCaches.Light[PointIndexA], NormalizedPointCaches.Light[PointIndexB], PointAlpha);
+				Sample.Health = FMath::Lerp(NormalizedPointCaches.Health[PointIndexA], NormalizedPointCaches.Health[PointIndexB], PointAlpha);
+				Sample.Scale = FMath::Lerp(NormalizedPointCaches.PScales[PointIndexA], NormalizedPointCaches.PScales[PointIndexB], PointAlpha);
+			}
+			else
+			{
+				Sample.Light = 0.0f;
+				Sample.Health = 0.0f;
+				Sample.Scale = 0.0f;
+			}
+			
+			OutAttachmentPoints.Add({
+				.BranchIndex = BranchIndex,
+				.PointIndexA = PointIndexA,
+				.PointIndexB = PointIndexB,
+				.PointAlpha = PointAlpha,
+				.Position = Pos,
+				.UpDirection = UpDirection,
+				.NormalDirection = NormalDirection,
+				.Scale = Scale,
+				.LengthFromRoot = AttachmentPointLengthFromRoot,
+				.bIsTipInstance = bIsTipInstance,
+				.TipPointIndex = TipPointIndex,
+				.ConditionSample = Sample
+			});
+		}
+	}
+}
+
+TMap<int32, float> PV::ParametricDistributionHelper::ComputeBranchLengths(FPointLengthFromRootAttributeConstView PointLengthFromRootAttribute,
+                                                                          FBranchPointsAttributeConstView BranchPointsAttribute)
+{
+	const int32 NumOfBranches = BranchPointsAttribute.Num();
+	TMap<int32, float> BranchIndicesToBranchLengths;
+	BranchIndicesToBranchLengths.Reserve(NumOfBranches);
+
+	for (int32 BranchIndex = 0; BranchIndex < NumOfBranches; BranchIndex++)
+	{
+		const TArray<int32>& BranchPoints = BranchPointsAttribute[BranchIndex];
+		if (BranchPoints.Num() < 2)
+		{
+			BranchIndicesToBranchLengths.Add(BranchIndex, 0.0f);
+			continue;
+		}
+
+		TArray<int32> SortedPointIndices = BranchPoints;
+		SortedPointIndices.Sort([PointLengthFromRootAttribute](const int32& A, const int32& B)
+			{
+				return PointLengthFromRootAttribute[A] < PointLengthFromRootAttribute[B];
+			});
+
+		const int32 RootPointIndex = SortedPointIndices[0];
+		const int32 BranchTipPointIndex = SortedPointIndices.Last();
+
+		const float RootPointLFR = PointLengthFromRootAttribute[RootPointIndex];
+		const float TipPointLFR = PointLengthFromRootAttribute[BranchTipPointIndex];
+		const float BranchLength = TipPointLFR - RootPointLFR;
+
+		BranchIndicesToBranchLengths.Add(BranchIndex, BranchLength);
+	}
+
+	return BranchIndicesToBranchLengths;
+}
+
+FVector3f PV::ParametricDistributionHelper::SampleVectorByFloat(const float SampleValue, const TArray<int32>& SortedPointIndices,
+                                                                FPointPlantGradientAttributeConstView PointPlantGradientAttribute,
+                                                                FPointPositionAttributeConstView PointPositionAttribute)
+{
+	if (SortedPointIndices.Num() < 2)
+	{
+		return FVector3f();
+	}
+
+	const int32 FirstPointIndex = SortedPointIndices[0];
+	const int32 LastPointIndex = SortedPointIndices.Last();
+	float Value0 = 1.0f - PointPlantGradientAttribute[FirstPointIndex];
+	float Value1 = 1.0f - PointPlantGradientAttribute[LastPointIndex];
+	FVector3f OutValue(Value0, Value0, Value0);
+
+	if (FMath::IsWithinInclusive(SampleValue, Value0, Value1))
+	{
+		for (int32 i = 1; i < SortedPointIndices.Num(); i++)
+		{
+			Value1 = 1.0f - PointPlantGradientAttribute[SortedPointIndices[i]];
+			if (FMath::IsWithinInclusive(SampleValue, Value0, Value1))
+			{
+				FVector3f P0 = PointPositionAttribute[SortedPointIndices[i - 1]];
+				FVector3f P1 = PointPositionAttribute[SortedPointIndices[i]];
+
+				float Blend = FMath::GetMappedRangeValueClamped(
+					FVector2f(Value0, Value1),
+					FVector2f(0.0f, 1.0f),
+					SampleValue);
+
+				OutValue = FMath::Lerp(P0, P1, Blend);
+			}
+
+			Value0 = Value1;
+		}
+	}
+	else
+	{
+		if (FMath::Abs(Value0 - SampleValue) > FMath::Abs(Value1 - SampleValue))
+		{
+			OutValue = FVector3f(Value1, Value1, Value1);
+		}
+	}
+
+	return OutValue;
+}
+
+float PV::ParametricDistributionHelper::SampleFloatByFloat(const float SampleValue, const TArray<int32>& SortedPointIndices,
+                                                           FPointPlantGradientAttributeConstView PointPlantGradientAttribute,
+                                                           FPointLengthFromRootAttributeConstView PointLengthFromRootAttribute)
+{
+	if (SortedPointIndices.Num() < 2)
+	{
+		return -1.0f;
+	}
+
+	const int32 FirstPointIndex = SortedPointIndices[0];
+	const int32 LastPointIndex = SortedPointIndices.Last();
+	float Value0 = 1.0f - PointPlantGradientAttribute[FirstPointIndex];
+	float Value1 = 1.0f - PointPlantGradientAttribute[LastPointIndex];
+	float OutValue = Value0;
+
+	if (FMath::IsWithinInclusive(SampleValue, Value0, Value1))
+	{
+		for (int32 i = 1; i < SortedPointIndices.Num(); i++)
+		{
+			Value1 = 1.0f - PointPlantGradientAttribute[SortedPointIndices[i]];
+			if (FMath::IsWithinInclusive(SampleValue, Value0, Value1))
+			{
+				float P0 = PointLengthFromRootAttribute[SortedPointIndices[i - 1]];
+				float P1 = PointLengthFromRootAttribute[SortedPointIndices[i]];
+
+				float Blend = FMath::GetMappedRangeValueClamped(
+					FVector2f(Value0, Value1),
+					FVector2f(0.0f, 1.0f),
+					SampleValue);
+
+				OutValue = FMath::Lerp(P0, P1, Blend);
+			}
+
+			Value0 = Value1;
+		}
+	}
+	else
+	{
+		if (FMath::Abs(Value0 - SampleValue) > FMath::Abs(Value1 - SampleValue))
+		{
+			OutValue = Value1;
+		}
+	}
+
+	return OutValue;
+}
+
+bool PV::ParametricDistributionHelper::FindPointByNormalizedLengthFromRoot(float SampleValue, const TArray<int32>& SortedPointIndices,
+                                                                           FPointLengthFromRootAttributeConstView PointLengthFromRootAttribute,
+                                                                           const float LFRofRoot, const float LFRofTip, int32& PreviousPointIndex,
+                                                                           int32& NextPointIndex, float& OutAlpha)
+{
+	if (SortedPointIndices.Num() < 2)
+	{
+		return false;
+	}
+
+	SampleValue = FMath::Clamp(SampleValue, 0.0f, 1.0f);
+
+	const float SampleLengthFromRoot = FMath::GetMappedRangeValueClamped(
+		FVector2f(0.0f, 1.0f),
+		FVector2f(LFRofRoot, LFRofTip),
+		SampleValue);
+
+	for (int32 i = 0; i < SortedPointIndices.Num() - 1; ++i)
+	{
+		const int32 P = SortedPointIndices[i];
+		const int32 N = SortedPointIndices[i + 1];
+		const float LFR0 = PointLengthFromRootAttribute[P];
+		const float LFR1 = PointLengthFromRootAttribute[N];
+
+		if (FMath::IsWithinInclusive(SampleLengthFromRoot, LFR0, LFR1)
+			|| FMath::IsWithinInclusive(SampleLengthFromRoot, LFR1, LFR0))
+		{
+			const float Denominator = (LFR1 - LFR0);
+			OutAlpha = (FMath::Abs(Denominator) <= KINDA_SMALL_NUMBER)
+				? 0.0f
+				: ((SampleLengthFromRoot - LFR0) / Denominator);
+			OutAlpha = FMath::Clamp(OutAlpha, 0.0f, 1.0f);
+			PreviousPointIndex = P;
+			NextPointIndex = N;
+
+			return true;
+		}
+	}
+
+	PreviousPointIndex = SortedPointIndices[0];
+	NextPointIndex = SortedPointIndices.Last();
+	OutAlpha = (SampleValue < 0.5f)
+		? 0.0f
+		: 1.0f;
+
+	return true;
+}
+
+void PV::ParametricDistributionHelper::ComputeRotationAnglesForAttachment(const FVector3f& ApicalDirection, const FVector3f& BudDirectionUp,
+                                                                          const float AxilAngle, const float TwistAroundApical,
+                                                                          FVector3f& OutUpDirection, FVector3f& OutNormalDirection)
+{
+	const FVector3f ApicalDirectionNormalized = ApicalDirection.GetSafeNormal();
+	if (ApicalDirectionNormalized.IsNearlyZero())
+	{
+		OutUpDirection = FVector3f(1, 0, 0);
+		OutNormalDirection = FVector3f(0, 0, 1);
+		return;
+	}
+
+	FVector3f BudDirectionNormalized = BudDirectionUp.GetSafeNormal();
+	if (BudDirectionNormalized.IsNearlyZero())
+	{
+		FVector Axis1, Axis2;
+		FVector(ApicalDirectionNormalized).FindBestAxisVectors(Axis1, Axis2);
+		BudDirectionNormalized = FVector3f(Axis2);
+	}
+
+	FVector3f AxilAxis = FVector3f::CrossProduct(BudDirectionNormalized, ApicalDirectionNormalized).GetSafeNormal();
+	if (AxilAxis.IsNearlyZero())
+	{
+		FVector Axis1, Axis2;
+		FVector(ApicalDirectionNormalized).FindBestAxisVectors(Axis1, Axis2);
+		AxilAxis = FVector3f(Axis2);
+	}
+
+	const FQuat4f AxilQuat(AxilAxis, AxilAngle);
+	const FVector3f AxilAdjusted = AxilQuat.RotateVector(BudDirectionNormalized);
+
+	const FQuat4f TwistQuat(ApicalDirectionNormalized, TwistAroundApical);
+
+	OutUpDirection = TwistQuat.RotateVector(AxilAdjusted).GetSafeNormal();
+	OutNormalDirection = TwistQuat.RotateVector(AxilQuat.RotateVector(ApicalDirectionNormalized)).GetSafeNormal();
+}
diff --git a/Engine/Plugins/Experimental/ProceduralVegetationEditor/Source/ProceduralVegetation/Private/Helpers/PVParametricDistributionHelper.h b/Engine/Plugins/Experimental/ProceduralVegetationEditor/Source/ProceduralVegetation/Private/Helpers/PVParametricDistributionHelper.h
new file mode 100644
--- /dev/null
+++ b/Engine/Plugins/Experimental/ProceduralVegetationEditor/Source/ProceduralVegetation/Private/Helpers/PVParametricDistributionHelper.h
@@ -0,0 +1,115 @@
+// Copyright Epic Games, Inc. All Rights Reserved.
+
+#pragma once
+
+#include "CoreMinimal.h"
+#include "PVFoliageHelper.h"
+#include "Utils/PVAttributes.h"
+#include "Utils/PVFloatRamp.h"
+
+using namespace PV::DistributionConditionUtils;
+
+namespace PV::ParametricDistributionHelper
+{
+	namespace FParametricSettings
+	{
+		enum class EDistributionBasis : uint8
+		{
+			Plant,
+			Branch
+		};
+
+		struct FSpacingSettings
+		{
+			const int32 BranchDensity;
+			const float RelativeStart;
+			const float RelativeEnd;
+			const bool LimitStartGeneration;
+			const int32 StartGeneration;
+			const bool LimitEndGeneration;
+			const int32 EndGeneration;
+			const EDistributionBasis SpacingBasis;
+			const FPVFloatRamp* SpacingRamp;
+		};
+
+		struct FAngleSettings
+		{
+			const float Rotation;
+			const float AxilAngle;
+			const float RandomizeAxilAngleMinimum;
+			const float RandomizeAxilAngleMaximum;
+			const EDistributionBasis AxilAngleRampBasis;
+			const FPVFloatRamp* AxilAngleRamp;
+		};
+
+		struct FScaleSettings
+		{
+			const EDistributionBasis ScaleRampBasis;
+			const float BaseScale;
+			const float RandomizeScaleAngleMinimum;
+			const float RandomizeScaleAngleMaximum;
+			const FPVFloatRamp* ScaleRamp;
+		};
+	};
+
+	struct FAttachmentPoint
+	{
+		int32 BranchIndex;
+		int32 PointIndexA;
+		int32 PointIndexB;
+		float PointAlpha;
+		FVector3f Position;
+		FVector3f UpDirection;
+		FVector3f NormalDirection;
+
+		float Scale;
+		float LengthFromRoot;
+		bool bIsTipInstance;
+		int32 TipPointIndex;
+
+		ConditionPicker::FPointConditionSample ConditionSample;
+	};
+
+	void ComputeAttachmentPoints(
+		FManagedArrayCollection& OutCollection,
+		const FNormalizedPointCaches& NormalizedPointCaches,
+		const FParametricSettings::FSpacingSettings& SpacingSettings,
+		const FParametricSettings::FAngleSettings& AngleSettings,
+		const FParametricSettings::FScaleSettings& ScaleSettings,
+		FRandomStream& RandomStream,
+		TArray<FAttachmentPoint>& OutAttachmentPoints);
+
+	TMap<int32, float> ComputeBranchLengths(
+		FPointLengthFromRootAttributeConstView PointLengthFromRootAttribute,
+		FBranchPointsAttributeConstView BranchPointsAttribute);
+
+	FVector3f SampleVectorByFloat(
+		const float SampleValue,
+		const TArray<int32>& SortedPointIndices,
+		FPointPlantGradientAttributeConstView PointPlantGradientAttribute,
+		FPointPositionAttributeConstView PointPositionAttribute);
+
+	float SampleFloatByFloat(
+		const float SampleValue,
+		const TArray<int32>& SortedPointIndices,
+		FPointPlantGradientAttributeConstView PointPlantGradientAttribute,
+		FPointLengthFromRootAttributeConstView PointLengthFromRootAttribute);
+
+	bool FindPointByNormalizedLengthFromRoot(
+		float SampleValue,
+		const TArray<int32>& SortedPointIndices,
+		FPointLengthFromRootAttributeConstView PointLengthFromRootAttribute,
+		const float LFRofRoot,
+		const float LFRofTip,
+		int32& PreviousPointIndex,
+		int32& NextPointIndex,
+		float& OutAlpha);
+
+	void ComputeRotationAnglesForAttachment(
+		const FVector3f& ApicalDirection,
+		const FVector3f& BudDirectionUp,
+		const float AxilAngle,
+		const float TwistAroundApical,
+		FVector3f& OutUpDirection,
+		FVector3f& OutNormalDirection);
+}
diff --git a/Engine/Plugins/Experimental/ProceduralVegetationEditor/Source/ProceduralVegetation/Private/Implementations/PVFoliage.cpp b/Engine/Plugins/Experimental/ProceduralVegetationEditor/Source/ProceduralVegetation/Private/Implementations/PVFoliage.cpp
--- a/Engine/Plugins/Experimental/ProceduralVegetationEditor/Source/ProceduralVegetation/Private/Implementations/PVFoliage.cpp
+++ b/Engine/Plugins/Experimental/ProceduralVegetationEditor/Source/ProceduralVegetation/Private/Implementations/PVFoliage.cpp
@@ -14,14 +14,25 @@
 #include "Helpers/PVFoliageHelper.h"
 #include "Helpers/PVPhyllotaxyHelper.h"
 
- 
-void FPVFoliage::DistributeFoliage(FManagedArrayCollection& OutCollection, const FManagedArrayCollection& SourceCollection, const FManagedArrayCollection& FoliageCollection,
-                                   const FDistributionSettings& DistributionSettings, const FScaleSettings& ScaleSettings,
-                                   const FVectorSettings& VectorSettings, const FPhyllotaxySettings& PhyllotaxySettings,
-                                   const FMiscSettings& MiscSettings, const FPVFoliageDistributionConditions& DistributionConditions)
+using namespace PV::DistributionConditionUtils;
+using PV::DistributionConditionUtils::FConditionAttributes;
+
+void FPVFoliage::DistributeFoliageWithHormoneBasedSettings(
+	FManagedArrayCollection& OutCollection,
+	const FManagedArrayCollection& FoliageCollection,
+	const PV::HormoneBasedSettings::FDistributionSettings& DistributionSettings,
+	const PV::HormoneBasedSettings::FScaleSettings& ScaleSettings,
+	const PV::HormoneBasedSettings::FVectorSettings& VectorSettings,
+	const PV::HormoneBasedSettings::FPhyllotaxySettings& PhyllotaxySettings,
+	const int32 RandomSeed,
+	const FPVFoliageDistributionConditions& DistributionConditions)
 {
-	using namespace PV::DistributionConditionUtils;
-	using PV::DistributionConditionUtils::FConditionAttributes;
+	TRACE_CPUPROFILER_EVENT_SCOPE(PVFoliage::DistributeFoliageWithHormoneBasedSettings);
+	
+	if (!(DistributionSettings.InstanceSpacingRamp && ScaleSettings.ScaleRamp && VectorSettings.AxilAngleRamp))
+	{
+		return;
+	}
 	
 	PV::Facades::FFoliageFacade FoliageFacadeOutput(OutCollection);
 	PV::Facades::FFoliageFacade FoliagePaletteFacade(FoliageCollection);
@@ -35,10 +46,10 @@ void FPVFoliage::DistributeFoliage(FManagedArrayCollection& OutCollection, const
 	FoliageFacadeOutput.SetFoliageInfos(FoliagePaletteFacade.GetFoliageInfos());
  
 	FRandomStream RandomStream;
-	RandomStream.Initialize(MiscSettings.RandomSeed);
+	RandomStream.Initialize(RandomSeed);
  
-	PV::Facades::FBranchFacade BranchFacadeSource(SourceCollection);
-	PV::Facades::FPointFacade PointFacadeSource(SourceCollection);
+	PV::Facades::FBranchFacade BranchFacadeSource(OutCollection);
+	PV::Facades::FPointFacade PointFacadeSource(OutCollection);
 
 	TManagedArrayAccessor<TArray<FVector3f>> BudDirectionAttribute(OutCollection, PV::AttributeNames::BudDirection,
 			PV::GroupNames::PointGroup);
@@ -122,11 +133,10 @@ void FPVFoliage::DistributeFoliage(FManagedArrayCollection& OutCollection, const
 		BuildOptions.bComputeHealth = DistributionConditions.IsActiveCondition(EFoliageDistributionCondition::Health);
 
 		BuildNormalizedPointCaches(
-			SourceCollection,
+			OutCollection,
 			BuildOptions,
 			NormalizedPointCaches);
 		
-		// Compute attachment points in terms of normalized distances
 		OutCollection.EmptyGroup(PV::GroupNames::FoliageGroup);
  
 		// Collect non-empty mesh name indices
@@ -589,4 +599,121 @@ void FPVFoliage::DistributeFoliage(FManagedArrayCollection& OutCollection, const
 			FoliageFacadeOutput.SetFoliageIdsArray(BranchId, FoliageIDs);
 		}
 	}
-}
\ No newline at end of file
+}
+
+void FPVFoliage::DistributeFoliageWithParametricSettings(
+	FManagedArrayCollection& OutCollection,
+	const FManagedArrayCollection& FoliageCollection,
+	const FParametricSettings::FSpacingSettings& SpacingSettings,
+	const FParametricSettings::FAngleSettings& AngleSettings,
+	const FParametricSettings::FScaleSettings& ScaleSettings,
+	const int32 RandomSeed,
+	const FPVFoliageDistributionConditions& DistributionConditions)
+{
+	TRACE_CPUPROFILER_EVENT_SCOPE(PVFoliage::DistributeFoliageWithParametricSettings);
+	
+	if (!(SpacingSettings.SpacingRamp && AngleSettings.AxilAngleRamp && ScaleSettings.ScaleRamp))
+	{
+		return;
+	}
+
+	PV::Facades::FFoliageFacade FoliageFacadeOutput(OutCollection);
+	PV::Facades::FFoliageFacade FoliagePaletteFacade(FoliageCollection);
+	
+	if (const int32 NumFoliageMeshes = FoliagePaletteFacade.NumFoliageInfo(); NumFoliageMeshes == 0)
+	{
+		UE_LOG(LogProceduralVegetation, Warning, TEXT("There are no foliage meshes available in input for distribution"));
+		return;
+	}
+
+	FoliageFacadeOutput.SetFoliageInfos(FoliagePaletteFacade.GetFoliageInfos());
+	
+	FRandomStream RandomStream;
+	RandomStream.Initialize(RandomSeed);
+	
+	FNormalizedPointCaches NormalizedPointCaches;
+	FBuildOptions BuildOptions;
+	BuildOptions.bComputePScales = DistributionConditions.IsActiveCondition(EFoliageDistributionCondition::Scale);
+	BuildOptions.bComputeLight = DistributionConditions.IsActiveCondition(EFoliageDistributionCondition::Light);
+	BuildOptions.bComputeHealth = DistributionConditions.IsActiveCondition(EFoliageDistributionCondition::Health);
+
+	BuildNormalizedPointCaches(
+		OutCollection,
+		BuildOptions,
+		NormalizedPointCaches);
+	
+	TArray<FAttachmentPoint> AttachmentPoints;
+	ComputeAttachmentPoints(
+		OutCollection, 
+		NormalizedPointCaches, 
+		SpacingSettings, 
+		AngleSettings, 
+		ScaleSettings, 
+		RandomStream, 
+		AttachmentPoints);
+	
+	OutCollection.EmptyGroup(PV::GroupNames::FoliageGroup);
+	
+	const int32 NumFoliageInfo = FoliagePaletteFacade.NumFoliageInfo();
+	const TArray<FPVFoliageInfo>& FoliageInfos = FoliagePaletteFacade.GetFoliageInfos();
+
+	auto IsValidCandidate = [&](int32 CandidateIndex)
+		{
+			return FoliageInfos.IsValidIndex(CandidateIndex) && FoliageInfos[CandidateIndex].IsValid();
+		};
+
+	auto GetCandidateValue = [&](int32 CandidateIndex, EFoliageDistributionCondition Condition) -> float
+		{
+			return FoliageInfos[CandidateIndex].Attributes.GetAttributeValue(Condition);
+		};
+	
+	TManagedArrayAccessor<FVector3f> PivotPointsAttribute(OutCollection, PV::AttributeNames::FoliagePivotPoint,
+				PV::GroupNames::FoliageGroup);
+	
+	for (const auto& AttachmentPoint : AttachmentPoints)
+	{
+		const int32 PickedFoliageIndex = ConditionPicker::PickCandidateIndex(
+			NumFoliageInfo,
+			DistributionConditions,
+			AttachmentPoint.ConditionSample,
+			RandomStream,
+			GetCandidateValue,
+			IsValidCandidate);
+		
+		if (PickedFoliageIndex == INDEX_NONE)
+		{
+			continue;
+		}
+
+		const int32 NewIndex = PivotPointsAttribute.AddElements(1);
+		FoliageFacadeOutput.SetFoliageEntry(NewIndex, {
+			.NameId = PickedFoliageIndex,
+			.BranchId = AttachmentPoint.BranchIndex,
+			.PivotPoint = AttachmentPoint.Position,
+			.UpVector = AttachmentPoint.UpDirection,
+			.NormalVector = AttachmentPoint.NormalDirection,
+			.Scale = AttachmentPoint.Scale,
+			.LengthFromRoot = AttachmentPoint.LengthFromRoot,
+			.ConditionUpAlignment = AttachmentPoint.ConditionSample.UpAlignment,
+			.ConditionTip = AttachmentPoint.ConditionSample.Tip,
+			.ConditionLight = AttachmentPoint.ConditionSample.Light,
+			.ConditionScale = AttachmentPoint.ConditionSample.Scale,
+			.ConditionHealth = AttachmentPoint.ConditionSample.Health
+		});
+	}
+	
+	// Re-calculate Foliage IDs for each branch
+	TMap<int32, TArray<int32>> BranchIndexToFoliageIDs;
+	for (int32 i = 0; i < FoliageFacadeOutput.NumFoliageEntries(); ++i)
+	{
+		const int32 BranchIndex = FoliageFacadeOutput.GetFoliageEntry(i).BranchId;
+		TArray<int32>& IDs = BranchIndexToFoliageIDs.FindOrAdd(BranchIndex);
+		IDs.Add(i);
+	}
+	for (const TPair<int32, TArray<int32>>& Pair : BranchIndexToFoliageIDs)
+	{
+		int32 BranchId = Pair.Key;
+		const TArray<int32>& FoliageIDs = Pair.Value;
+		FoliageFacadeOutput.SetFoliageIdsArray(BranchId, FoliageIDs);
+	}
+}
diff --git a/Engine/Plugins/Experimental/ProceduralVegetationEditor/Source/ProceduralVegetation/Private/Implementations/PVFoliage.h b/Engine/Plugins/Experimental/ProceduralVegetationEditor/Source/ProceduralVegetation/Private/Implementations/PVFoliage.h
--- a/Engine/Plugins/Experimental/ProceduralVegetationEditor/Source/ProceduralVegetation/Private/Implementations/PVFoliage.h
+++ b/Engine/Plugins/Experimental/ProceduralVegetationEditor/Source/ProceduralVegetation/Private/Implementations/PVFoliage.h
@@ -1,66 +1,79 @@
 // Copyright Epic Games, Inc. All Rights Reserved.
- 
+
 #pragma once
- 
+
 #include "CoreMinimal.h"
 #include "GeometryCollection/ManagedArrayCollection.h"
+#include "Helpers/PVParametricDistributionHelper.h"
 
 struct FPVFoliageDistributionConditions;
 enum class EPhyllotaxyFormation : uint8;
 enum class EPhyllotaxyType : uint8;
 struct FPVFloatRamp;
- 
-struct FDistributionSettings
-{
-	const float EthyleneThreshold;
-	const bool OverrideDistribution;
-	const float InstanceSpacing;
-	const FPVFloatRamp* InstanceSpacingRamp;
-	const float InstanceSpacingRampEffect;
-	const int32 MaxPerBranch;
-};
- 
-struct FScaleSettings
-{
-	const float BaseScale;
-	const float BranchScaleImpact;
-	const float MinScale;
-	const float MaxScale;
-	const float RandomScaleMin;
-	const float RandomScaleMax;
-	const FPVFloatRamp* ScaleRamp;
-};
- 
-struct FVectorSettings
-{
-	const bool OverrideAxilAngle;
-	const float AxilAngle;
-	const FPVFloatRamp* AxilAngleRamp;
-	const float AxilAngleRampUpperValue;
-	const float AxilAngleRampEffect;
-};
- 
-struct FPhyllotaxySettings
-{
-	const bool OverridePhyllotaxy;
-	const EPhyllotaxyType PhyllotaxyType;
-	const EPhyllotaxyFormation PhyllotaxyFormation;
-	const int32 MinimumNodeBuds;
-	const int32 MaximumNodeBuds;
-	const float PhyllotaxyAdditionalAngle;
-};
 
- 
-struct FMiscSettings
+using namespace PV::ParametricDistributionHelper;
+
+namespace PV::HormoneBasedSettings
 {
-	const int32 RandomSeed;
-};
+	struct FDistributionSettings
+	{
+		const float EthyleneThreshold;
+		const bool OverrideDistribution;
+		const float InstanceSpacing;
+		const FPVFloatRamp* InstanceSpacingRamp;
+		const float InstanceSpacingRampEffect;
+		const int32 MaxPerBranch;
+	};
+
+	struct FScaleSettings
+	{
+		const float BaseScale;
+		const float BranchScaleImpact;
+		const float MinScale;
+		const float MaxScale;
+		const float RandomScaleMin;
+		const float RandomScaleMax;
+		const FPVFloatRamp* ScaleRamp;
+	};
+
+	struct FVectorSettings
+	{
+		const bool OverrideAxilAngle;
+		const float AxilAngle;
+		const FPVFloatRamp* AxilAngleRamp;
+		const float AxilAngleRampUpperValue;
+		const float AxilAngleRampEffect;
+	};
+
+	struct FPhyllotaxySettings
+	{
+		const bool OverridePhyllotaxy;
+		const EPhyllotaxyType PhyllotaxyType;
+		const EPhyllotaxyFormation PhyllotaxyFormation;
+		const int32 MinimumNodeBuds;
+		const int32 MaximumNodeBuds;
+		const float PhyllotaxyAdditionalAngle;
+	};
+}
 
- 
 struct FPVFoliage
 {
-	static void DistributeFoliage(FManagedArrayCollection& OutCollection, const FManagedArrayCollection& SourceCollection, const FManagedArrayCollection& FoliageCollection,  
-	                              const FDistributionSettings& DistributionSettings, const FScaleSettings& ScaleSettings,
-	                              const FVectorSettings& VectorSettings, const FPhyllotaxySettings& PhyllotaxySettings,
-	                              const FMiscSettings& MiscSettings, const FPVFoliageDistributionConditions& DistributionConditions);
-};
\ No newline at end of file
+	static void DistributeFoliageWithHormoneBasedSettings(
+		FManagedArrayCollection& OutCollection,
+		const FManagedArrayCollection& FoliageCollection,
+		const PV::HormoneBasedSettings::FDistributionSettings& DistributionSettings,
+		const PV::HormoneBasedSettings::FScaleSettings& ScaleSettings,
+		const PV::HormoneBasedSettings::FVectorSettings& VectorSettings,
+		const PV::HormoneBasedSettings::FPhyllotaxySettings& PhyllotaxySettings,
+		const int32 RandomSeed,
+		const FPVFoliageDistributionConditions& DistributionConditions);
+
+	static void DistributeFoliageWithParametricSettings(
+		FManagedArrayCollection& OutCollection,
+		const FManagedArrayCollection& FoliageCollection,
+		const FParametricSettings::FSpacingSettings& SpacingSettings,
+		const FParametricSettings::FAngleSettings& AngleSettings,
+		const FParametricSettings::FScaleSettings& ScaleSettings,
+		const int32 RandomSeed,
+		const FPVFoliageDistributionConditions& DistributionConditions);
+};
diff --git a/Engine/Plugins/Experimental/ProceduralVegetationEditor/Source/ProceduralVegetation/Private/Implementations/PVGrafter.cpp b/Engine/Plugins/Experimental/ProceduralVegetationEditor/Source/ProceduralVegetation/Private/Implementations/PVGrafter.cpp
--- a/Engine/Plugins/Experimental/ProceduralVegetationEditor/Source/ProceduralVegetation/Private/Implementations/PVGrafter.cpp
+++ b/Engine/Plugins/Experimental/ProceduralVegetationEditor/Source/ProceduralVegetation/Private/Implementations/PVGrafter.cpp
@@ -18,6 +18,7 @@
 #include "Helpers/PVAttributesHelper.h"
 
 using namespace PV::DistributionConditionUtils;
+using namespace PV::ParametricDistributionHelper;
 using PV::DistributionConditionUtils::FConditionAttributes;
 
 template <typename Fn>
@@ -140,8 +141,9 @@ void FPVGrafter::RescaleBranchForGraftingOnTip(FManagedArrayCollection& OutColle
 	}
 }
 
-void FPVGrafter::SetGraft(FManagedArrayCollection& OutCollection, const FManagedArrayCollection& SourceSkeletonCollection,
-                          const FManagedArrayCollection& SourceGraftCollection, const FPVGraftAttributes& InAttributes)
+void FPVGrafter::SetGraft(FManagedArrayCollection& OutCollection,
+                          const FManagedArrayCollection& SourceGraftCollection,
+                          const FPVGraftAttributes& InAttributes)
 {
 	// Transform and set points
 	PV::Facades::FPointFacade PointFacadeOut(OutCollection);
@@ -645,23 +647,20 @@ TArray<FManagedArrayCollection> FPVGrafter::ExtractIndividualPlants(const FManag
 	return IndividualPlants;
 }
 
-void FPVGrafter::DistributeGraftWithHormoneBasedSettings(FManagedArrayCollection& OutCollection,
-                                                         const FManagedArrayCollection& SourceSkeletonCollection,
-                                                         TArray<FManagedArrayCollection>& SourceGraftCollections,
-                                                         const FHormoneBasedSettings::FDistributionSettings& DistributionSettings,
-                                                         const FHormoneBasedSettings::FPhyllotaxySettings& PhyllotaxySettings,
-                                                         const FHormoneBasedSettings::FScaleSettings& ScaleSettings,
-                                                         const FHormoneBasedSettings::FVectorSettings& VectorSettings,
-                                                         const FMiscSettings& MiscSettings,
-                                                         const FPVFoliageDistributionConditions& DistributionConditions)
+void FPVGrafter::DistributeGraftWithHormoneBasedSettings(
+	FManagedArrayCollection& OutCollection,
+	TArray<FManagedArrayCollection>& SourceGraftCollections,
+	const FHormoneBasedSettings::FDistributionSettings& DistributionSettings,
+	const FHormoneBasedSettings::FPhyllotaxySettings& PhyllotaxySettings,
+	const FHormoneBasedSettings::FScaleSettings& ScaleSettings,
+	const FHormoneBasedSettings::FVectorSettings& VectorSettings,
+	const int32 RandomSeed,
+	const FPVFoliageDistributionConditions& DistributionConditions)
 {
 	TRACE_CPUPROFILER_EVENT_SCOPE(PVGrafter::DistributeGraftWithHormoneBasedSettings);
 
 	FRandomStream RandomStream;
-	RandomStream.Initialize(MiscSettings.RandomSeed);
-
-	PV::Facades::FBranchFacade BranchFacadeSourceSkeleton(SourceSkeletonCollection);
-	PV::Facades::FPointFacade PointFacadeSourceSkeleton(SourceSkeletonCollection);
+	RandomStream.Initialize(RandomSeed);
 
 	auto PointBranchGradientAttribute = PV::FPointBranchGradientAttribute::GetAttribute(OutCollection);
 	auto BranchParentNumberAttribute = PV::FBranchParentNumberAttribute::GetAttribute(OutCollection);
@@ -670,9 +669,10 @@ void FPVGrafter::DistributeGraftWithHormoneBasedSettings(FManagedArrayCollection
 	auto BranchPointsAttribute = PV::FBranchPointsAttribute::GetAttribute(OutCollection);
 
 	// Only need to do this temporarily since legacy skeleton has inconsistent data in children
-	// TODO: Remove this once support for legacy skeletons is droppped
+	// TODO: Remove this once support for legacy skeletons is dropped
 	PV::Facades::FBranchFacade BranchFacadeOut(OutCollection);
 	BranchFacadeOut.RecomputeBranchChildren();
+	BranchFacadeOut.RecomputeBranchParents();
 
 	ComputeBranchGradientAttribute(
 		PointBranchGradientAttribute,
@@ -688,7 +688,7 @@ void FPVGrafter::DistributeGraftWithHormoneBasedSettings(FManagedArrayCollection
 	BuildOptions.bComputeHealth = DistributionConditions.IsActiveCondition(EFoliageDistributionCondition::Health);
 
 	BuildNormalizedPointCaches(
-		SourceSkeletonCollection,
+		OutCollection,
 		BuildOptions,
 		NormalizedPointCaches);
 
@@ -710,12 +710,9 @@ void FPVGrafter::DistributeGraftWithHormoneBasedSettings(FManagedArrayCollection
 		CandidateGraftData[i] = MetaInfoFacade.GetGraftAttributes();
 	}
 
-	TManagedArrayAccessor<TArray<FVector3f>> BudDirectionAttribute(OutCollection, PV::AttributeNames::BudDirection,
-		PV::GroupNames::PointGroup);
-	TManagedArrayAccessor<float> BranchGradientsAttribute(OutCollection, PV::AttributeNames::BranchGradient,
-		PV::GroupNames::PointGroup);
-	TManagedArrayAccessor<TArray<float>> BudHormoneLevels(OutCollection, PV::AttributeNames::BudHormoneLevels,
-		PV::GroupNames::PointGroup);
+	auto BudDirectionAttribute = PV::FBudDirectionAttribute::GetAttribute(OutCollection);
+	auto BranchGradientsAttribute = PV::FPointBranchGradientAttribute::GetAttribute(OutCollection);
+	auto BudHormoneLevels = PV::FBudHormoneLevelsAttribute::GetAttribute(OutCollection);
 
 	float PhyllotaxyFormationLeaf = 180;
 	float MinBudsLeaf = 1;
@@ -760,11 +757,9 @@ void FPVGrafter::DistributeGraftWithHormoneBasedSettings(FManagedArrayCollection
 		}
 	}
 
-	const int32 NumBranches = BranchFacadeSourceSkeleton.GetElementCount();
-	const TManagedArray<int32>& SourceSkeletonBudNumbers = PointFacadeSourceSkeleton.GetBudNumbersAttribute();
-
-	check(SourceSkeletonBudNumbers.Num() > 0);
-	const int32 MaxBudNumber = *Algo::MaxElement(SourceSkeletonBudNumbers.GetConstArray());
+	auto BudNumberAttribute = PV::FPointBudNumberAttribute::GetAttribute(OutCollection);
+	check(BudNumberAttribute.Num() > 0);
+	const int32 MaxBudNumber = *Algo::MaxElement(BudNumberAttribute);
 	int32 CurrentBudNumber = MaxBudNumber + 1;
 
 	auto IsValidCandidate = [&](int32 CandidateIndex) -> bool
@@ -797,6 +792,10 @@ void FPVGrafter::DistributeGraftWithHormoneBasedSettings(FManagedArrayCollection
 		};
 
 	// Compute attachment points in terms of normalized distances
+	auto PointLengthFromRootAttribute = PV::FPointLengthFromRootAttribute::GetAttribute(OutCollection);
+	auto PointScaleAttribute = PV::FPointScaleAttribute::GetAttribute(OutCollection);
+	auto PointPositionAttribute = PV::FPointPositionAttribute::GetAttribute(OutCollection);
+	const int32 NumBranches = BranchPointsAttribute.Num();
 	for (int32 BranchIndex = 0; BranchIndex < NumBranches; ++BranchIndex)
 	{
 		FConditionAttributes ConditionAttrs;
@@ -808,15 +807,15 @@ void FPVGrafter::DistributeGraftWithHormoneBasedSettings(FManagedArrayCollection
 		TArray<FVector3f> AttachmentPointsNormalVectors;
 		TArray<float> AttachmentPointsScales;
 		TArray<int32> AttachmentPointsPointIndices;
-		const TArray<int32>& BranchPoints = BranchFacadeSourceSkeleton.GetPoints(BranchIndex);
+		const TArray<int32>& BranchPoints = BranchPointsAttribute[BranchIndex];
 
 		if (BranchPoints.Num() < 2)
 		{
 			continue;
 		}
 
-		const float FirstPointLengthFromRoot = PointFacadeSourceSkeleton.GetLengthFromRoot(BranchPoints[0]);
-		const float LastPointLengthFromRoot = PointFacadeSourceSkeleton.GetLengthFromRoot(BranchPoints.Last());
+		const float FirstPointLengthFromRoot = PointLengthFromRootAttribute[BranchPoints[0]];
+		const float LastPointLengthFromRoot = PointLengthFromRootAttribute[BranchPoints.Last()];
 		const float BranchLength = LastPointLengthFromRoot - FirstPointLengthFromRoot;
 
 		if (BranchLength <= KINDA_SMALL_NUMBER)
@@ -851,9 +850,9 @@ void FPVGrafter::DistributeGraftWithHormoneBasedSettings(FManagedArrayCollection
 		// Compute Position, Angles, Scale and LengthFromRoot for attachment points
 
 		// Compute based on first point attributes
-		const FVector3f InitialApicalDirection(BudDirectionAttribute[BranchPoints[0]][0]);
-		FVector3f InitialAxillaryDirection(BudDirectionAttribute[BranchPoints[0]][1]);
-		const FVector3f InitialUpVector(BudDirectionAttribute[BranchPoints[0]][5]);
+		const FVector3f InitialApicalDirection(BudDirectionAttribute[BranchPoints[0]].Apical);
+		FVector3f InitialAxillaryDirection(BudDirectionAttribute[BranchPoints[0]].Axillary);
+		const FVector3f InitialUpVector(BudDirectionAttribute[BranchPoints[0]].UpVector);
 
 		if (ResetPhyllotaxyLeaf)
 		{
@@ -875,12 +874,12 @@ void FPVGrafter::DistributeGraftWithHormoneBasedSettings(FManagedArrayCollection
 
 				const int32 CurrentPointIndex = BranchPoints[BranchPointIndex];
 				const int32 NextPointIndex = BranchPoints[BranchPointIndex + 1];
-				const FVector3f CurrentPointPosition = PointFacadeSourceSkeleton.GetPosition(CurrentPointIndex);
-				const FVector3f NextPointPosition = PointFacadeSourceSkeleton.GetPosition(NextPointIndex);
-				float EthyleneLevelAtCurrentPoint = BudHormoneLevels[CurrentPointIndex][4];
-				float EthyleneLevelAtNextPoint = BudHormoneLevels[NextPointIndex][4];
-				const float CurrentPointLengthFromRoot = PointFacadeSourceSkeleton.GetLengthFromRoot(CurrentPointIndex);
-				const float NextPointLengthFromRoot = PointFacadeSourceSkeleton.GetLengthFromRoot(NextPointIndex);
+				const FVector3f CurrentPointPosition = PointPositionAttribute[CurrentPointIndex];
+				const FVector3f NextPointPosition = PointPositionAttribute[NextPointIndex];
+				float EthyleneLevelAtCurrentPoint = BudHormoneLevels[CurrentPointIndex].Ethylene;
+				float EthyleneLevelAtNextPoint = BudHormoneLevels[NextPointIndex].Ethylene;
+				const float CurrentPointLengthFromRoot = PointLengthFromRootAttribute[CurrentPointIndex];
+				const float NextPointLengthFromRoot = PointLengthFromRootAttribute[NextPointIndex];
 				const float CurrentPointLfrNormalized = FMath::GetMappedRangeValueClamped(
 					FVector2f(FirstPointLengthFromRoot, LastPointLengthFromRoot),
 					FVector2f(0.0f, 1.0f), CurrentPointLengthFromRoot);
@@ -945,8 +944,8 @@ void FPVGrafter::DistributeGraftWithHormoneBasedSettings(FManagedArrayCollection
 						? ScaleSettings.ScaleRamp->GetRichCurveConst()->Eval(1.0f - BranchGradient)
 						: 1.0f;
 
-					const float CurrentPointScale = PointFacadeSourceSkeleton.GetPointScale(CurrentPointIndex);
-					const float NextPointScale = PointFacadeSourceSkeleton.GetPointScale(NextPointIndex);
+					const float CurrentPointScale = PointScaleAttribute[CurrentPointIndex];
+					const float NextPointScale = PointScaleAttribute[NextPointIndex];
 					const float BranchScale = FMath::Lerp(CurrentPointScale, NextPointScale, DistanceAlpha);
 
 					const float BranchScaledValue = BranchGradientRampValue * (BranchScale + 0.1f);
@@ -964,8 +963,8 @@ void FPVGrafter::DistributeGraftWithHormoneBasedSettings(FManagedArrayCollection
 
 					AttachmentPointsScales.Add(Scale);
 
-					const FVector3f CurrentPointApicalDirectionVector(BudDirectionAttribute[CurrentPointIndex][0]);
-					const FVector3f NextPointApicalDirectionVector(BudDirectionAttribute[NextPointIndex][0]);
+					const FVector3f CurrentPointApicalDirectionVector(BudDirectionAttribute[CurrentPointIndex].Apical);
+					const FVector3f NextPointApicalDirectionVector(BudDirectionAttribute[NextPointIndex].Apical);
 					const FVector3f BudApicalDirectionVector = FMath::Lerp(CurrentPointApicalDirectionVector,
 						NextPointApicalDirectionVector, DistanceAlpha);
 
@@ -1067,8 +1066,8 @@ void FPVGrafter::DistributeGraftWithHormoneBasedSettings(FManagedArrayCollection
 					}
 
 					// Compute Normal vectors based on light optimal direction
-					const FVector3f CurrentPointLightOptimalVector(BudDirectionAttribute[CurrentPointIndex][2]);
-					const FVector3f NextPointLightOptimalVector(BudDirectionAttribute[NextPointIndex][2]);
+					const FVector3f CurrentPointLightOptimalVector(BudDirectionAttribute[CurrentPointIndex].LightOptimal);
+					const FVector3f NextPointLightOptimalVector(BudDirectionAttribute[NextPointIndex].LightOptimal);
 					const FVector3f BudLightOptimalVector = FMath::Lerp(CurrentPointLightOptimalVector, NextPointLightOptimalVector,
 						DistanceAlpha);
 
@@ -1124,273 +1123,23 @@ void FPVGrafter::DistributeGraftWithHormoneBasedSettings(FManagedArrayCollection
 
 			const FManagedArrayCollection& PickedSourceGraft = GraftSourcesExtracted[PickedCandidateIndex];
 
-			SetGraft(OutCollection, SourceSkeletonCollection, PickedSourceGraft, Attributes);
-		}
-	}
-}
-
-TMap<int32, float> FPVGrafter::ComputeBranchLengths(PV::FPointLengthFromRootAttributeConstView PointLengthFromRootAttribute,
-                                                    PV::FBranchPointsAttributeConstView BranchPointsAttribute)
-{
-	const int32 NumOfBranches = BranchPointsAttribute.Num();
-	TMap<int32, float> BranchIndicesToBranchLengths;
-	BranchIndicesToBranchLengths.Reserve(NumOfBranches);
-
-	for (int32 BranchIndex = 0; BranchIndex < NumOfBranches; BranchIndex++)
-	{
-		const TArray<int32>& BranchPoints = BranchPointsAttribute[BranchIndex];
-		if (BranchPoints.Num() < 2)
-		{
-			BranchIndicesToBranchLengths.Add(BranchIndex, 0.0f);
-			continue;
-		}
-
-		TArray<int32> SortedPointIndices = BranchPoints;
-		SortedPointIndices.Sort([PointLengthFromRootAttribute](const int32& A, const int32& B)
-			{
-				return PointLengthFromRootAttribute[A] < PointLengthFromRootAttribute[B];
-			});
-
-		const int32 RootPointIndex = SortedPointIndices[0];
-		const int32 BranchTipPointIndex = SortedPointIndices.Last();
-
-		const float RootPointLFR = PointLengthFromRootAttribute[RootPointIndex];
-		const float TipPointLFR = PointLengthFromRootAttribute[BranchTipPointIndex];
-		const float BranchLength = TipPointLFR - RootPointLFR;
-
-		BranchIndicesToBranchLengths.Add(BranchIndex, BranchLength);
-	}
-
-	return BranchIndicesToBranchLengths;
-}
-
-FVector3f FPVGrafter::SampleVectorByFloat(const float SampleValue, const TArray<int32>& SortedPointIndices,
-                                          PV::FPointPlantGradientAttributeConstView PointPlantGradientAttribute,
-                                          PV::FPointPositionAttributeConstView PointPositionAttribute)
-{
-	if (SortedPointIndices.Num() < 2)
-	{
-		return FVector3f();
-	}
-
-	const int32 FirstPointIndex = SortedPointIndices[0];
-	const int32 LastPointIndex = SortedPointIndices.Last();
-	float Value0 = 1.0f - PointPlantGradientAttribute[FirstPointIndex];
-	float Value1 = 1.0f - PointPlantGradientAttribute[LastPointIndex];
-	FVector3f OutValue(Value0, Value0, Value0);
-
-	if (FMath::IsWithinInclusive(SampleValue, Value0, Value1))
-	{
-		for (int32 i = 1; i < SortedPointIndices.Num(); i++)
-		{
-			Value1 = 1.0f - PointPlantGradientAttribute[SortedPointIndices[i]];
-			if (FMath::IsWithinInclusive(SampleValue, Value0, Value1))
-			{
-				FVector3f P0 = PointPositionAttribute[SortedPointIndices[i - 1]];
-				FVector3f P1 = PointPositionAttribute[SortedPointIndices[i]];
-
-				float Blend = FMath::GetMappedRangeValueClamped(
-					FVector2f(Value0, Value1),
-					FVector2f(0.0f, 1.0f),
-					SampleValue);
-
-				OutValue = FMath::Lerp(P0, P1, Blend);
-			}
-
-			Value0 = Value1;
+			SetGraft(OutCollection, PickedSourceGraft, Attributes);
 		}
 	}
-	else
-	{
-		if (FMath::Abs(Value0 - SampleValue) > FMath::Abs(Value1 - SampleValue))
-		{
-			OutValue = FVector3f(Value1, Value1, Value1);
-		}
-	}
-
-	return OutValue;
-}
-
-float FPVGrafter::SampleFloatByFloat(const float SampleValue, const TArray<int32>& SortedPointIndices,
-                                     PV::FPointPlantGradientAttributeConstView PointPlantGradientAttribute,
-                                     PV::FPointLengthFromRootAttributeConstView PointLengthFromRootAttribute)
-{
-	if (SortedPointIndices.Num() < 2)
-	{
-		return -1.0f;
-	}
-
-	const int32 FirstPointIndex = SortedPointIndices[0];
-	const int32 LastPointIndex = SortedPointIndices.Last();
-	float Value0 = 1.0f - PointPlantGradientAttribute[FirstPointIndex];
-	float Value1 = 1.0f - PointPlantGradientAttribute[LastPointIndex];
-	float OutValue = Value0;
-
-	if (FMath::IsWithinInclusive(SampleValue, Value0, Value1))
-	{
-		for (int32 i = 1; i < SortedPointIndices.Num(); i++)
-		{
-			Value1 = 1.0f - PointPlantGradientAttribute[SortedPointIndices[i]];
-			if (FMath::IsWithinInclusive(SampleValue, Value0, Value1))
-			{
-				float P0 = PointLengthFromRootAttribute[SortedPointIndices[i - 1]];
-				float P1 = PointLengthFromRootAttribute[SortedPointIndices[i]];
-
-				float Blend = FMath::GetMappedRangeValueClamped(
-					FVector2f(Value0, Value1),
-					FVector2f(0.0f, 1.0f),
-					SampleValue);
-
-				OutValue = FMath::Lerp(P0, P1, Blend);
-			}
-
-			Value0 = Value1;
-		}
-	}
-	else
-	{
-		if (FMath::Abs(Value0 - SampleValue) > FMath::Abs(Value1 - SampleValue))
-		{
-			OutValue = Value1;
-		}
-	}
-
-	return OutValue;
-}
-
-bool FPVGrafter::FindPointByNormalizedLengthFromRoot(float SampleValue, const TArray<int32>& SortedPointIndices,
-                                                     PV::FPointLengthFromRootAttributeConstView PointLengthFromRootAttribute,
-                                                     const float LFRofRoot, const float LFRofTip, int32& PreviousPointIndex, int32& NextPointIndex,
-                                                     float& OutAlpha)
-{
-	if (SortedPointIndices.Num() < 2)
-	{
-		return false;
-	}
-
-	SampleValue = FMath::Clamp(SampleValue, 0.0f, 1.0f);
-
-	const float SampleLengthFromRoot = FMath::GetMappedRangeValueClamped(
-		FVector2f(0.0f, 1.0f),
-		FVector2f(LFRofRoot, LFRofTip),
-		SampleValue);
-
-	for (int32 i = 0; i < SortedPointIndices.Num() - 1; ++i)
-	{
-		const int32 P = SortedPointIndices[i];
-		const int32 N = SortedPointIndices[i + 1];
-		const float LFR0 = PointLengthFromRootAttribute[P];
-		const float LFR1 = PointLengthFromRootAttribute[N];
-
-		if (FMath::IsWithinInclusive(SampleLengthFromRoot, LFR0, LFR1)
-			|| FMath::IsWithinInclusive(SampleLengthFromRoot, LFR1, LFR0))
-		{
-			const float Denominator = (LFR1 - LFR0);
-			OutAlpha = (FMath::Abs(Denominator) <= KINDA_SMALL_NUMBER)
-				? 0.0f
-				: ((SampleLengthFromRoot - LFR0) / Denominator);
-			OutAlpha = FMath::Clamp(OutAlpha, 0.0f, 1.0f);
-			PreviousPointIndex = P;
-			NextPointIndex = N;
-
-			return true;
-		}
-	}
-
-	PreviousPointIndex = SortedPointIndices[0];
-	NextPointIndex = SortedPointIndices.Last();
-	OutAlpha = (SampleValue < 0.5f)
-		? 0.0f
-		: 1.0f;
-
-	return true;
-}
-
-void FPVGrafter::ComputeRotationAnglesForGraft(const FVector3f& ApicalDirection, const FVector3f& BudDirectionUp, const float AxilAngle,
-                                               const float TwistAroundApical, FVector3f& OutUpDirection, FVector3f& OutNormalDirection)
-{
-	const FVector3f ApicalDirectionNormalized = ApicalDirection.GetSafeNormal();
-	if (ApicalDirectionNormalized.IsNearlyZero())
-	{
-		OutUpDirection = FVector3f(1, 0, 0);
-		OutNormalDirection = FVector3f(0, 0, 1);
-		return;
-	}
-
-	FVector3f BudDirectionNormalized = BudDirectionUp.GetSafeNormal();
-	if (BudDirectionNormalized.IsNearlyZero())
-	{
-		FVector Axis1, Axis2;
-		FVector(ApicalDirectionNormalized).FindBestAxisVectors(Axis1, Axis2);
-		BudDirectionNormalized = FVector3f(Axis2);
-	}
-
-	FVector3f AxilAxis = FVector3f::CrossProduct(BudDirectionNormalized, ApicalDirectionNormalized).GetSafeNormal();
-	if (AxilAxis.IsNearlyZero())
-	{
-		FVector Axis1, Axis2;
-		FVector(ApicalDirectionNormalized).FindBestAxisVectors(Axis1, Axis2);
-		AxilAxis = FVector3f(Axis2);
-	}
-
-	const FQuat4f AxilQuat(AxilAxis, AxilAngle);
-	const FVector3f AxilAdjusted = AxilQuat.RotateVector(BudDirectionNormalized);
-
-	const FQuat4f TwistQuat(ApicalDirectionNormalized, TwistAroundApical);
-
-	OutUpDirection = TwistQuat.RotateVector(AxilAdjusted).GetSafeNormal();
-	OutNormalDirection = TwistQuat.RotateVector(AxilQuat.RotateVector(ApicalDirectionNormalized)).GetSafeNormal();
 }
 
 void FPVGrafter::DistributeGraftWithParametricSettings(FManagedArrayCollection& OutCollection,
-                                                       const FManagedArrayCollection& SourceSkeletonCollection,
                                                        TArray<FManagedArrayCollection>& SourceGraftCollections,
                                                        const FParametricSettings::FSpacingSettings& SpacingSettings,
                                                        const FParametricSettings::FAngleSettings& AngleSettings,
                                                        const FParametricSettings::FScaleSettings& ScaleSettings,
-                                                       const FMiscSettings& MiscSettings,
+                                                       const int32 RandomSeed,
                                                        const FPVFoliageDistributionConditions& DistributionConditions)
 {
 	TRACE_CPUPROFILER_EVENT_SCOPE(PVGrafter::DistributeGraftWithHormoneBasedSettings);
 
 	FRandomStream RandomStream;
-	RandomStream.Initialize(MiscSettings.RandomSeed);
-
-	PV::Facades::FBranchFacade BranchFacadeSourceSkeleton(SourceSkeletonCollection);
-	PV::Facades::FPointFacade PointFacadeSourceSkeleton(SourceSkeletonCollection);
-	PV::Facades::FPointFacade PointFacadeOut(OutCollection);
-
-	auto PointPlantGradientAttribute = PV::FPointPlantGradientAttribute::GetAttribute(OutCollection);
-	auto PointBranchGradientAttribute = PV::FPointBranchGradientAttribute::GetAttribute(OutCollection);
-	auto BranchParentNumberAttribute = PV::FBranchParentNumberAttribute::GetAttribute(OutCollection);
-	auto BranchChildrenAttribute = PV::FBranchChildrenAttribute::GetAttribute(OutCollection);
-	auto BranchNumberAttribute = PV::FBranchNumberAttribute::GetAttribute(OutCollection);
-	auto BranchPointsAttribute = PV::FBranchPointsAttribute::GetAttribute(OutCollection);
-	auto PointPositionAttribute = PV::FPointPositionAttribute::GetAttribute(OutCollection);
-	auto BudDirectionAttribute = PV::FBudDirectionAttribute::GetAttribute(OutCollection);
-	
-	const int32 NumOfBranches = BranchPointsAttribute.Num();
-	check(NumOfBranches > 0);
-	if (NumOfBranches == 0)
-	{
-		return;
-	}
-
-	// Only need to do this temporarily since legacy skeleton has inconsistent data in children
-	// TODO: Remove this once support for legacy skeletons is droppped
-	PV::Facades::FBranchFacade BranchFacadeOut(OutCollection);
-	BranchFacadeOut.RecomputeBranchChildren();
-	BranchFacadeOut.RecomputeBranchParents();
-
-	PV::AttributesHelper::ComputePointPlantAndBranchGradient(
-		{
-			PointPlantGradientAttribute,
-			PointBranchGradientAttribute,
-			BranchPointsAttribute,
-			BranchParentNumberAttribute,
-			BranchNumberAttribute,
-			BranchChildrenAttribute
-		});
+	RandomStream.Initialize(RandomSeed);
 
 	FNormalizedPointCaches NormalizedPointCaches;
 	FBuildOptions BuildOptions;
@@ -1399,7 +1148,7 @@ void FPVGrafter::DistributeGraftWithParametricSettings(FManagedArrayCollection&
 	BuildOptions.bComputeHealth = DistributionConditions.IsActiveCondition(EFoliageDistributionCondition::Health);
 
 	BuildNormalizedPointCaches(
-		SourceSkeletonCollection,
+		OutCollection,
 		BuildOptions,
 		NormalizedPointCaches);
 
@@ -1421,32 +1170,21 @@ void FPVGrafter::DistributeGraftWithParametricSettings(FManagedArrayCollection&
 		CandidateGraftData[i] = MetaInfoFacade.GetGraftAttributes();
 	}
 
-	const TManagedArray<int32>& SourceSkeletonBudNumbers = PointFacadeSourceSkeleton.GetBudNumbersAttribute();
-	check(SourceSkeletonBudNumbers.Num() > 0);
-	const int32 MaxBudNumber = *Algo::MaxElement(SourceSkeletonBudNumbers.GetConstArray());
+	TArray<FAttachmentPoint> AttachmentPoints;
+	ComputeAttachmentPoints(
+		OutCollection, 
+		NormalizedPointCaches, 
+		SpacingSettings, 
+		AngleSettings, 
+		ScaleSettings, 
+		RandomStream, 
+		AttachmentPoints);
+
+	auto BudNumberAttribute = PV::FPointBudNumberAttribute::GetAttribute(OutCollection);
+	check(BudNumberAttribute.Num() > 0);
+	const int32 MaxBudNumber = *Algo::MaxElement(BudNumberAttribute);
 	int32 CurrentBudNumber = MaxBudNumber + 1;
 
-	auto PointLengthFromRootAttribute = PV::FPointLengthFromRootAttribute::GetAttribute(OutCollection);
-
-	TMap<int32, float> BranchIndicesToBranchLengths = ComputeBranchLengths(PointLengthFromRootAttribute, BranchPointsAttribute);
-	float MaxBranchLength = -1.0f;
-	float AverageBranchLength = 0.0f;
-	for (const auto IndexLengthPair : BranchIndicesToBranchLengths)
-	{
-		const float BranchLength = IndexLengthPair.Value;
-		MaxBranchLength = FMath::Max(BranchLength, MaxBranchLength);
-		AverageBranchLength = AverageBranchLength + BranchLength;
-	}
-	
-	check(MaxBranchLength > KINDA_SMALL_NUMBER);
-	if (MaxBranchLength < KINDA_SMALL_NUMBER)
-	{
-		return;
-	}
-
-	AverageBranchLength = AverageBranchLength / BranchIndicesToBranchLengths.Num();
-	AverageBranchLength = FMath::Lerp(AverageBranchLength, MaxBranchLength, 0.33f);
-
 	auto IsValidCandidate = [&](int32 CandidateIndex) -> bool
 		{
 			return CandidateGraftData.IsValidIndex(CandidateIndex);
@@ -1476,259 +1214,55 @@ void FPVGrafter::DistributeGraftWithParametricSettings(FManagedArrayCollection&
 			}
 		};
 
-	for (int32 BranchIndex = 0; BranchIndex < NumOfBranches; BranchIndex++)
+	for (const auto& AttachmentPoint : AttachmentPoints)
 	{
-		// Check generation limit
-		const TArray<int32>& BranchParents = BranchFacadeOut.GetParents(BranchIndex);
-		if (SpacingSettings.LimitStartGeneration && BranchParents.Num() < SpacingSettings.StartGeneration)
+		const int32 PickedCandidateIndex = ConditionPicker::PickCandidateIndex(
+			GraftSourcesExtracted.Num(),
+			DistributionConditions,
+			AttachmentPoint.ConditionSample,
+			RandomStream,
+			GetCandidateValue,
+			IsValidCandidate);
+
+		if (PickedCandidateIndex == INDEX_NONE)
 		{
 			continue;
 		}
 
-		if (SpacingSettings.LimitEndGeneration && BranchParents.Num() > SpacingSettings.EndGeneration)
-		{
-			continue;
-		}
+		int32 AttachmentPointIndex = INDEX_NONE;
 
-		const float LengthRatio = BranchIndicesToBranchLengths[BranchIndex] / MaxBranchLength;
-		int32 LoopNumber = SpacingSettings.BranchDensity;
-		if (SpacingSettings.SpacingBasis == FParametricSettings::EDistributionBasis::Branch)
-		{
-			LoopNumber = static_cast<int32>(LoopNumber * LengthRatio);
-		}
-		else if (SpacingSettings.SpacingBasis == FParametricSettings::EDistributionBasis::Plant)
+		if (AttachmentPoint.bIsTipInstance)
 		{
-			LoopNumber = static_cast<int32>(LoopNumber * (AverageBranchLength / MaxBranchLength));
+			AttachmentPointIndex = AttachmentPoint.TipPointIndex;
 		}
-
-		LoopNumber = FMath::Max(LoopNumber, 1);
-
-		const TArray<int32>& BranchPoints = BranchPointsAttribute[BranchIndex];
-
-		if (BranchPoints.Num() < 2)
+		else
 		{
-			continue;
+			AttachmentPointIndex = PV::Utilities::AddInterpolatedPointToCollection(
+				OutCollection,
+				AttachmentPoint.PointIndexA,
+				AttachmentPoint.PointIndexB,
+				AttachmentPoint.PointAlpha,
+				CurrentBudNumber,
+				CurrentBudNumber);
+
+			CurrentBudNumber++;
 		}
 
-		TArray<int32> SortedPointIndices = BranchPoints;
-		SortedPointIndices.Sort([PointLengthFromRootAttribute](const int32& A, const int32& B)
-			{
-				return PointLengthFromRootAttribute[A] < PointLengthFromRootAttribute[B];
-			});
-
-		const int32 FirstPointIndex = SortedPointIndices[0];
-		const int32 TipPointIndex = SortedPointIndices.Last();
-		const float PlantGradientAtRoot = 1.0f - PointPlantGradientAttribute[FirstPointIndex];
-		const float LFRofRoot = PointLengthFromRootAttribute[FirstPointIndex];
-		const float LFRofTip = PointLengthFromRootAttribute[TipPointIndex];
-
-		for (int32 j = 0; j < LoopNumber; j++)
+		FPVGraftAttributes Attributes
 		{
-			float BaseValue = static_cast<float>(j) / FMath::Max(LoopNumber - 1, 1);
-			float PlantGradient;
-			float BranchGradient = 0.0f;
-			float AttachmentPointLengthFromRoot = 0.0f;
-			FVector3f Pos(0.0f, 0.0f, 0.0f);
-			int32 PointIndexA = INDEX_NONE;
-			int32 PointIndexB = INDEX_NONE;
-			float PointAlpha = 0.0f;
-
-			if (SpacingSettings.SpacingBasis == FParametricSettings::EDistributionBasis::Plant)
-			{
-				PlantGradient = BaseValue;
-				PlantGradient = SpacingSettings.SpacingRamp->GetRichCurveConst()->Eval(PlantGradient);
-				PlantGradient = FMath::GetMappedRangeValueClamped(
-					FVector2f(0.0f, 1.0f),
-					FVector2f(SpacingSettings.RelativeStart, SpacingSettings.RelativeEnd),
-					PlantGradient);
-				if (PlantGradient < PlantGradientAtRoot)
-				{
-					continue;
-				}
-
-				Pos = SampleVectorByFloat(PlantGradient, SortedPointIndices, PointPlantGradientAttribute, PointPositionAttribute);
-				AttachmentPointLengthFromRoot = SampleFloatByFloat(PlantGradient, SortedPointIndices, PointPlantGradientAttribute,
-					PointLengthFromRootAttribute);
-				BranchGradient = FMath::GetMappedRangeValueClamped(
-					FVector2f(LFRofRoot, LFRofTip),
-					FVector2f(0.0f, 1.0f),
-					AttachmentPointLengthFromRoot);
-
-				FindPointByNormalizedLengthFromRoot(BranchGradient, SortedPointIndices, PointLengthFromRootAttribute, LFRofRoot, LFRofTip,
-					PointIndexA,
-					PointIndexB, PointAlpha);
-			}
-			else // SpacingBasis is Branch
-			{
-				BranchGradient = SpacingSettings.SpacingRamp->GetRichCurveConst()->Eval(BaseValue);
-				BranchGradient = FMath::GetMappedRangeValueClamped(
-					FVector2f(0.0f, 1.0f),
-					FVector2f(SpacingSettings.RelativeStart, SpacingSettings.RelativeEnd),
-					BranchGradient);
-
-				if (SpacingSettings.BranchDensity == 1)
-				{
-					BranchGradient = SpacingSettings.RelativeEnd;
-					BaseValue = 1.0f;
-				}
-
-				if (!FindPointByNormalizedLengthFromRoot(BranchGradient, SortedPointIndices, PointLengthFromRootAttribute, LFRofRoot, LFRofTip,
-					PointIndexA, PointIndexB, PointAlpha))
-				{
-					continue;
-				}
-
-				const FVector3f P0 = PointPositionAttribute[PointIndexA];
-				const FVector3f P1 = PointPositionAttribute[PointIndexB];
-				Pos = FMath::Lerp(P0, P1, PointAlpha);
-
-				AttachmentPointLengthFromRoot = FMath::GetMappedRangeValueClamped(
-					FVector2f(0.0f, 1.0f),
-					FVector2f(LFRofRoot, LFRofTip),
-					BranchGradient);
-			}
-
-			const float PlantGradient0 = 1.0f - PointPlantGradientAttribute[PointIndexA];
-			const float PlantGradient1 = 1.0f - PointPlantGradientAttribute[PointIndexB];
-			PlantGradient = FMath::Lerp(PlantGradient0, PlantGradient1, PointAlpha);
-
-			const float ScaleRampLookUp =
-				(ScaleSettings.ScaleRampBasis == FParametricSettings::EDistributionBasis::Branch)
-				? BranchGradient
-				: PlantGradient;
-
-			float ScaleRampValue = ScaleSettings.ScaleRamp->GetRichCurveConst()->Eval(ScaleRampLookUp);
-
-			if (ScaleSettings.ScaleRampBasis == FParametricSettings::EDistributionBasis::Branch)
-			{
-				ScaleRampValue = FMath::GetMappedRangeValueClamped(
-					FVector2f(0.0f, 1.0f),
-					FVector2f(0.0f, LengthRatio),
-					ScaleRampValue);
-			}
-
-			const float RandomScaleFactor = RandomStream.FRandRange(
-				ScaleSettings.RandomizeScaleAngleMinimum,
-				ScaleSettings.RandomizeScaleAngleMaximum);
-
-			const float Scale = ScaleRampValue * RandomScaleFactor * ScaleSettings.BaseScale;
-
-			const FVector3f BudDirectionApicalA = BudDirectionAttribute[PointIndexA].Apical;
-			const FVector3f BudDirectionApicalB = BudDirectionAttribute[PointIndexB].Apical;
-			const FVector3f BudDirectionUpA = BudDirectionAttribute[PointIndexA].UpVector;
-			const FVector3f BudDirectionUpB = BudDirectionAttribute[PointIndexB].UpVector;
-
-			const FVector3f ApicalDirection = FMath::Lerp(BudDirectionApicalA, BudDirectionApicalB, PointAlpha);
-			const FVector3f BudDirectionUp = FMath::Lerp(BudDirectionUpA, BudDirectionUpB, PointAlpha);
-
-			const float AxilRampValue = AngleSettings.AxilAngleRamp->GetRichCurveConst()->Eval(BaseValue);
-			float AxilAngleInDegrees = AngleSettings.AxilAngle * AxilRampValue;
-
-			AxilAngleInDegrees += RandomStream.FRandRange(AngleSettings.RandomizeAxilAngleMinimum, AngleSettings.RandomizeAxilAngleMaximum);
-			AxilAngleInDegrees = FMath::Clamp(AxilAngleInDegrees, -90.0f, 90.0f);
-
-			const float Twist = FMath::DegreesToRadians(AngleSettings.Rotation) * static_cast<float>(j);
-
-			FVector3f UpDirection, NormalDirection;
-			ComputeRotationAnglesForGraft(
-				ApicalDirection,
-				BudDirectionUp,
-				FMath::DegreesToRadians(AxilAngleInDegrees),
-				Twist,
-				UpDirection,
-				NormalDirection);
-
-			const bool bIsTipInstance =
-				FMath::IsNearlyEqual(AttachmentPointLengthFromRoot, LFRofTip, 0.001f) &&
-				(BranchGradient >= 0.999f || FMath::IsNearlyEqual(BranchGradient, SpacingSettings.RelativeEnd, 0.001f));
-
-			int32 AttachmentPointIndex = INDEX_NONE;
-			if (bIsTipInstance)
-			{
-				AttachmentPointIndex = TipPointIndex;
-
-				const int32 SecondToLastPointIndex = SortedPointIndices[SortedPointIndices.Num() - 2];
-				const FVector3f TipPosition = PointPositionAttribute[TipPointIndex];
-				const FVector3f SecondToLastPointPosition = PointPositionAttribute[SecondToLastPointIndex];
-				FVector3f TipTangent = (TipPosition - SecondToLastPointPosition).GetSafeNormal();
-				if (TipTangent.IsNearlyZero())
-				{
-					TipTangent = ApicalDirection.GetSafeNormal();
-				}
-
-				const FVector3f OldUp = UpDirection.GetSafeNormal();
-				FQuat4f AlignQuat = FQuat4f::FindBetweenNormals(OldUp, TipTangent);
-
-				UpDirection = TipTangent;
-				NormalDirection = AlignQuat.RotateVector(NormalDirection).GetSafeNormal();
-
-				if (FMath::Abs(FVector3f::DotProduct(UpDirection, NormalDirection)) > 0.999f)
-				{
-					FVector Axis1, Axis2;
-					FVector(UpDirection).FindBestAxisVectors(Axis1, Axis2);
-					NormalDirection = FVector3f(Axis1);
-				}
-			}
-			else
-			{
-				AttachmentPointIndex = PV::Utilities::AddInterpolatedPointToCollection(
-					OutCollection,
-					PointIndexA,
-					PointIndexB,
-					PointAlpha,
-					CurrentBudNumber,
-					CurrentBudNumber);
-
-				CurrentBudNumber++;
-			}
-
-			FPVGraftAttributes Attributes
-			{
-				.BranchIndex = BranchIndex,
-				.PointIndex = AttachmentPointIndex,
-				.AttachmentPoint = Pos,
-				.UpVector = UpDirection,
-				.NormalVector = NormalDirection,
-				.Scale = Scale,
-				.LengthFromRoot = AttachmentPointLengthFromRoot,
-				.bIsTipInstance = bIsTipInstance
-			};
-
-			ConditionPicker::FPointConditionSample Sample;
-			Sample.UpAlignment = UpAlignmentNormalized(ApicalDirection);
-			Sample.Tip = Tip(bIsTipInstance);
-
-			if (PointIndexA != INDEX_NONE && PointIndexB != INDEX_NONE)
-			{
-				Sample.Light = FMath::Lerp(NormalizedPointCaches.Light[PointIndexA], NormalizedPointCaches.Light[PointIndexB], PointAlpha);
-				Sample.Health = FMath::Lerp(NormalizedPointCaches.Health[PointIndexA], NormalizedPointCaches.Health[PointIndexB], PointAlpha);
-				Sample.Scale = FMath::Lerp(NormalizedPointCaches.PScales[PointIndexA], NormalizedPointCaches.PScales[PointIndexB], PointAlpha);
-			}
-			else
-			{
-				Sample.Light = 0.0f;
-				Sample.Health = 0.0f;
-				Sample.Scale = 0.0f;
-			}
-
-			check(GraftSourcesExtracted.Num() > 0);
-
-			const int32 PickedCandidateIndex = ConditionPicker::PickCandidateIndex(
-				GraftSourcesExtracted.Num(),
-				DistributionConditions,
-				Sample,
-				RandomStream,
-				GetCandidateValue,
-				IsValidCandidate);
-
-			if (PickedCandidateIndex == INDEX_NONE)
-			{
-				continue;
-			}
+			.BranchIndex = AttachmentPoint.BranchIndex,
+			.PointIndex = AttachmentPointIndex,
+			.AttachmentPoint = AttachmentPoint.Position,
+			.UpVector = AttachmentPoint.UpDirection,
+			.NormalVector = AttachmentPoint.NormalDirection,
+			.Scale = AttachmentPoint.Scale,
+			.LengthFromRoot = AttachmentPoint.LengthFromRoot,
+			.bIsTipInstance = AttachmentPoint.bIsTipInstance
+		};
 
-			const FManagedArrayCollection& PickedSourceGraft = GraftSourcesExtracted[PickedCandidateIndex];
+		check(GraftSourcesExtracted.Num() > 0);
+		const FManagedArrayCollection& PickedSourceGraft = GraftSourcesExtracted[PickedCandidateIndex];
 
-			SetGraft(OutCollection, SourceSkeletonCollection, PickedSourceGraft, Attributes);
-		}
+		SetGraft(OutCollection, PickedSourceGraft, Attributes);
 	}
 }
diff --git a/Engine/Plugins/Experimental/ProceduralVegetationEditor/Source/ProceduralVegetation/Private/Implementations/PVGrafter.h b/Engine/Plugins/Experimental/ProceduralVegetationEditor/Source/ProceduralVegetation/Private/Implementations/PVGrafter.h
--- a/Engine/Plugins/Experimental/ProceduralVegetationEditor/Source/ProceduralVegetation/Private/Implementations/PVGrafter.h
+++ b/Engine/Plugins/Experimental/ProceduralVegetationEditor/Source/ProceduralVegetation/Private/Implementations/PVGrafter.h
@@ -5,8 +5,11 @@
 #include "CoreMinimal.h"
 #include "Utils/PVAttributes.h"
 #include "GeometryCollection/ManagedArrayCollection.h"
+#include "Helpers/PVParametricDistributionHelper.h"
 #include "Helpers/PVPhyllotaxyHelper.h"
 
+using namespace PV::ParametricDistributionHelper;
+
 struct FPVFoliageDistributionConditions;
 struct FPVFloatRamp;
 
@@ -22,7 +25,7 @@ struct FPVGrafter
 			const float InstanceSpacingRampEffect;
 			const int32 MaxPerBranch;
 		};
-		
+
 		struct FPhyllotaxySettings
 		{
 			const bool ResetPhyllotaxy;
@@ -55,52 +58,6 @@ struct FPVGrafter
 		};
 	};
 
-	struct FParametricSettings
-	{
-		enum class EDistributionBasis : uint8
-		{
-			Plant,
-			Branch
-		};
-
-		struct FSpacingSettings
-		{
-			const int32 BranchDensity;
-			const float RelativeStart;
-			const float RelativeEnd;
-			const bool LimitStartGeneration;
-			const int32 StartGeneration;
-			const bool LimitEndGeneration;
-			const int32 EndGeneration;
-			const EDistributionBasis SpacingBasis;
-			const FPVFloatRamp* SpacingRamp;
-		};
-
-		struct FAngleSettings
-		{
-			const float Rotation;
-			const float AxilAngle;
-			const float RandomizeAxilAngleMinimum;
-			const float RandomizeAxilAngleMaximum;
-			const EDistributionBasis AxilAngleRampBasis;
-			const FPVFloatRamp* AxilAngleRamp;
-		};
-
-		struct FScaleSettings
-		{
-			const EDistributionBasis ScaleRampBasis;
-			const float BaseScale;
-			const float RandomizeScaleAngleMinimum;
-			const float RandomizeScaleAngleMaximum;
-			const FPVFloatRamp* ScaleRamp;
-		};
-	};
-
-	struct FMiscSettings
-	{
-		const int32 RandomSeed;
-	};
-
 	struct FPVGraftAttributes
 	{
 		int32 BranchIndex;
@@ -129,63 +86,29 @@ struct FPVGrafter
 	template <typename Fn>
 	static void ForEachGraftPointExcludingRoot(const int32 NumPoints, const int32 RootIndex, Fn&& Func);
 
-	static void SetGraft(FManagedArrayCollection& OutCollection, const FManagedArrayCollection& SourceSkeletonCollection,
+	static void SetGraft(FManagedArrayCollection& OutCollection,
 	                     const FManagedArrayCollection& SourceGraftCollection, const FPVGraftAttributes& InAttributes);
 
 	static bool RemoveFoliageDataIfPresent(FManagedArrayCollection& OutCollection);
-	
-	static TArray<FManagedArrayCollection> ExtractIndividualPlants(const FManagedArrayCollection& InCollection);
 
-	static void DistributeGraftWithHormoneBasedSettings(FManagedArrayCollection& OutCollection,
-	                                                    const FManagedArrayCollection& SourceSkeletonCollection,
-	                                                    TArray<FManagedArrayCollection>& SourceGraftCollections,
-	                                                    const FHormoneBasedSettings::FDistributionSettings& DistributionSettings,
-	                                                    const FHormoneBasedSettings::FPhyllotaxySettings& PhyllotaxySettings,
-	                                                    const FHormoneBasedSettings::FScaleSettings& ScaleSettings,
-	                                                    const FHormoneBasedSettings::FVectorSettings& VectorSettings,
-	                                                    const FMiscSettings& MiscSettings,
-	                                                    const FPVFoliageDistributionConditions& DistributionConditions);
-
-	static TMap<int32, float> ComputeBranchLengths(
-		PV::FPointLengthFromRootAttributeConstView PointLengthFromRootAttribute,
-		PV::FBranchPointsAttributeConstView BranchPointsAttribute);
+	static TArray<FManagedArrayCollection> ExtractIndividualPlants(const FManagedArrayCollection& InCollection);
 
-	static FVector3f SampleVectorByFloat(
-		const float SampleValue,
-		const TArray<int32>& SortedPointIndices,
-		PV::FPointPlantGradientAttributeConstView PointPlantGradientAttribute,
-		PV::FPointPositionAttributeConstView PointPositionAttribute);
-
-	static float SampleFloatByFloat(
-		const float SampleValue,
-		const TArray<int32>& SortedPointIndices,
-		PV::FPointPlantGradientAttributeConstView PointPlantGradientAttribute,
-		PV::FPointLengthFromRootAttributeConstView PointLengthFromRootAttribute);
-
-	static bool FindPointByNormalizedLengthFromRoot(
-		float SampleValue,
-		const TArray<int32>& SortedPointIndices,
-		PV::FPointLengthFromRootAttributeConstView PointLengthFromRootAttribute,
-		const float LFRofRoot,
-		const float LFRofTip,
-		int32& PreviousPointIndex,
-		int32& NextPointIndex,
-		float& OutAlpha);
-
-	static void ComputeRotationAnglesForGraft(
-		const FVector3f& ApicalDirection,
-		const FVector3f& BudDirectionUp,
-		const float AxilAngle,
-		const float TwistAroundApical,
-		FVector3f& OutUpDirection,
-		FVector3f& OutNormalDirection);
-
-	static void DistributeGraftWithParametricSettings(FManagedArrayCollection& OutCollection,
-	                                                  const FManagedArrayCollection& SourceSkeletonCollection,
-	                                                  TArray<FManagedArrayCollection>& SourceGraftCollections,
-	                                                  const FParametricSettings::FSpacingSettings& SpacingSettings,
-	                                                  const FParametricSettings::FAngleSettings& AngleSettings,
-	                                                  const FParametricSettings::FScaleSettings& ScaleSettings,
-	                                                  const FMiscSettings& MiscSettings,
-	                                                  const FPVFoliageDistributionConditions& DistributionConditions);
+	static void DistributeGraftWithHormoneBasedSettings(
+		FManagedArrayCollection& OutCollection,
+		TArray<FManagedArrayCollection>& SourceGraftCollections,
+		const FHormoneBasedSettings::FDistributionSettings& DistributionSettings,
+		const FHormoneBasedSettings::FPhyllotaxySettings& PhyllotaxySettings,
+		const FHormoneBasedSettings::FScaleSettings& ScaleSettings,
+		const FHormoneBasedSettings::FVectorSettings& VectorSettings,
+		int32 RandomSeed,
+		const FPVFoliageDistributionConditions& DistributionConditions);
+
+	static void DistributeGraftWithParametricSettings(
+		FManagedArrayCollection& OutCollection,
+		TArray<FManagedArrayCollection>& SourceGraftCollections,
+		const FParametricSettings::FSpacingSettings& SpacingSettings,
+		const FParametricSettings::FAngleSettings& AngleSettings,
+		const FParametricSettings::FScaleSettings& ScaleSettings,
+		int32 RandomSeed,
+		const FPVFoliageDistributionConditions& DistributionConditions);
 };
diff --git a/Engine/Plugins/Experimental/ProceduralVegetationEditor/Source/ProceduralVegetation/Private/Nodes/PVFoliageDistributorSettings.cpp b/Engine/Plugins/Experimental/ProceduralVegetationEditor/Source/ProceduralVegetation/Private/Nodes/PVFoliageDistributorSettings.cpp
--- a/Engine/Plugins/Experimental/ProceduralVegetationEditor/Source/ProceduralVegetation/Private/Nodes/PVFoliageDistributorSettings.cpp
+++ b/Engine/Plugins/Experimental/ProceduralVegetationEditor/Source/ProceduralVegetation/Private/Nodes/PVFoliageDistributorSettings.cpp
@@ -37,26 +37,47 @@ UPVFoliageDistributorSettings::UPVFoliageDistributorSettings()
 {
 	if (PCGHelpers::IsNewObjectAndNotDefault(this))
 	{
-		FRichCurve* RichCurve = InstanceSpacingRamp.GetRichCurve();
+		FRichCurve* RichCurve = HormoneBasedSettings.InstanceSpacingRamp.GetRichCurve();
 		RichCurve->Reset();
 		const FKeyHandle Handle0 = RichCurve->AddKey(0.0f, 0.0f);
 		const FKeyHandle Handle1 = RichCurve->AddKey(1.0f, 1.0f);
 		RichCurve->SetKeyInterpMode(Handle0, RCIM_Linear);
 		RichCurve->SetKeyInterpMode(Handle1, RCIM_Linear);
  
-		RichCurve = ScaleRamp.GetRichCurve();
+		RichCurve = HormoneBasedSettings.ScaleRamp.GetRichCurve();
 		RichCurve->Reset();
 		const FKeyHandle Handle2 = RichCurve->AddKey(0.0f, 0.0f);
 		const FKeyHandle Handle3 = RichCurve->AddKey(1.0f, 1.0f);
 		RichCurve->SetKeyInterpMode(Handle2, RCIM_Linear);
 		RichCurve->SetKeyInterpMode(Handle3, RCIM_Linear);
  
-		RichCurve = AxilAngleRamp.GetRichCurve();
+		RichCurve = HormoneBasedSettings.AxilAngleRamp.GetRichCurve();
 		RichCurve->Reset();
 		const FKeyHandle Handle4 = RichCurve->AddKey(0.0f, 1.0f);
 		const FKeyHandle Handle5 = RichCurve->AddKey(1.0f, 0.5f);
 		RichCurve->SetKeyInterpMode(Handle4, RCIM_Linear);
 		RichCurve->SetKeyInterpMode(Handle5, RCIM_Linear);
+		
+		RichCurve = ParametricSettings.SpacingRamp.GetRichCurve();
+		RichCurve->Reset();
+		const FKeyHandle Handle6 = RichCurve->AddKey(0.0f, 0.0f);
+		const FKeyHandle Handle7 = RichCurve->AddKey(1.0f, 1.0f);
+		RichCurve->SetKeyInterpMode(Handle6, RCIM_Linear);
+		RichCurve->SetKeyInterpMode(Handle7, RCIM_Linear);
+
+		RichCurve = ParametricSettings.AxilAngleRamp.GetRichCurve();
+		RichCurve->Reset();
+		const FKeyHandle Handle8 = RichCurve->AddKey(0.0f, 0.0f);
+		const FKeyHandle Handle9 = RichCurve->AddKey(1.0f, 1.0f);
+		RichCurve->SetKeyInterpMode(Handle8, RCIM_Linear);
+		RichCurve->SetKeyInterpMode(Handle9, RCIM_Linear);
+
+		RichCurve = ParametricSettings.ScaleRamp.GetRichCurve();
+		RichCurve->Reset();
+		const FKeyHandle Handle10 = RichCurve->AddKey(0.0f, 1.0f);
+		const FKeyHandle Handle11 = RichCurve->AddKey(1.0f, 0.25f);
+		RichCurve->SetKeyInterpMode(Handle10, RCIM_Linear);
+		RichCurve->SetKeyInterpMode(Handle11, RCIM_Linear);
 	}
 }
  
@@ -112,49 +133,101 @@ bool FPVFoliageDistributorElement::ExecuteInternal(FPCGContext* InContext) const
 			FManagedArrayCollection Collection = InputData->GetCollection();
 			
 			UPVFoliageMeshData* OutManagedArrayCollectionData = FPCGContext::NewObject_AnyThread<UPVFoliageMeshData>(InContext);
- 
-			const FDistributionSettings DistributionSettings = {
-				.EthyleneThreshold = Settings->EthyleneThreshold,
-				.OverrideDistribution = true,
-				.InstanceSpacing = Settings->InstanceSpacing,
-				.InstanceSpacingRamp = &Settings->InstanceSpacingRamp,
-				.InstanceSpacingRampEffect = Settings->InstanceSpacingRampEffect,
-				.MaxPerBranch = Settings->MaxPerBranch
-			};
- 
-			const FScaleSettings ScaleSettings = {
-				.BaseScale = Settings->BaseScale,
-				.BranchScaleImpact = Settings->BranchScaleImpact,
-				.MinScale = Settings->MinScale,
-				.MaxScale = Settings->MaxScale,
-				.RandomScaleMin = Settings->RandomScaleMin,
-				.RandomScaleMax = Settings->RandomScaleMax,
-				.ScaleRamp = &Settings->ScaleRamp
-			};
- 
-			const FVectorSettings VectorSettings = {
-				.OverrideAxilAngle = Settings->OverrideAxilAngle,
-				.AxilAngle = Settings->AxilAngle,
-				.AxilAngleRamp = &Settings->AxilAngleRamp,
-				.AxilAngleRampUpperValue = Settings->AxilAngleRampUpperValue,
-				.AxilAngleRampEffect = Settings->AxilAngleRampEffect
-			};
- 
-			const FPhyllotaxySettings PhyllotaxySettings = {
-				.OverridePhyllotaxy = Settings->OverridePhyllotaxy,
-				.PhyllotaxyType = Settings->PhyllotaxyType,
-				.PhyllotaxyFormation = Settings->PhyllotaxyFormation,
-				.MinimumNodeBuds = Settings->MinimumNodeBuds,
-				.MaximumNodeBuds = Settings->MaximumNodeBuds,
-				.PhyllotaxyAdditionalAngle = Settings->PhyllotaxyAdditionalAngle
-			};
- 
-			const FMiscSettings MiscSettings = {
-				.RandomSeed = Settings->RandomSeed
-			};
- 
-			FPVFoliage::DistributeFoliage(Collection, Collection, FoliageCollection, DistributionSettings, ScaleSettings, VectorSettings,
-												 PhyllotaxySettings, MiscSettings, Settings->Conditions);
+			
+			if (Settings->Mode == EPVFoliageDistributionSettingsMode::HormoneBasedSettings)
+			{
+				const PV::HormoneBasedSettings::FDistributionSettings DistributionSettings = {
+					.EthyleneThreshold = Settings->HormoneBasedSettings.EthyleneThreshold,
+					.OverrideDistribution = true,
+					.InstanceSpacing = Settings->HormoneBasedSettings.InstanceSpacing,
+					.InstanceSpacingRamp = &Settings->HormoneBasedSettings.InstanceSpacingRamp,
+					.InstanceSpacingRampEffect = Settings->HormoneBasedSettings.InstanceSpacingRampEffect,
+					.MaxPerBranch = Settings->HormoneBasedSettings.MaxPerBranch
+				};
+ 
+				const PV::HormoneBasedSettings::FScaleSettings ScaleSettings = {
+					.BaseScale = Settings->HormoneBasedSettings.BaseScale,
+					.BranchScaleImpact = Settings->HormoneBasedSettings.BranchScaleImpact,
+					.MinScale = Settings->HormoneBasedSettings.MinScale,
+					.MaxScale = Settings->HormoneBasedSettings.MaxScale,
+					.RandomScaleMin = Settings->HormoneBasedSettings.RandomScaleMin,
+					.RandomScaleMax = Settings->HormoneBasedSettings.RandomScaleMax,
+					.ScaleRamp = &Settings->HormoneBasedSettings.ScaleRamp
+				};
+ 
+				const PV::HormoneBasedSettings::FVectorSettings VectorSettings = {
+					.OverrideAxilAngle = Settings->HormoneBasedSettings.OverrideAxilAngle,
+					.AxilAngle = Settings->HormoneBasedSettings.AxilAngle,
+					.AxilAngleRamp = &Settings->HormoneBasedSettings.AxilAngleRamp,
+					.AxilAngleRampUpperValue = Settings->HormoneBasedSettings.AxilAngleRampUpperValue,
+					.AxilAngleRampEffect = Settings->HormoneBasedSettings.AxilAngleRampEffect
+				};
+ 
+				const PV::HormoneBasedSettings::FPhyllotaxySettings PhyllotaxySettings = {
+					.OverridePhyllotaxy = Settings->HormoneBasedSettings.OverridePhyllotaxy,
+					.PhyllotaxyType = Settings->HormoneBasedSettings.PhyllotaxyType,
+					.PhyllotaxyFormation = Settings->HormoneBasedSettings.PhyllotaxyFormation,
+					.MinimumNodeBuds = Settings->HormoneBasedSettings.MinimumNodeBuds,
+					.MaximumNodeBuds = Settings->HormoneBasedSettings.MaximumNodeBuds,
+					.PhyllotaxyAdditionalAngle = Settings->HormoneBasedSettings.PhyllotaxyAdditionalAngle
+				};
+
+				FPVFoliage::DistributeFoliageWithHormoneBasedSettings(
+					Collection, 
+					FoliageCollection, 
+					DistributionSettings, 
+					ScaleSettings, 
+					VectorSettings,
+					PhyllotaxySettings,
+					Settings->RandomSeed, 
+					Settings->Conditions);
+			}
+			else
+			{
+				const FParametricSettings::FSpacingSettings SpacingSettings = {
+					.BranchDensity = Settings->ParametricSettings.BranchDensity,
+					.RelativeStart = Settings->ParametricSettings.RelativeStart,
+					.RelativeEnd = Settings->ParametricSettings.RelativeEnd,
+					.LimitStartGeneration = Settings->ParametricSettings.LimitStartGeneration,
+					.StartGeneration = Settings->ParametricSettings.StartGeneration,
+					.LimitEndGeneration = Settings->ParametricSettings.LimitEndGeneration,
+					.EndGeneration = Settings->ParametricSettings.EndGeneration,
+					.SpacingBasis = Settings->ParametricSettings.SpacingBasis == EPVFoliageDistributionBasis::Plant
+					? FParametricSettings::EDistributionBasis::Plant
+					: FParametricSettings::EDistributionBasis::Branch,
+					.SpacingRamp = &Settings->ParametricSettings.SpacingRamp
+				};
+
+				const FParametricSettings::FAngleSettings AngleSettings = {
+					.Rotation = Settings->ParametricSettings.Rotation,
+					.AxilAngle = Settings->ParametricSettings.AxilAngle,
+					.RandomizeAxilAngleMinimum = Settings->ParametricSettings.RandomizeAxilAngleMinimum,
+					.RandomizeAxilAngleMaximum = Settings->ParametricSettings.RandomizeAxilAngleMaximum,
+					.AxilAngleRampBasis = Settings->ParametricSettings.AxilAngleRampBasis == EPVFoliageDistributionBasis::Plant
+					? FParametricSettings::EDistributionBasis::Plant
+					: FParametricSettings::EDistributionBasis::Branch,
+					.AxilAngleRamp = &Settings->ParametricSettings.AxilAngleRamp
+				};
+
+				const FParametricSettings::FScaleSettings ScaleSettings = {
+					.ScaleRampBasis = Settings->ParametricSettings.ScaleRampBasis == EPVFoliageDistributionBasis::Plant
+					? FParametricSettings::EDistributionBasis::Plant
+					: FParametricSettings::EDistributionBasis::Branch,
+					.BaseScale = Settings->ParametricSettings.BaseScale,
+					.RandomizeScaleAngleMinimum = Settings->ParametricSettings.RandomizeScaleAngleMinimum,
+					.RandomizeScaleAngleMaximum = Settings->ParametricSettings.RandomizeScaleAngleMaximum,
+					.ScaleRamp = &Settings->ParametricSettings.ScaleRamp
+				};
+				
+				FPVFoliage::DistributeFoliageWithParametricSettings(
+					Collection,
+					FoliageCollection,
+					SpacingSettings,
+					AngleSettings,
+					ScaleSettings,
+					Settings->RandomSeed,
+					Settings->Conditions);
+			}
  
 			OutManagedArrayCollectionData->Initialize(MoveTemp(Collection));
 			InContext->OutputData.TaggedData.Emplace(OutManagedArrayCollectionData);
diff --git a/Engine/Plugins/Experimental/ProceduralVegetationEditor/Source/ProceduralVegetation/Private/Nodes/PVFoliageDistributorSettings.h b/Engine/Plugins/Experimental/ProceduralVegetationEditor/Source/ProceduralVegetation/Private/Nodes/PVFoliageDistributorSettings.h
--- a/Engine/Plugins/Experimental/ProceduralVegetationEditor/Source/ProceduralVegetation/Private/Nodes/PVFoliageDistributorSettings.h
+++ b/Engine/Plugins/Experimental/ProceduralVegetationEditor/Source/ProceduralVegetation/Private/Nodes/PVFoliageDistributorSettings.h
@@ -13,36 +13,26 @@
 #include "Nodes/PVBaseSettings.h"
  
 #include "PVFoliageDistributorSettings.generated.h"
- 
-UCLASS(BlueprintType, ClassGroup = (Procedural))
-class UPVFoliageDistributorSettings : public UPVBaseSettings
+
+UENUM()
+enum class EPVFoliageDistributionSettingsMode : uint8
+{
+	ParametricSettings UMETA(DisplayName="Use Parametric Settings"),
+	HormoneBasedSettings UMETA(DisplayName="Use Hormone Based Settings")
+};
+
+UENUM()
+enum class EPVFoliageDistributionBasis : uint8
+{
+	Plant UMETA(DisplayName="Plant"),
+	Branch UMETA(DisplayName="Branch")
+};
+
+USTRUCT()
+struct FPVFoliageDistributionHormoneBasedSettings
 {
 	GENERATED_BODY()
- 
-	friend class FPVFoliageDistributorElement;
 	
-public:
-	UPVFoliageDistributorSettings();
- 
-	#if WITH_EDITOR
-	virtual FName GetDefaultNodeName() const override { return FName(TEXT("ProceduralVegetationFoliageDistributor")); }
-	virtual FText GetDefaultNodeTitle() const override;
-	virtual FText GetNodeTooltipText() const override;
-	virtual FLinearColor GetNodeTitleColor() const override { return FLinearColor::Blue; }
- 
-	virtual TArray<EPVRenderType> GetDefaultRenderType() const override { return TArray{ EPVRenderType::Mesh, EPVRenderType::Foliage }; }
-	virtual TArray<EPVRenderType> GetSupportedRenderTypes() const override { return TArray{ EPVRenderType::Mesh, EPVRenderType::Foliage, EPVRenderType::FoliageAttachments }; }
-	#endif
- 
-protected:
-	virtual TArray<FPCGPinProperties> InputPinProperties() const override;
-	virtual FPCGDataTypeIdentifier GetInputPinTypeIdentifier() const override;
-	virtual FPCGDataTypeIdentifier GetOutputPinTypeIdentifier() const override;
- 
-	virtual FPCGElementPtr CreateElement() const override;
- 
-private:
- 
 	UPROPERTY(EditAnywhere, Category="Distribution Settings", meta=(PCG_Overridable, ClampMin=0.0f, ClampMax=1.0f, UIMin=0.0f, UIMax=1.0f, Tooltip="Threshold controlling foliage release based on an ethylene-like signal.\n\nHigher Threshold → more buds/branches retained, making the plant denser and bushier.\nLower Threshold → fewer buds/branches retained, resulting in a sparser, pruned structure."))
 	float EthyleneThreshold = 0.85f;
  
@@ -111,6 +101,118 @@ private:
  
 	UPROPERTY(EditAnywhere, Category="Phyllotaxy Settings", meta=(PCG_Overridable, EditCondition="OverridePhyllotaxy", ClampMin=0.0f, ClampMax=360.0f, UIMin=0.0f, UIMax=360.0f, Tooltip="Extra divergence angle (°) added to the phyllotaxy pattern\n\nAdds an angular offset to the base phyllotactic divergence."))
 	float PhyllotaxyAdditionalAngle = 0.0f;
+};
+
+USTRUCT()
+struct FPVFoliageDistributionParametricSettings
+{
+	GENERATED_BODY()
+
+	UPROPERTY(EditAnywhere, Category="Spacing Settings", meta=(PCG_Overridable, ClampMin=0, ClampMax=100))
+	int32 BranchDensity = 45;
+
+	UPROPERTY(EditAnywhere, Category="Spacing Settings", meta=(PCG_Overridable, ClampMin=0.0f, ClampMax=1.0f))
+	float RelativeStart = 0.0f;
+
+	UPROPERTY(EditAnywhere, Category="Spacing Settings", meta=(PCG_Overridable, ClampMin=0.0f, ClampMax=1.0f))
+	float RelativeEnd = 1.0f;
+
+	UPROPERTY(EditAnywhere, Category="Spacing Settings", meta=(PCG_Overridable))
+	bool LimitStartGeneration = false;
+
+	UPROPERTY(EditAnywhere, Category="Spacing Settings",
+		meta=(PCG_Overridable, ClampMin=0, ClampMax=10, EditCondition="LimitStartGeneration", EditConditionHides))
+	int32 StartGeneration = 0;
+
+	UPROPERTY(EditAnywhere, Category="Spacing Settings", meta=(PCG_Overridable))
+	bool LimitEndGeneration = false;
+
+	UPROPERTY(EditAnywhere, Category="Spacing Settings",
+		meta=(PCG_Overridable, ClampMin=0, ClampMax=10, EditCondition="LimitEndGeneration", EditConditionHides))
+	int32 EndGeneration = 0;
+
+	UPROPERTY(EditAnywhere, Category="Spacing Settings", meta=(PCG_Overridable))
+	EPVFoliageDistributionBasis SpacingBasis = EPVFoliageDistributionBasis::Plant;
+
+	UPROPERTY(EditAnywhere, Category="Spacing Settings",
+		meta=(PCG_Overridable, XAxisMin=0.0f, XAxisMax=1.0f, YAxisMin=0.0f, YAxisMax=1.0f))
+	FPVFloatRamp SpacingRamp;
+
+	UPROPERTY(EditAnywhere, Category="Angle Settings", meta=(PCG_Overridable, ClampMin=0.0f, ClampMax=360.0f))
+	float Rotation = 0.0f;
+
+	UPROPERTY(EditAnywhere, Category="Angle Settings", meta=(PCG_Overridable, ClampMin=-90.0f, ClampMax=90.0f))
+	float AxilAngle = 0.0f;
+
+	UPROPERTY(EditAnywhere, Category="Angle Settings", meta=(PCG_Overridable, ClampMin=-45.0f, ClampMax=0.0f))
+	float RandomizeAxilAngleMinimum = 0.0f;
+
+	UPROPERTY(EditAnywhere, Category="Angle Settings", meta=(PCG_Overridable, ClampMin=0.0f, ClampMax=45.0f))
+	float RandomizeAxilAngleMaximum = 0.0f;
+
+	UPROPERTY(EditAnywhere, Category="Angle Settings", meta=(PCG_Overridable))
+	EPVFoliageDistributionBasis AxilAngleRampBasis = EPVFoliageDistributionBasis::Plant;
+
+	UPROPERTY(EditAnywhere, Category="Angle Settings",
+		meta=(PCG_Overridable, XAxisMin=0.0f, XAxisMax=1.0f, YAxisMin=0.0f, YAxisMax=1.0f))
+	FPVFloatRamp AxilAngleRamp;
+
+	UPROPERTY(EditAnywhere, Category="Parametric Scale Settings", meta=(PCG_Overridable))
+	EPVFoliageDistributionBasis ScaleRampBasis = EPVFoliageDistributionBasis::Plant;
+
+	UPROPERTY(EditAnywhere, Category="Parametric Scale Settings", meta=(PCG_Overridable, ClampMin=0.0f, ClampMax=10.0f))
+	float BaseScale = 0.2f;
+
+	UPROPERTY(EditAnywhere, Category="Parametric Scale Settings", meta=(PCG_Overridable, ClampMin=0.0f, ClampMax=1.0f))
+	float RandomizeScaleAngleMinimum = 1.0f;
+
+	UPROPERTY(EditAnywhere, Category="Parametric Scale Settings", meta=(PCG_Overridable, ClampMin=1.0f, ClampMax=2.0f))
+	float RandomizeScaleAngleMaximum = 1.0f;
+
+	UPROPERTY(EditAnywhere, Category="Parametric Scale Settings",
+		meta=(PCG_Overridable, XAxisMin=0.0f, XAxisMax=1.0f, YAxisMin=0.0f, YAxisMax=1.0f))
+	FPVFloatRamp ScaleRamp;
+};
+ 
+UCLASS(BlueprintType, ClassGroup = (Procedural))
+class UPVFoliageDistributorSettings : public UPVBaseSettings
+{
+	GENERATED_BODY()
+ 
+	friend class FPVFoliageDistributorElement;
+	
+public:
+	UPVFoliageDistributorSettings();
+ 
+	#if WITH_EDITOR
+	virtual FName GetDefaultNodeName() const override { return FName(TEXT("ProceduralVegetationFoliageDistributor")); }
+	virtual FText GetDefaultNodeTitle() const override;
+	virtual FText GetNodeTooltipText() const override;
+	virtual FLinearColor GetNodeTitleColor() const override { return FLinearColor::Blue; }
+ 
+	virtual TArray<EPVRenderType> GetDefaultRenderType() const override { return TArray{ EPVRenderType::Mesh, EPVRenderType::Foliage }; }
+	virtual TArray<EPVRenderType> GetSupportedRenderTypes() const override { return TArray{ EPVRenderType::Mesh, EPVRenderType::Foliage, EPVRenderType::FoliageAttachments }; }
+	#endif
+ 
+protected:
+	virtual TArray<FPCGPinProperties> InputPinProperties() const override;
+	virtual FPCGDataTypeIdentifier GetInputPinTypeIdentifier() const override;
+	virtual FPCGDataTypeIdentifier GetOutputPinTypeIdentifier() const override;
+ 
+	virtual FPCGElementPtr CreateElement() const override;
+ 
+private:
+ 
+	UPROPERTY(EditAnywhere, Category="Distribution Settings Mode", meta=(PCG_Overridable))
+	EPVFoliageDistributionSettingsMode Mode = EPVFoliageDistributionSettingsMode::HormoneBasedSettings;
+	
+	UPROPERTY(EditAnywhere, Category="Parametric Settings",
+		meta=(PCG_Overridable, EditCondition="Mode == EPVFoliageDistributionSettingsMode::ParametricSettings", EditConditionHides))
+	FPVFoliageDistributionParametricSettings ParametricSettings;
+	
+	UPROPERTY(EditAnywhere, Category="Hormone Based Settings",
+		meta=(PCG_Overridable, EditCondition="Mode == EPVFoliageDistributionSettingsMode::HormoneBasedSettings", EditConditionHides))
+	FPVFoliageDistributionHormoneBasedSettings HormoneBasedSettings;
  
 	UPROPERTY(EditAnywhere, Category="Misc Settings", meta=(PCG_Overridable))
 	int32 RandomSeed = 123456;
diff --git a/Engine/Plugins/Experimental/ProceduralVegetationEditor/Source/ProceduralVegetation/Private/Nodes/PVGrafterSettings.cpp b/Engine/Plugins/Experimental/ProceduralVegetationEditor/Source/ProceduralVegetation/Private/Nodes/PVGrafterSettings.cpp
--- a/Engine/Plugins/Experimental/ProceduralVegetationEditor/Source/ProceduralVegetation/Private/Nodes/PVGrafterSettings.cpp
+++ b/Engine/Plugins/Experimental/ProceduralVegetationEditor/Source/ProceduralVegetation/Private/Nodes/PVGrafterSettings.cpp
@@ -112,33 +112,6 @@ FPCGElementPtr UPVGrafterSettings::CreateElement() const
 	return MakeShared<FPVGrafterElement>();
 }
 
-#if WITH_EDITOR
-bool UPVGrafterSettings::CanEditChange(const FProperty* InProperty) const
-{
-	bool bEditCondition = Super::CanEditChange(InProperty);
-	
-	if (!bEditCondition && (this->GetClass() != InProperty->GetOwnerClass()))
-	{
-		return bEditCondition;
-	}
-	
-	const bool ModeIsParametric = this->Mode == EPVDistributionSettingsMode::ParametricSettings;
-	const bool ModeIsHormoneBased = this->Mode == EPVDistributionSettingsMode::HormoneBasedSettings;
-	
-	if (InProperty->GetOwnerProperty()->GetFName() == GET_MEMBER_NAME_CHECKED(UPVGrafterSettings, ParametricSettings))
-	{
-		bEditCondition = ModeIsParametric;
-	}
-	
-	if (InProperty->GetOwnerProperty()->GetFName() == GET_MEMBER_NAME_CHECKED(UPVGrafterSettings, HormoneBasedSettings))
-	{
-		bEditCondition = ModeIsHormoneBased;
-	}
-	
-	return bEditCondition;
-}
-#endif
-
 bool FPVGrafterElement::ExecuteInternal(FPCGContext* InContext) const
 {
 	TRACE_CPUPROFILER_EVENT_SCOPE(FPVGrafterElement::ExecuteInternal);
@@ -183,6 +156,11 @@ bool FPVGrafterElement::ExecuteInternal(FPCGContext* InContext) const
 		FManagedArrayCollection SourceGraft = GraftElement->GetCollection();
 		SourceGraftCollections.Add(SourceGraft);
 	}
+	
+	if (SourceGraftCollections.Num() == 0)
+	{
+		return true;
+	}
 
 	FManagedArrayCollection OutCollection;
 	SourceSkeletonCollection.CopyTo(&OutCollection);
@@ -194,7 +172,7 @@ bool FPVGrafterElement::ExecuteInternal(FPCGContext* InContext) const
 		PCGLog::LogWarningOnGraph(LOCTEXT("FoliageDataRemovedFromSourceSkeleton", "Foliage data removed from source skeleton"), InContext);
 	}
 
-	if (Settings->Mode == EPVDistributionSettingsMode::HormoneBasedSettings)
+	if (Settings->Mode == EPVGrafterDistributionSettingsMode::HormoneBasedSettings)
 	{
 		const FPVGrafter::FHormoneBasedSettings::FDistributionSettings DistributionSettings = {
 			.EthyleneThreshold = Settings->HormoneBasedSettings.EthyleneThreshold,
@@ -232,24 +210,19 @@ bool FPVGrafterElement::ExecuteInternal(FPCGContext* InContext) const
 			.AxilAngleRampEffect = Settings->HormoneBasedSettings.AxilAngleRampEffect
 		};
 
-		const FPVGrafter::FMiscSettings MiscSettings = {
-			.RandomSeed = Settings->RandomSeed
-		};
-
 		FPVGrafter::DistributeGraftWithHormoneBasedSettings(
 			OutCollection,
-			SourceSkeletonCollection,
 			SourceGraftCollections,
 			DistributionSettings,
 			PhyllotaxySettings,
 			ScaleSettings,
 			VectorSettings,
-			MiscSettings,
+			Settings->RandomSeed,
 			Settings->Conditions);
 	}
 	else
 	{
-		const FPVGrafter::FParametricSettings::FSpacingSettings SpacingSettings = {
+		const FParametricSettings::FSpacingSettings SpacingSettings = {
 			.BranchDensity = Settings->ParametricSettings.BranchDensity,
 			.RelativeStart = Settings->ParametricSettings.RelativeStart,
 			.RelativeEnd = Settings->ParametricSettings.RelativeEnd,
@@ -257,45 +230,40 @@ bool FPVGrafterElement::ExecuteInternal(FPCGContext* InContext) const
 			.StartGeneration = Settings->ParametricSettings.StartGeneration,
 			.LimitEndGeneration = Settings->ParametricSettings.LimitEndGeneration,
 			.EndGeneration = Settings->ParametricSettings.EndGeneration,
-			.SpacingBasis = Settings->ParametricSettings.SpacingBasis == EPVDistributionBasis::Plant
-			? FPVGrafter::FParametricSettings::EDistributionBasis::Plant
-			: FPVGrafter::FParametricSettings::EDistributionBasis::Branch,
+			.SpacingBasis = Settings->ParametricSettings.SpacingBasis == EPVGrafterDistributionBasis::Plant
+			? FParametricSettings::EDistributionBasis::Plant
+			: FParametricSettings::EDistributionBasis::Branch,
 			.SpacingRamp = &Settings->ParametricSettings.SpacingRamp
 		};
 
-		const FPVGrafter::FParametricSettings::FAngleSettings AngleSettings = {
+		const FParametricSettings::FAngleSettings AngleSettings = {
 			.Rotation = Settings->ParametricSettings.Rotation,
 			.AxilAngle = Settings->ParametricSettings.AxilAngle,
 			.RandomizeAxilAngleMinimum = Settings->ParametricSettings.RandomizeAxilAngleMinimum,
 			.RandomizeAxilAngleMaximum = Settings->ParametricSettings.RandomizeAxilAngleMaximum,
-			.AxilAngleRampBasis = Settings->ParametricSettings.AxilAngleRampBasis == EPVDistributionBasis::Plant
-			? FPVGrafter::FParametricSettings::EDistributionBasis::Plant
-			: FPVGrafter::FParametricSettings::EDistributionBasis::Branch,
+			.AxilAngleRampBasis = Settings->ParametricSettings.AxilAngleRampBasis == EPVGrafterDistributionBasis::Plant
+			? FParametricSettings::EDistributionBasis::Plant
+			: FParametricSettings::EDistributionBasis::Branch,
 			.AxilAngleRamp = &Settings->ParametricSettings.AxilAngleRamp
 		};
 
-		const FPVGrafter::FParametricSettings::FScaleSettings ScaleSettings = {
-			.ScaleRampBasis = Settings->ParametricSettings.ScaleRampBasis == EPVDistributionBasis::Plant
-			? FPVGrafter::FParametricSettings::EDistributionBasis::Plant
-			: FPVGrafter::FParametricSettings::EDistributionBasis::Branch,
+		const FParametricSettings::FScaleSettings ScaleSettings = {
+			.ScaleRampBasis = Settings->ParametricSettings.ScaleRampBasis == EPVGrafterDistributionBasis::Plant
+			? FParametricSettings::EDistributionBasis::Plant
+			: FParametricSettings::EDistributionBasis::Branch,
 			.BaseScale = Settings->ParametricSettings.BaseScale,
 			.RandomizeScaleAngleMinimum = Settings->ParametricSettings.RandomizeScaleAngleMinimum,
 			.RandomizeScaleAngleMaximum = Settings->ParametricSettings.RandomizeScaleAngleMaximum,
 			.ScaleRamp = &Settings->ParametricSettings.ScaleRamp
 		};
 
-		const FPVGrafter::FMiscSettings MiscSettings = {
-			.RandomSeed = Settings->RandomSeed
-		};
-
 		FPVGrafter::DistributeGraftWithParametricSettings(
 			OutCollection,
-			SourceSkeletonCollection,
 			SourceGraftCollections,
 			SpacingSettings,
 			AngleSettings,
 			ScaleSettings,
-			MiscSettings,
+			Settings->RandomSeed,
 			Settings->Conditions);
 	}
 
diff --git a/Engine/Plugins/Experimental/ProceduralVegetationEditor/Source/ProceduralVegetation/Private/Nodes/PVGrafterSettings.h b/Engine/Plugins/Experimental/ProceduralVegetationEditor/Source/ProceduralVegetation/Private/Nodes/PVGrafterSettings.h
--- a/Engine/Plugins/Experimental/ProceduralVegetationEditor/Source/ProceduralVegetation/Private/Nodes/PVGrafterSettings.h
+++ b/Engine/Plugins/Experimental/ProceduralVegetationEditor/Source/ProceduralVegetation/Private/Nodes/PVGrafterSettings.h
@@ -16,21 +16,21 @@ namespace PVGrafterInputPins
 }
 
 UENUM()
-enum class EPVDistributionSettingsMode : uint8
+enum class EPVGrafterDistributionSettingsMode : uint8
 {
 	ParametricSettings UMETA(DisplayName="Use Parametric Settings"),
 	HormoneBasedSettings UMETA(DisplayName="Use Hormone Based Settings")
 };
 
 UENUM()
-enum class EPVDistributionBasis : uint8
+enum class EPVGrafterDistributionBasis : uint8
 {
 	Plant UMETA(DisplayName="Plant"),
 	Branch UMETA(DisplayName="Branch")
 };
 
 USTRUCT()
-struct FPVDistributionParametricSettings
+struct FPVGrafterDistributionParametricSettings
 {
 	GENERATED_BODY()
 
@@ -58,7 +58,7 @@ struct FPVDistributionParametricSettings
 	int32 EndGeneration = 0;
 
 	UPROPERTY(EditAnywhere, Category="Spacing Settings", meta=(PCG_Overridable))
-	EPVDistributionBasis SpacingBasis = EPVDistributionBasis::Plant;
+	EPVGrafterDistributionBasis SpacingBasis = EPVGrafterDistributionBasis::Plant;
 
 	UPROPERTY(EditAnywhere, Category="Spacing Settings",
 		meta=(PCG_Overridable, XAxisMin=0.0f, XAxisMax=1.0f, YAxisMin=0.0f, YAxisMax=1.0f))
@@ -77,14 +77,14 @@ struct FPVDistributionParametricSettings
 	float RandomizeAxilAngleMaximum = 0.0f;
 
 	UPROPERTY(EditAnywhere, Category="Angle Settings", meta=(PCG_Overridable))
-	EPVDistributionBasis AxilAngleRampBasis = EPVDistributionBasis::Plant;
+	EPVGrafterDistributionBasis AxilAngleRampBasis = EPVGrafterDistributionBasis::Plant;
 
 	UPROPERTY(EditAnywhere, Category="Angle Settings",
 		meta=(PCG_Overridable, XAxisMin=0.0f, XAxisMax=1.0f, YAxisMin=0.0f, YAxisMax=1.0f))
 	FPVFloatRamp AxilAngleRamp;
 
 	UPROPERTY(EditAnywhere, Category="Parametric Scale Settings", meta=(PCG_Overridable))
-	EPVDistributionBasis ScaleRampBasis = EPVDistributionBasis::Plant;
+	EPVGrafterDistributionBasis ScaleRampBasis = EPVGrafterDistributionBasis::Plant;
 
 	UPROPERTY(EditAnywhere, Category="Parametric Scale Settings", meta=(PCG_Overridable, ClampMin=0.0f, ClampMax=10.0f))
 	float BaseScale = 0.2f;
@@ -101,7 +101,7 @@ struct FPVDistributionParametricSettings
 };
 
 USTRUCT()
-struct FPVDistributionHormoneBasedSettings
+struct FPVGrafterDistributionHormoneBasedSettings
 {
 	GENERATED_BODY()
 
@@ -274,21 +274,18 @@ protected:
 	virtual FPCGDataTypeIdentifier GetSkeletonPinTypeIdentifier() const;
 	virtual FPCGDataTypeIdentifier GetGraftPinTypeIdentifier() const;
 	virtual FPCGElementPtr CreateElement() const override;
-#if WITH_EDITOR
-	virtual bool CanEditChange(const FProperty* InProperty) const override;
-#endif
 
 private:
 	UPROPERTY(EditAnywhere, Category="Distribution Settings Mode", meta=(PCG_Overridable))
-	EPVDistributionSettingsMode Mode = EPVDistributionSettingsMode::ParametricSettings;
+	EPVGrafterDistributionSettingsMode Mode = EPVGrafterDistributionSettingsMode::ParametricSettings;
 
 	UPROPERTY(EditAnywhere, Category="Parametric Settings",
-		meta=(PCG_Overridable, ShowOnlyInnerProperties))
-	FPVDistributionParametricSettings ParametricSettings;
+		meta=(PCG_Overridable, EditCondition="Mode == EPVGrafterDistributionSettingsMode::ParametricSettings", EditConditionHides))
+	FPVGrafterDistributionParametricSettings ParametricSettings;
 
 	UPROPERTY(EditAnywhere, Category="Hormone Based Settings",
-		meta=(PCG_Overridable, ShowOnlyInnerProperties))
-	FPVDistributionHormoneBasedSettings HormoneBasedSettings;
+		meta=(PCG_Overridable, EditCondition="Mode == EPVGrafterDistributionSettingsMode::HormoneBasedSettings", EditConditionHides))
+	FPVGrafterDistributionHormoneBasedSettings HormoneBasedSettings;
 
 	UPROPERTY(EditAnywhere, Category="Misc Settings", meta=(PCG_Overridable))
 	int32 RandomSeed = 123456;
diff --git a/Engine/Plugins/Experimental/ProceduralVegetationEditor/Source/ProceduralVegetationEditor/ProceduralVegetationEditor.Build.cs b/Engine/Plugins/Experimental/ProceduralVegetationEditor/Source/ProceduralVegetationEditor/ProceduralVegetationEditor.Build.cs
--- a/Engine/Plugins/Experimental/ProceduralVegetationEditor/Source/ProceduralVegetationEditor/ProceduralVegetationEditor.Build.cs
+++ b/Engine/Plugins/Experimental/ProceduralVegetationEditor/Source/ProceduralVegetationEditor/ProceduralVegetationEditor.Build.cs
@@ -1,5 +1,6 @@
 // Copyright Epic Games, Inc. All Rights Reserved.
 
+using System.IO;
 using UnrealBuildTool;
 
 public class ProceduralVegetationEditor : ModuleRules
@@ -7,6 +8,8 @@ public class ProceduralVegetationEditor : ModuleRules
 	public ProceduralVegetationEditor(ReadOnlyTargetRules Target) : base(Target)
 	{
 		PCHUsage = PCHUsageMode.UseExplicitOrSharedPCHs;
+		
+		PrivateIncludePaths.Add(Path.Combine(ModuleDirectory, "../ProceduralVegetation/Private"));
 
 		PrivateDependencyModuleNames.AddRange(
 			new string[]
