[PVE] Issue: Clicking on editor's export toolbar button causes a crash

Fix/How: When you dont have any output node present in the graph, it causes a crash. We were expecting that atleast one output node would always be present in the graph but thats not the case. We are trying to select output 0th entry which was causing a crash. Added a index validation before trying to access it.

There was another issue which was not showing any output in the export panel. When you have a subgraph and you connect it to output node it adds a conversion node. We were checking PinOfType in the code which doesnt work in this case. Instead of pin type, we can check the data type to avoid this issue.

#Synced-CL 51579190

diff --git a/Engine/Plugins/Experimental/ProceduralVegetationEditor/Source/ProceduralVegetationEditor/Private/PVEditor.cpp b/Engine/Plugins/Experimental/ProceduralVegetationEditor/Source/ProceduralVegetationEditor/Private/PVEditor.cpp
--- a/Engine/Plugins/Experimental/ProceduralVegetationEditor/Source/ProceduralVegetationEditor/Private/PVEditor.cpp
+++ b/Engine/Plugins/Experimental/ProceduralVegetationEditor/Source/ProceduralVegetationEditor/Private/PVEditor.cpp
@@ -701,19 +701,20 @@ void FPVEditor::OnExport()
 				}
 
 				UPCGPin* const InPin = GetInPinFromNode(Node);
-				UPCGPin* const UpstreamOutPin = GetUpstreamPin(InPin);
-
-				const bool bIsMeshDataPin = IsPinOfType<FPVDataTypeInfoMesh>(UpstreamOutPin);
-				const bool bIsFoliageMeshDataPin = IsPinOfType<FPVDataTypeInfoFoliageMesh>(UpstreamOutPin);
-				if (!bIsMeshDataPin && !bIsFoliageMeshDataPin)
-				{
-					return true;
-				}
-
+				
 				if (const UPVData* const NodeData = GetDataFromPin(InPin))
 				{
+					const bool bIsMeshDataPin = NodeData->IsA<UPVMeshData>();
+					const bool bIsFoliageMeshDataPin = NodeData->IsA<UPVFoliageMeshData>();
+					
+					if (!bIsMeshDataPin && !bIsFoliageMeshDataPin)
+					{
+						return true;
+					}
+					
 					const TObjectPtr<UPVExportEntry>& ExportEntry = ExportEntries.
 						Emplace_GetRef(NewObject<UPVExportEntry>(ProceduralVegetationGraph));
+					
 					ExportEntry->Initialize(NodeData->GetSharedCollection(), Node, IsNodeSelected(Node), bIsFoliageMeshDataPin);
 				}
 
diff --git a/Engine/Plugins/Experimental/ProceduralVegetationEditor/Source/ProceduralVegetationEditor/Private/Widgets/SPVExportSelectionDialog.cpp b/Engine/Plugins/Experimental/ProceduralVegetationEditor/Source/ProceduralVegetationEditor/Private/Widgets/SPVExportSelectionDialog.cpp
--- a/Engine/Plugins/Experimental/ProceduralVegetationEditor/Source/ProceduralVegetationEditor/Private/Widgets/SPVExportSelectionDialog.cpp
+++ b/Engine/Plugins/Experimental/ProceduralVegetationEditor/Source/ProceduralVegetationEditor/Private/Widgets/SPVExportSelectionDialog.cpp
@@ -211,7 +211,7 @@ void SPVExportSelectionDialog::Construct(const FArguments& InArgs)
 		]
 	);
 
-	if (ExportEntriesTableView)
+	if (ExportEntriesTableView && ExportEntries.Num() > 0)
 	{
 		ExportEntriesTableView->SetSelection(ExportEntries[0]);
 	}
