Re-work PVImporter_StaticMesh.cpp::ConnectBranchesToParent

Rework the branch-to-parent connection algorithm to find the closest edge instead of the closest vertex, inserting a split point on the target edge for more accurate attachment
Replace the single-pass greedy merge with a multi-step approach: build candidate connections via triangle-based proximity queries, select a cycle-free spanning tree via BFS from root nodes, split target edges with alpha-sorted multi-split, then stitch endpoints to the resulting split vertices
Add a FindTrianglesInRadius helper and reorder FindRootPoints to run before ConnectBranchesToParent so root data is available during connection selection

#Synced-CL 51495698

diff --git a/Engine/Plugins/Experimental/ProceduralVegetationEditor/Source/ProceduralVegetation/Private/Helpers/PVDynamicMeshHelpers.cpp b/Engine/Plugins/Experimental/ProceduralVegetationEditor/Source/ProceduralVegetation/Private/Helpers/PVDynamicMeshHelpers.cpp
--- a/Engine/Plugins/Experimental/ProceduralVegetationEditor/Source/ProceduralVegetation/Private/Helpers/PVDynamicMeshHelpers.cpp
+++ b/Engine/Plugins/Experimental/ProceduralVegetationEditor/Source/ProceduralVegetation/Private/Helpers/PVDynamicMeshHelpers.cpp
@@ -70,4 +70,41 @@ TArray<int32> PV::DynamicMeshHelper::FindVerticesInRadius(const UE::Geometry::FD
 	DynamicMeshAABBTree.DoTraversal(TreeTraversal);
 
 	return OutVertexIndices;
+}
+
+TArray<int32> PV::DynamicMeshHelper::FindTrianglesInRadius(const UE::Geometry::FDynamicMeshAABBTree3& DynamicMeshAABBTree, const FVector& Pos, float Radius)
+{
+	TArray<int32> OutTriangleIDs;
+
+	const UE::Geometry::FDynamicMesh3* DynamicMesh = DynamicMeshAABBTree.GetMesh();
+	if (DynamicMesh == nullptr)
+	{
+		return OutTriangleIDs;
+	}
+
+	const double RadiusSqr = Radius * Radius;
+
+	UE::Geometry::FDynamicMeshAABBTree3::FTreeTraversal TreeTraversal;
+
+	TreeTraversal.NextBoxF = [&](const UE::Geometry::FAxisAlignedBox3d& Box, int32 Depth)->bool
+	{
+		return Box.DistanceSquared(Pos) <= RadiusSqr;
+	};
+
+	TreeTraversal.NextTriangleF = [&](int32 TID)
+	{
+		const UE::Geometry::FIndex3i Triangle = DynamicMesh->GetTriangle(TID);
+
+		using FDynamicMeshQueries3 = UE::Geometry::TMeshQueries<UE::Geometry::FDynamicMesh3>;
+		const UE::Geometry::FDistPoint3Triangle3d Query = FDynamicMeshQueries3::TriangleDistance(*DynamicMesh, TID, Pos);
+
+		if (FVector::DistSquared(Pos, Query.ClosestTrianglePoint) < RadiusSqr)
+		{
+			OutTriangleIDs.AddUnique(TID);
+		}
+	};
+
+	DynamicMeshAABBTree.DoTraversal(TreeTraversal);
+
+	return OutTriangleIDs;
 }
\ No newline at end of file
diff --git a/Engine/Plugins/Experimental/ProceduralVegetationEditor/Source/ProceduralVegetation/Private/Helpers/PVDynamicMeshHelpers.h b/Engine/Plugins/Experimental/ProceduralVegetationEditor/Source/ProceduralVegetation/Private/Helpers/PVDynamicMeshHelpers.h
--- a/Engine/Plugins/Experimental/ProceduralVegetationEditor/Source/ProceduralVegetation/Private/Helpers/PVDynamicMeshHelpers.h
+++ b/Engine/Plugins/Experimental/ProceduralVegetationEditor/Source/ProceduralVegetation/Private/Helpers/PVDynamicMeshHelpers.h
@@ -19,4 +19,6 @@ namespace PV::DynamicMeshHelper
 	TArray<int32> GetVertexNeighbours(const UE::Geometry::FDynamicMesh3& DynamicMesh, int32 VertexID);
 
 	TArray<int32> FindVerticesInRadius(const UE::Geometry::FDynamicMeshAABBTree3& DynamicMeshAABBTree, const FVector& Pos, float Radius);
+
+	TArray<int32> FindTrianglesInRadius(const UE::Geometry::FDynamicMeshAABBTree3& DynamicMeshAABBTree, const FVector& Pos, float Radius);
 };
\ No newline at end of file
diff --git a/Engine/Plugins/Experimental/ProceduralVegetationEditor/Source/ProceduralVegetation/Private/Helpers/PVImportHelpers.cpp b/Engine/Plugins/Experimental/ProceduralVegetationEditor/Source/ProceduralVegetation/Private/Helpers/PVImportHelpers.cpp
--- a/Engine/Plugins/Experimental/ProceduralVegetationEditor/Source/ProceduralVegetation/Private/Helpers/PVImportHelpers.cpp
+++ b/Engine/Plugins/Experimental/ProceduralVegetationEditor/Source/ProceduralVegetation/Private/Helpers/PVImportHelpers.cpp
@@ -198,7 +198,7 @@ namespace PV::ImportHelper
 			const FVector3d& PrevVertexPos = DynamicMesh.GetVertex(PrevVertexID);
 			const FVector3d ApicalDirection = (VertexPos - PrevVertexPos).GetSafeNormal();
 
-			float BestDot = -1;
+			float BestDot = -MAX_flt;
 			int32 BestDotIndex = INDEX_NONE;
 			for (int32 NextVertexID : NextPoints)
 			{
diff --git a/Engine/Plugins/Experimental/ProceduralVegetationEditor/Source/ProceduralVegetation/Private/Implementations/PVImporter_StaticMesh.cpp b/Engine/Plugins/Experimental/ProceduralVegetationEditor/Source/ProceduralVegetation/Private/Implementations/PVImporter_StaticMesh.cpp
--- a/Engine/Plugins/Experimental/ProceduralVegetationEditor/Source/ProceduralVegetation/Private/Implementations/PVImporter_StaticMesh.cpp
+++ b/Engine/Plugins/Experimental/ProceduralVegetationEditor/Source/ProceduralVegetation/Private/Implementations/PVImporter_StaticMesh.cpp
@@ -652,129 +652,323 @@ static void FindEndPoints(
 }
 
 static void ConnectBranchesToParent(
-	UE::Geometry::FDynamicMesh3& DynamicMesh,
-	const FFloatVertexAttribute& VertexScaleAttribute,
+	FDynamicMesh3& DynamicMesh,
+	const FBooleanVertexAttribute& RootPointAttribute,
+	FFloatVertexAttribute& VertexScaleAttribute,
 	FBooleanVertexAttribute& EndPointAttribute,
 	FBooleanVertexAttribute& FillerVertexAttribute
 )
 {
-	TRACE_CPUPROFILER_EVENT_SCOPE(ConnectBranchesToParent);
-
-	// This algorithm takes the tips of all the branches which have been individually collapsed into a line, and tries to connect them to their parent.
-	// This is done by performing an overlap query at the tip of each end point, and connecting that tip to the closest overlapping vertex which
-	// is not part of the same mesh component.
+	// Connects the ends of branch mesh components onto their parent (trunk or larger branch) components,
+	// so the combined DynamicMesh forms a single connected tree rather than a set of disjoint meshes.
 	//
-	// Issues with this algorithm:
-	// 1. Since we do not know what is the "tip" and the "base" of the branch at this point, we perform the query for each
-	// end point and just hope that the tips do not overlap any other vertex. They could however be overlapping, which would lead to the tips becoming
-	// fused and for loops to be created in the tree.
-	// 2. We only fuse the end point to the parent if the radii of the end point and the nearby vertex overlap. This may lead to endpoints not being
-	// connected if the closest vertex is too far away (even if the end-point is very close to the parent mesh itself). This would lead to the entire
-	// branch being excluded from the final tree.
-
-	const UE::Geometry::FDynamicMeshAABBTree3 DynamicMeshAABBTree(&DynamicMesh);
-	UE::Geometry::FMeshConnectedComponents ConnectedComponents(&DynamicMesh);
+	// Node graph representation
+	//   Each connected component in the DynamicMesh is treated as one node in a graph.
+	//   - RootPoints  : vertices tagged on the trunk/root component; used as BFS seeds.
+	//   - EndPoints   : vertices at the open end of each branch component. Each EndPoint 
+	//                   generates one outgoing directed edge candidate toward
+	//                   whichever neighbouring component it is physically touching.
+	//   - Connections : candidate directed edges (SourceNode→TargetNode). Each connection records the
+	//                   closest non-filler edge on the target component and the alpha along that edge,
+	//                   so the edge can later be split precisely at the attachment point.
+	//
+	// Execution steps
+	//   1. Build candidate connections  – for every EndPoint, sphere-query nearby triangles, discard
+	//      same-component hits, and for each neighbouring node find the closest non-filler edge. Discard
+	//      candidates where the radii of the two vertices do not overlap (i.e. no physical adjacency).
+	//   2. BFS spanning-tree selection  – BFS from every RootPoint node; at each visited node, pick the
+	//      closest-distance incoming connection from each unvisited neighbour and mark it accepted.
+	//      This guarantees a cycle-free tree and a single attachment point per branch component.
+	//   3. Split target edges           – accepted connections are grouped by target edge and sorted by
+	//      alpha. Each edge is split with a tail-advancing pattern so earlier splits do not invalidate
+	//      later vertex IDs. Points within 1 cm of an existing endpoint snap to it to avoid degenerate
+	//      zero-length splits.
+	//   4. Weld EndPoints               – for each accepted connection, call MergeVertices to weld the
+	//      EndPoint vertex to the target vertex, keeping the target vertex's position. Both vertices are
+	//      unmarked as endpoints before the merge so they are no longer treated as open ends.
+
+	FMeshConnectedComponents ConnectedComponents(&DynamicMesh);
 	ConnectedComponents.FindConnectedVertices();
 
-	TMap<int32, int32> VertexComponentMap;
-	VertexComponentMap.Reserve(DynamicMesh.VertexCount());
+	if (ConnectedComponents.Components.Num() == 1)
+	{
+		return;
+	}
 
-	for (int32 i = 0; i < ConnectedComponents.Components.Num(); ++i)
+	const TArray<int32> EndPointVertexIndices = EndPointAttribute.FindAllNonZero();
+	const int32 NumEndPoints = EndPointVertexIndices.Num();
+
+	TArray<int32> VertexIndexToNodeIndex;
 	{
-		for (int32 VertexID : ConnectedComponents.Components[i].Indices)
+		VertexIndexToNodeIndex.SetNum(DynamicMesh.MaxVertexID());
+		for (int32 i = 0; i < ConnectedComponents.Components.Num(); ++i)
 		{
-			VertexComponentMap.Add(VertexID, i);
+			for (int32 VertexID : ConnectedComponents.Components[i].Indices)
+			{
+				VertexIndexToNodeIndex[VertexID] = i;
+			}
 		}
 	}
 
-	const TArray<int32> EndPoints = EndPointAttribute.FindAllNonZero();
-	TArray<int32> MergeMap;
-	MergeMap.SetNum(EndPoints.Num());
-	for (int32 i = 0; i < EndPoints.Num(); ++i)
+	using FConnectionArray = TArray<int32, TInlineAllocator<4>>;
+	struct FNode
+	{
+		FConnectionArray OutgoingConnections;
+		TMap<int32, FConnectionArray> IncomingConnections; // Maps OtherNodeIndex->ConnectionIndices
+	};
+
+	struct FConnection
 	{
-		const int32 EndPointVertexIndex = EndPoints[i];
+		int32 SourceNode;
+		int32 TargetNode;
+		int32 SourceEndPointIndex;
 		
-		const FVector VertexPosition = DynamicMesh.GetVertex(EndPointVertexIndex);
-		const float VertexRadius = VertexScaleAttribute.GetValue(EndPointVertexIndex);
-		const float SearchRadius = VertexRadius * 4.f;
-		const TArray<int32> OverlappingVertices = PV::DynamicMeshHelper::FindVerticesInRadius(DynamicMeshAABBTree, VertexPosition, SearchRadius);
+		FIndex2i TargetEdge;
+		FVector TargetPoint = FVector::ZeroVector;
+		float TargetPointEdgeAlpha = 0;
+		float Dist = 0;
+		float TargetPointScale = 0;
+	};
+
+	TArray<FNode> Nodes;
+	Nodes.SetNum(ConnectedComponents.Components.Num());
 
-		float ClosestVertexDistSqr = MAX_flt;
-		int32 VertexToMergeWith = INDEX_NONE;
-		for (int32 OverlappingVertexIndex : OverlappingVertices)
+	// Step 1: Build candidate connections.
+	TArray<FConnection> Connections;
+	{
+		const FDynamicMeshAABBTree3 DynamicMeshAABBTree(&DynamicMesh);
+		for (int32 EndPointIndex = 0; EndPointIndex < NumEndPoints; ++EndPointIndex)
 		{
-			if (EndPointAttribute.GetValue(OverlappingVertexIndex)
-				|| FillerVertexAttribute.GetValue(OverlappingVertexIndex))
+			const int32 EndPointVertexIndex = EndPointVertexIndices[EndPointIndex];
+			const FVector VertexPosition = DynamicMesh.GetVertex(EndPointVertexIndex);
+			const int32 EndPointNodeIndex = VertexIndexToNodeIndex[EndPointVertexIndex];
+			const float SearchRadius = 100.f; // TODO: Base this on the trunk radius?
+
+			// Potential optimization: FindTrianglesInRadius could return a large number of triangles, could use FDynamicMeshAABBTree3::FindNearestTriangle
+			const TArray<int32> CloseTriangles = PV::DynamicMeshHelper::FindTrianglesInRadius(DynamicMeshAABBTree, VertexPosition, SearchRadius);
+
+			TMap<int32, TArray<FIndex3i>> NodeToTriangles;
+			NodeToTriangles.Reserve(Nodes.Num());
+			for (const int32 TriangleID : CloseTriangles)
 			{
-				continue;
+				const FIndex3i Triangle = DynamicMesh.GetTriangle(TriangleID);
+				const int32 OtherNodeIndex = VertexIndexToNodeIndex[Triangle.A];
+				const bool bIsSameComponent = OtherNodeIndex == EndPointNodeIndex;
+				if (bIsSameComponent) // Ignore triangles on the same component as the EndPoint
+				{
+					continue;
+				}
+
+				NodeToTriangles.FindOrAdd(OtherNodeIndex).Add(Triangle);
 			}
 
-			int32* ComponentIndex1 = VertexComponentMap.Find(EndPointVertexIndex);
-			int32* ComponentIndex2 = VertexComponentMap.Find(OverlappingVertexIndex);
-			check(ComponentIndex1 && ComponentIndex2);
-			if (*ComponentIndex1 == *ComponentIndex2)
+			for (const auto& [NodeIndex, Triangles] : NodeToTriangles)
 			{
-				continue;
+				FIndex2i ClosestEdge = FIndex2i(INDEX_NONE, INDEX_NONE);
+				FVector ClosestEdgePoint = FVector::ZeroVector;
+				float ClosestEdgeDistSqr = MAX_flt;
+				float ClosestPointScale = 0;
+				float ClosestPointEdgeAlpha = 0;
+				for (const FIndex3i& Triangle : Triangles)
+				{
+					for (int32 j = 0; j < 3; ++j)
+					{
+						FIndex2i Edge = FIndex2i(Triangle[j], Triangle[(j + 1) % 3]);
+						Edge.Sort();
+						if (FillerVertexAttribute.GetValue(Edge.A) || FillerVertexAttribute.GetValue(Edge.B))
+						{
+							continue;
+						}
+
+						const FVector VertexA = DynamicMesh.GetVertex(Edge.A);
+						const FVector VertexB = DynamicMesh.GetVertex(Edge.B);
+						const FVector ClosestPoint = FMath::ClosestPointOnSegment(VertexPosition, VertexA, VertexB);
+						const float DistSqr = (VertexPosition - ClosestPoint).SizeSquared();
+						if (DistSqr < ClosestEdgeDistSqr)
+						{
+							const float DistAlongOriginalEdge = (VertexA - ClosestPoint).Size();
+							const float OriginalEdgeSize = (VertexB - VertexA).Size();
+							const float EdgeAlpha = OriginalEdgeSize > 0 ? DistAlongOriginalEdge / OriginalEdgeSize : 0.f;
+							
+							const float VertexScaleA = VertexScaleAttribute.GetValue(Edge.A);
+							const float VertexScaleB = VertexScaleAttribute.GetValue(Edge.B);
+							const float ThisPointScale = VertexScaleAttribute.GetValue(EndPointVertexIndex);
+							const float TargetPointScale = FMath::Lerp(VertexScaleA, VertexScaleB, EdgeAlpha);
+							
+							const float Dist = DistSqr > 0 ? FMath::Sqrt(DistSqr) : 0;
+							if (Dist > (TargetPointScale + ThisPointScale) * 1.5f) // Do not connect if radii do not overlap
+							{
+								continue; 
+							}
+
+							ClosestEdge = Edge;
+							ClosestEdgePoint = ClosestPoint;
+							ClosestEdgeDistSqr = DistSqr;
+							ClosestPointScale = TargetPointScale;
+							ClosestPointEdgeAlpha = EdgeAlpha;
+						}
+					}
+				}
+
+				if (ClosestEdge.A == INDEX_NONE || ClosestEdge.B == INDEX_NONE)
+				{
+					continue;
+				}
+
+				FConnection& Connection = Connections.AddDefaulted_GetRef();
+				Connection.SourceNode = EndPointNodeIndex;
+				Connection.TargetNode = NodeIndex;
+				Connection.SourceEndPointIndex = EndPointIndex;
+				Connection.TargetPoint = ClosestEdgePoint;
+				Connection.Dist = ClosestEdgeDistSqr > 0 ? FMath::Sqrt(ClosestEdgeDistSqr) : 0;
+				Connection.TargetEdge = ClosestEdge;
+				Connection.TargetPointEdgeAlpha = ClosestPointEdgeAlpha;
+				Connection.TargetPointScale = ClosestPointScale;
+
+				Nodes[Connection.SourceNode].OutgoingConnections.Add(Connections.Num() - 1);
+				Nodes[Connection.TargetNode].IncomingConnections.FindOrAdd(Connection.SourceNode).Add(Connections.Num() - 1);
 			}
+		}
+	}
+
+	// Step 2: BFS from the root(s) to select one connection per node (spanning tree, no cycles)
+	TArray<bool> ConnectionFilter;
+	{
+		ConnectionFilter.SetNumZeroed(Connections.Num());
+
+		TArray<bool> ConectedNodes;
+		ConectedNodes.SetNumZeroed(Nodes.Num());
+
+		TQueue<int32> NodeQueue;
 
-			const float OverlappingVertexRadius = VertexScaleAttribute.GetValue(OverlappingVertexIndex);
-			const FVector OverlappingVertexPosition = DynamicMesh.GetVertex(OverlappingVertexIndex);
-			const float CombinedRadius = OverlappingVertexRadius + VertexRadius;
-			const float CombinedRadiusSqr = CombinedRadius * CombinedRadius;
-			const float DistSquared = (OverlappingVertexPosition - VertexPosition).SizeSquared();
-			if (DistSquared > CombinedRadiusSqr)
+		const TArray<int32> RootPointVertexIndices = RootPointAttribute.FindAllNonZero();
+		for (int32 RootPointVertexIndex : RootPointVertexIndices)
+		{
+			const int32 RootNodeIndex = VertexIndexToNodeIndex[RootPointVertexIndex];
+			NodeQueue.Enqueue(RootNodeIndex);
+			ConectedNodes[RootNodeIndex] = true;
+		}
+
+		int32 NodeIndex = INDEX_NONE;
+		while (NodeQueue.Dequeue(NodeIndex))
+		{
+			const FNode& Node = Nodes[NodeIndex];
+			for (const auto& [OtherNodeIndex, ConnectionIndices] : Node.IncomingConnections)
 			{
-				continue;
+				if (ConectedNodes[OtherNodeIndex])
+				{
+					continue;
+				}
+
+				int32 BestConnectionIndex = INDEX_NONE;
+				float BestConnectionScore = -1.f;
+				for (int32 ConnectionIndex : ConnectionIndices)
+				{
+					const FConnection& Connection = Connections[ConnectionIndex];
+					const float Score = (1.f / FMath::Max(Connection.Dist, 1.f));
+					if (Score > BestConnectionScore)
+					{
+						BestConnectionIndex = ConnectionIndex;
+						BestConnectionScore = Score;
+					}
+				}
+
+				ConnectionFilter[BestConnectionIndex] = true;
+				ConectedNodes[OtherNodeIndex] = true;
+				NodeQueue.Enqueue(OtherNodeIndex);
 			}
+		}
+	}
+
+	// Step 3: Split target edges and populate EndPointToVertexConnection for Step 4.
+	TArray<int32> EndPointToVertexConnection;
+	{
+		EndPointToVertexConnection.SetNum(EndPointVertexIndices.Num());
+		for (int32 EndPointIndex = 0; EndPointIndex < EndPointVertexIndices.Num(); ++EndPointIndex)
+		{
+			EndPointToVertexConnection[EndPointIndex] = INDEX_NONE;
+		}
 
-			const FVector OtherVertexPosition = DynamicMesh.GetVertex(OverlappingVertexIndex);
-			const float DistSqr = (VertexPosition - OtherVertexPosition).SizeSquared();
-			if (DistSqr < ClosestVertexDistSqr)
+		// Group accepted connections by target edge so multi-split can be handled per edge.
+		TMap<FIndex2i, TArray<int32>> EdgeToAcceptedConnectionIndices;
+		for (int32 ConnectionIndex = 0; ConnectionIndex < Connections.Num(); ++ConnectionIndex)
+		{
+			if (ConnectionFilter[ConnectionIndex])
 			{
-				VertexToMergeWith = OverlappingVertexIndex;
-				ClosestVertexDistSqr = DistSqr;
+				EdgeToAcceptedConnectionIndices.FindOrAdd(Connections[ConnectionIndex].TargetEdge).Add(ConnectionIndex);
 			}
 		}
 
-		MergeMap[i] = VertexToMergeWith;
+		for (auto& [Edge, ConnectionIndices] : EdgeToAcceptedConnectionIndices)
+		{
+			ConnectionIndices.Sort([&](int32 A, int32 B)
+			{
+				return Connections[A].TargetPointEdgeAlpha < Connections[B].TargetPointEdgeAlpha;
+			});
+
+			int32 TailA = Edge.A;
+			const int32 TailB = Edge.B;
+			for (const int32 ConnectionIndex : ConnectionIndices)
+			{
+				const FConnection& Connection = Connections[ConnectionIndex];
+
+				// If the target point is within 1cm of an existing tail endpoint, reuse that vertex to avoid a degenerate split.
+				const bool bSnapToTailA = DynamicMesh.GetVertex(TailA).Equals(Connection.TargetPoint, 1.f);
+				const bool bSnapToTailB = DynamicMesh.GetVertex(TailB).Equals(Connection.TargetPoint, 1.f);
+				if (bSnapToTailA)
+				{
+					EndPointToVertexConnection[Connection.SourceEndPointIndex] = TailA;
+					continue;
+				}
+				if (bSnapToTailB)
+				{
+					EndPointToVertexConnection[Connection.SourceEndPointIndex] = TailB;
+					continue;
+				}
+
+				FDynamicMesh3::FEdgeSplitInfo EdgeSplitInfo;
+				const UE::Geometry::EMeshResult SplitEdgeResult = DynamicMesh.SplitEdge(TailA, TailB, EdgeSplitInfo);
+				if (!ensure(SplitEdgeResult == UE::Geometry::EMeshResult::Ok))
+				{
+					continue;
+				}
+
+				DynamicMesh.SetVertex(EdgeSplitInfo.NewVertex, Connection.TargetPoint);
+				VertexScaleAttribute.SetValue(EdgeSplitInfo.NewVertex, Connection.TargetPointScale);
+				FillerVertexAttribute.SetValue(EdgeSplitInfo.NewVertex, false);
+
+				EndPointToVertexConnection[Connection.SourceEndPointIndex] = EdgeSplitInfo.NewVertex;
+				TailA = EdgeSplitInfo.NewVertex;
+			}
+		}
 	}
 
-	for (int32 i = 0; i < EndPoints.Num(); ++i)
+	// Step 4: Connect each endpoint to its attachment vertex
+	for (int32 i = 0; i < EndPointVertexIndices.Num(); ++i)
 	{
-		const int32 EndPointVertexIndex = EndPoints[i];
-		const int32 VertexToMergeWith = MergeMap[i];
+		const int32 EndPointVertexIndex = EndPointVertexIndices[i];
+		const int32 VertexToMergeWith = EndPointToVertexConnection[i];
 
 		if (VertexToMergeWith == INDEX_NONE)
 		{
 			continue;
 		}
-		const FVector VertexPosition = DynamicMesh.GetVertex(EndPointVertexIndex);
-		const FVector OtherVertexPosition = DynamicMesh.GetVertex(VertexToMergeWith);
-
-		const FVector Normal = (VertexPosition - OtherVertexPosition).GetSafeNormal();
-		const FVector NewVertexPosition = OtherVertexPosition + (Normal * 0.1); // We want to keep this vertex to not lose the radius information, so just move it really close and connect it to the other vertex
 
-		DynamicMesh.SetVertex(EndPointVertexIndex, NewVertexPosition);
-		
-		const TArray<int32> Neighbours = PV::DynamicMeshHelper::GetVertexNeighbours(DynamicMesh, EndPointVertexIndex);
-		for (int32 Neighbour : Neighbours)
-		{
-			if (FillerVertexAttribute.GetValue(Neighbour))
-			{
-				const TArray<int32> FillerNeighbours = PV::DynamicMeshHelper::GetVertexNeighbours(DynamicMesh, Neighbour);
-				const int32* AdjacentVertex = FillerNeighbours.FindByPredicate([&](const auto& Value) { return Value != EndPointVertexIndex && Value != Neighbour; });
-				check(AdjacentVertex != nullptr);
-				const FVector AdjacentVertexPosition = DynamicMesh.GetVertex(*AdjacentVertex);
-				const FVector NewFillerVertex = GenerateFillerVertex(AdjacentVertexPosition, NewVertexPosition);
-				
-				DynamicMesh.SetVertex(Neighbour, NewFillerVertex);
-			}
-		}
-		
-		const FVector NewFillerVertex = GenerateFillerVertex(NewVertexPosition, OtherVertexPosition);
-		const int32 NewFillerVertexIndex = DynamicMesh.AppendVertex(NewFillerVertex);
-		DynamicMesh.AppendTriangle(EndPointVertexIndex, VertexToMergeWith, NewFillerVertexIndex);
 		EndPointAttribute.SetValue(EndPointVertexIndex, false);
+		EndPointAttribute.SetValue(VertexToMergeWith, false);
+
+		FDynamicMesh3::FMergeVerticesOptions MergeOptions;
+		MergeOptions.bAllowNonBoundaryBowtieCreation = true; // VertexToMergeWith is interior (created by SplitEdge)
+
+		FDynamicMesh3::FMergeVerticesInfo MergeInfo;
+		const UE::Geometry::EMeshResult MergeResult = DynamicMesh.MergeVertices(
+			VertexToMergeWith,   // KeepVid: the attachment point on the branch skeleton
+			EndPointVertexIndex, // DiscardVid: the endpoint being welded in
+			0.0,                 // InterpolationT: 0 = keep KeepVid's position exactly
+			MergeOptions,
+			MergeInfo
+		);
+		ensure(MergeResult == UE::Geometry::EMeshResult::Ok);
 	}
 }
 
@@ -898,29 +1092,30 @@ PV::StaticMeshImport::EImportResult PV::StaticMeshImport::ImportGrowthDataFromSt
 		EndPointAttributeDefinition.CopyAttributeToCollection(DynamicMesh, DebugCollection);
 	}
 
-	ConnectBranchesToParent(
+	FindRootPoints(
 		DynamicMesh,
-		VertexScaleAttributeDefinition.GetOrAttachAttribute(DynamicMesh),
-		EndPointAttributeDefinition.GetOrAttachAttribute(DynamicMesh),
-		FillerVertexAttributeDefinition.GetOrAttachAttribute(DynamicMesh)
+		RootPointAttributeDefinition.GetOrAttachAttribute(DynamicMesh),
+		EndPointAttributeDefinition.GetOrAttachAttribute(DynamicMesh)
 	);
-	if (InParams.DebugState == EPVImportStaticMeshDebugState::ConnectBranchesToParent)
+	if (InParams.DebugState == EPVImportStaticMeshDebugState::FindRootPoints)
 	{
 		FManagedArrayCollection& DebugCollection = Output.DebugCollection.Emplace();
 		AddDynamicMeshToCollection(DynamicMesh, DebugCollection);
+		RootPointAttributeDefinition.CopyAttributeToCollection(DynamicMesh, DebugCollection);
 		EndPointAttributeDefinition.CopyAttributeToCollection(DynamicMesh, DebugCollection);
 	}
 
-	FindRootPoints(
+	ConnectBranchesToParent(
 		DynamicMesh,
 		RootPointAttributeDefinition.GetOrAttachAttribute(DynamicMesh),
-		EndPointAttributeDefinition.GetOrAttachAttribute(DynamicMesh)
+		VertexScaleAttributeDefinition.GetOrAttachAttribute(DynamicMesh),
+		EndPointAttributeDefinition.GetOrAttachAttribute(DynamicMesh),
+		FillerVertexAttributeDefinition.GetOrAttachAttribute(DynamicMesh)
 	);
-	if (InParams.DebugState == EPVImportStaticMeshDebugState::FindRootPoints)
+	if (InParams.DebugState == EPVImportStaticMeshDebugState::ConnectBranchesToParent)
 	{
 		FManagedArrayCollection& DebugCollection = Output.DebugCollection.Emplace();
 		AddDynamicMeshToCollection(DynamicMesh, DebugCollection);
-		RootPointAttributeDefinition.CopyAttributeToCollection(DynamicMesh, DebugCollection);
 		EndPointAttributeDefinition.CopyAttributeToCollection(DynamicMesh, DebugCollection);
 	}
 
