[PVE] Restrict Mesh property to StaticMesh and SkeletalMesh in Foliage palette

#Synced-CL 51579190

diff --git a/Engine/Plugins/Experimental/ProceduralVegetationEditor/Source/ProceduralVegetation/Private/Helpers/PVFoliageHelper.h b/Engine/Plugins/Experimental/ProceduralVegetationEditor/Source/ProceduralVegetation/Private/Helpers/PVFoliageHelper.h
--- a/Engine/Plugins/Experimental/ProceduralVegetationEditor/Source/ProceduralVegetation/Private/Helpers/PVFoliageHelper.h
+++ b/Engine/Plugins/Experimental/ProceduralVegetationEditor/Source/ProceduralVegetation/Private/Helpers/PVFoliageHelper.h
@@ -125,7 +125,7 @@ struct FPVFoliageInfo
 {
 	GENERATED_BODY()
 	
-	UPROPERTY(EditAnywhere, Category = "Info")
+	UPROPERTY(EditAnywhere, Category = "Info", meta = (AllowedClasses = "/Script/Engine.StaticMesh,/Script/Engine.SkeletalMesh"))
 	TSoftObjectPtr<UObject> Mesh;
  
 	UPROPERTY(EditAnywhere, Category="Info", meta=(ShowInnerProperties, FullyExpand="true", Tooltip=""))
