[PVE] Refactor hormone based settings code for re-usability across grafter and foliage distributor

#Synced-CL 51591271

diff --git a/Engine/Plugins/Experimental/ProceduralVegetationEditor/Source/ProceduralVegetation/Private/Helpers/PVFoliageHelper.h b/Engine/Plugins/Experimental/ProceduralVegetationEditor/Source/ProceduralVegetation/Private/Helpers/PVFoliageHelper.h
--- a/Engine/Plugins/Experimental/ProceduralVegetationEditor/Source/ProceduralVegetation/Private/Helpers/PVFoliageHelper.h
+++ b/Engine/Plugins/Experimental/ProceduralVegetationEditor/Source/ProceduralVegetation/Private/Helpers/PVFoliageHelper.h
@@ -247,3 +247,24 @@ namespace PV::DistributionConditionUtils
 			FPVFoliageDistributionConditionSettings& OutSettings);
 	}
 }
+
+namespace PV::DistributionHelper
+{
+	struct FAttachmentPoint
+	{
+		int32 BranchIndex;
+		int32 PointIndexA;
+		int32 PointIndexB;
+		float PointAlpha;
+		FVector3f Position;
+		FVector3f UpDirection;
+		FVector3f NormalDirection;
+
+		float Scale;
+		float LengthFromRoot;
+		bool bIsTipInstance;
+		int32 TipPointIndex;
+
+		DistributionConditionUtils::ConditionPicker::FPointConditionSample ConditionSample;
+	};
+}
\ No newline at end of file
diff --git a/Engine/Plugins/Experimental/ProceduralVegetationEditor/Source/ProceduralVegetation/Private/Helpers/PVHormoneDistributionHelper.cpp b/Engine/Plugins/Experimental/ProceduralVegetationEditor/Source/ProceduralVegetation/Private/Helpers/PVHormoneDistributionHelper.cpp
new file mode 100644
--- /dev/null
+++ b/Engine/Plugins/Experimental/ProceduralVegetationEditor/Source/ProceduralVegetation/Private/Helpers/PVHormoneDistributionHelper.cpp
@@ -0,0 +1,275 @@
+// Copyright Epic Games, Inc. All Rights Reserved.
+
+#include "PVHormoneDistributionHelper.h"
+
+void PV::HormoneDistributionHelper::ComputeHormoneBasedAttachmentPoints(
+	const FManagedArrayCollection& OutCollection,
+	const DistributionConditionUtils::FNormalizedPointCaches& NormalizedPointCaches,
+	const FHormoneSettings::FDistributionSettings& DistributionSettings,
+	const FHormoneSettings::FScaleSettings& ScaleSettings,
+	const FHormoneSettings::FVectorSettings& VectorSettings,
+	const FHormoneSettings::FPhyllotaxySettings& PhyllotaxySettings,
+	FRandomStream& RandomStream, TArray<DistributionHelper::FAttachmentPoint>& OutAttachmentPoints)
+{
+	auto BranchPointsAttribute = PV::FBranchPointsAttribute::GetAttribute(OutCollection);
+	auto BudDirectionAttribute = PV::FBudDirectionAttribute::GetAttribute(OutCollection);
+	auto BranchGradientsAttribute = PV::FPointBranchGradientAttribute::GetAttribute(OutCollection);
+	auto BudHormoneLevels = PV::FBudHormoneLevelsAttribute::GetAttribute(OutCollection);
+	auto PointLengthFromRootAttribute = PV::FPointLengthFromRootAttribute::GetAttribute(OutCollection);
+	auto PointScaleAttribute = PV::FPointScaleAttribute::GetAttribute(OutCollection);
+	auto PointPositionAttribute = PV::FPointPositionAttribute::GetAttribute(OutCollection);
+
+	const int32 NumBranches = BranchPointsAttribute.Num();
+
+	float PhyllotaxyFormationLeaf = 180;
+	float MinBudsLeaf = 1;
+	float MaxBudsLeaf = 1;
+	bool ResetPhyllotaxyLeaf = PhyllotaxySettings.ResetPhyllotaxy;
+	float PhyllotaxyOffsetLeaf = PhyllotaxySettings.PhyllotaxyOffset;
+
+	switch (PhyllotaxySettings.PhyllotaxyType)
+	{
+	case EPhyllotaxyType::Alternate: PhyllotaxyFormationLeaf = 180 + PhyllotaxySettings.PhyllotaxyAdditionalAngle;
+		MinBudsLeaf = MaxBudsLeaf = 1;
+		break;
+	case EPhyllotaxyType::Opposite: PhyllotaxyFormationLeaf = PhyllotaxySettings.PhyllotaxyAdditionalAngle;
+		MinBudsLeaf = MaxBudsLeaf = 2;
+		break;
+	case EPhyllotaxyType::Decussate: PhyllotaxyFormationLeaf = 90 + PhyllotaxySettings.PhyllotaxyAdditionalAngle;
+		MinBudsLeaf = MaxBudsLeaf = 2;
+		break;
+	case EPhyllotaxyType::Whorled: PhyllotaxyFormationLeaf = 90 + PhyllotaxySettings.PhyllotaxyAdditionalAngle;
+		MinBudsLeaf = PhyllotaxySettings.MinimumNodeBuds;
+		MaxBudsLeaf = PhyllotaxySettings.MaximumNodeBuds;
+		break;
+	case EPhyllotaxyType::Spiral:
+		{
+			// These leaf arrangement configurations where Distichous is two leaves, Tristichous is three etc.
+			switch (PhyllotaxySettings.PhyllotaxyFormation)
+			{
+			case EPhyllotaxyFormation::Distichous: PhyllotaxyFormationLeaf = 180 + PhyllotaxySettings.PhyllotaxyAdditionalAngle;
+				break;
+			case EPhyllotaxyFormation::Tristichous: PhyllotaxyFormationLeaf = 120 + PhyllotaxySettings.PhyllotaxyAdditionalAngle;
+				break;
+			case EPhyllotaxyFormation::Pentastichous: PhyllotaxyFormationLeaf = 144 + PhyllotaxySettings.PhyllotaxyAdditionalAngle;
+				break;
+			case EPhyllotaxyFormation::Octastichous: PhyllotaxyFormationLeaf = 135 + PhyllotaxySettings.PhyllotaxyAdditionalAngle;
+				break;
+			case EPhyllotaxyFormation::Parastichous: PhyllotaxyFormationLeaf = PhyllotaxySettings.PhyllotaxyAdditionalAngle;
+				break;
+			}
+
+			MinBudsLeaf = MaxBudsLeaf = 1;
+			break;
+		}
+	}
+
+	for (int32 BranchIndex = 0; BranchIndex < NumBranches; ++BranchIndex)
+	{
+		const TArray<int32>& BranchPoints = BranchPointsAttribute[BranchIndex];
+
+		if (BranchPoints.Num() < 2) continue;
+
+		const float FirstPointLengthFromRoot = PointLengthFromRootAttribute[BranchPoints[0]];
+		const float LastPointLengthFromRoot = PointLengthFromRootAttribute[BranchPoints.Last()];
+		const float BranchLength = LastPointLengthFromRoot - FirstPointLengthFromRoot;
+
+		if (BranchLength <= KINDA_SMALL_NUMBER) continue;
+
+		// Calculate Normalized Attachment Lookups
+		TArray<float> NormalizedAttachmentPoints;
+		float Increment = DistributionSettings.InstanceSpacing / BranchLength;
+		const float AdjustedMaxPerBranch = DistributionSettings.MaxPerBranch == -1
+			? BranchPoints.Num()
+			: DistributionSettings.MaxPerBranch;
+		float LookUp = 0.0f;
+
+		for (int32 i = 0; i < AdjustedMaxPerBranch; ++i)
+		{
+			if (LookUp > 1.0f) break;
+
+			float InstanceSpacingRampValue = DistributionSettings.InstanceSpacingRamp
+				? DistributionSettings.InstanceSpacingRamp->GetRichCurveConst()->Eval(1.0f - LookUp)
+				: Increment;
+
+			// The increment is nudged slightly before blending with the weighted ramp value, the number is chosen based on observation
+			float AdjustedIncrement = Increment + (InstanceSpacingRampValue * (Increment * 10.0f));
+			AdjustedIncrement = FMath::Lerp(Increment, AdjustedIncrement, DistributionSettings.InstanceSpacingRampEffect);
+
+			NormalizedAttachmentPoints.Add(1.0f - LookUp);
+			LookUp += AdjustedIncrement;
+		}
+
+		// Compute Position, Angles, Scale and LengthFromRoot for attachment points
+
+		// Compute based on first point attributes
+		const FVector3f InitialApicalDirection(BudDirectionAttribute[BranchPoints[0]].Apical);
+		FVector3f InitialAxillaryDirection(BudDirectionAttribute[BranchPoints[0]].Axillary);
+		const FVector3f InitialUpVector(BudDirectionAttribute[BranchPoints[0]].UpVector);
+
+		if (ResetPhyllotaxyLeaf)
+		{
+			InitialAxillaryDirection = FVector3f::CrossProduct(InitialApicalDirection, InitialUpVector);
+		}
+
+		if (PhyllotaxyOffsetLeaf > 0.01f)
+		{
+			const FQuat4f RotationQuat(InitialApicalDirection.GetSafeNormal(), FMath::DegreesToRadians(PhyllotaxyOffsetLeaf));
+			InitialAxillaryDirection = RotationQuat.RotateVector(InitialAxillaryDirection);
+		}
+
+		int32 AxillaryRotationIteration = 0;
+
+		for (int32 Idx = 0; Idx < NormalizedAttachmentPoints.Num(); ++Idx)
+		{
+			for (int32 BranchPointIndex = 0; BranchPointIndex < BranchPoints.Num() - 1; ++BranchPointIndex)
+			{
+				const int32 CurrentPointIndex = BranchPoints[BranchPointIndex];
+				const int32 NextPointIndex = BranchPoints[BranchPointIndex + 1];
+
+				const float CurrentPointLengthFromRoot = PointLengthFromRootAttribute[CurrentPointIndex];
+				const float NextPointLengthFromRoot = PointLengthFromRootAttribute[NextPointIndex];
+
+				const float CurrentPointLfrNormalized = FMath::GetMappedRangeValueClamped(
+					FVector2f(FirstPointLengthFromRoot, LastPointLengthFromRoot), FVector2f(0.0f, 1.0f), CurrentPointLengthFromRoot);
+				const float NextPointLfrNormalized = FMath::GetMappedRangeValueClamped(
+					FVector2f(FirstPointLengthFromRoot, LastPointLengthFromRoot), FVector2f(0.0f, 1.0f), NextPointLengthFromRoot);
+
+				if (NormalizedAttachmentPoints[Idx] >= CurrentPointLfrNormalized && NormalizedAttachmentPoints[Idx] <= NextPointLfrNormalized)
+				{
+					const float DistanceAlpha = (NextPointLfrNormalized - CurrentPointLfrNormalized) < KINDA_SMALL_NUMBER
+						? 1.0f
+						: (NormalizedAttachmentPoints[Idx] - CurrentPointLfrNormalized) / (NextPointLfrNormalized - CurrentPointLfrNormalized);
+
+					const float EthyleneLevel = FMath::Lerp(BudHormoneLevels[CurrentPointIndex].Ethylene, BudHormoneLevels[NextPointIndex].Ethylene,
+						DistanceAlpha);
+
+					if ((EthyleneLevel - 0.001f) >= DistributionSettings.EthyleneThreshold) continue;
+
+					// Shared Properties for this segment
+					const FVector3f Position = FMath::Lerp(PointPositionAttribute[CurrentPointIndex], PointPositionAttribute[NextPointIndex],
+						DistanceAlpha);
+					const float AttachmentPointLfr = FMath::GetMappedRangeValueClamped(
+						FVector2f(0.0f, 1.0f), FVector2f(FirstPointLengthFromRoot, LastPointLengthFromRoot), NormalizedAttachmentPoints[Idx]);
+
+					// Scale Calculation
+					const float CurrentPointBranchGradient = BranchGradientsAttribute[CurrentPointIndex];
+					const float NextPointBranchGradient = BranchGradientsAttribute[NextPointIndex];
+					const float BranchGradient = FMath::Clamp(FMath::Lerp(CurrentPointBranchGradient, NextPointBranchGradient, DistanceAlpha), 0.0f,
+						1.0f);
+					float BranchGradientRampValue = ScaleSettings.ScaleRamp
+						? ScaleSettings.ScaleRamp->GetRichCurveConst()->Eval(1.0f - BranchGradient)
+						: 1.0f;
+
+					const float BranchScale = FMath::Lerp(PointScaleAttribute[CurrentPointIndex], PointScaleAttribute[NextPointIndex], DistanceAlpha);
+					const float BranchScaledValue = BranchGradientRampValue * (BranchScale + 0.1f);
+					BranchGradientRampValue = FMath::Clamp(FMath::Lerp(BranchGradientRampValue, BranchScaledValue, ScaleSettings.BranchScaleImpact),
+						0.0f, 1.0f);
+
+					float Scale = FMath::GetMappedRangeValueClamped(FVector2f(0.0f, 1.0f), FVector2f(ScaleSettings.MinScale, ScaleSettings.MaxScale),
+						BranchGradientRampValue) * ScaleSettings.BaseScale;
+					Scale *= RandomStream.FRandRange(ScaleSettings.RandomScaleMin, ScaleSettings.RandomScaleMax);
+
+					// Conditions
+					const FVector3f BudApicalDirectionVector = FMath::Lerp(BudDirectionAttribute[CurrentPointIndex].Apical,
+						BudDirectionAttribute[NextPointIndex].Apical, DistanceAlpha);
+
+					DistributionConditionUtils::ConditionPicker::FPointConditionSample ConditionSample;
+					ConditionSample.UpAlignment = DistributionConditionUtils::UpAlignmentNormalized(BudApicalDirectionVector);
+					ConditionSample.Light = FMath::Lerp(NormalizedPointCaches.Light[CurrentPointIndex], NormalizedPointCaches.Light[NextPointIndex],
+						DistanceAlpha);
+					ConditionSample.Health = FMath::Lerp(NormalizedPointCaches.Health[CurrentPointIndex],
+						NormalizedPointCaches.Health[NextPointIndex], DistanceAlpha);
+					ConditionSample.Scale = FMath::Lerp(NormalizedPointCaches.PScales[CurrentPointIndex],
+						NormalizedPointCaches.PScales[NextPointIndex], DistanceAlpha);
+
+					const bool bIsTip = (BranchPointIndex == BranchPoints.Num() - 2) && FMath::IsNearlyEqual(DistanceAlpha, 1.0f, 1e-3f);
+					ConditionSample.Tip = DistributionConditionUtils::Tip(bIsTip);
+
+					// Tip detection logic shared by Grafter
+					bool bIsTipInstance = Idx == 0 && FMath::IsNearlyEqual(AttachmentPointLfr, LastPointLengthFromRoot, 0.001f);
+					int32 TipPointIndex = bIsTipInstance
+						? NextPointIndex
+						: INDEX_NONE;
+
+					// Determine Rotation / Multi-Buds
+					int32 NumBuds = 1;
+					TArray<FVector3f> ComputedUpVectors;
+					TArray<FVector3f> ComputedNormalVectors;
+
+					const FVector3f BudLightOptimalVector = FMath::Lerp(BudDirectionAttribute[CurrentPointIndex].LightOptimal,
+						BudDirectionAttribute[NextPointIndex].LightOptimal, DistanceAlpha);
+					const FQuat4f QuatRotation = FQuat4f::FindBetweenNormals(BudApicalDirectionVector.GetSafeNormal(),
+						BudLightOptimalVector.GetSafeNormal());
+					const FVector3f BaseNormalVector = QuatRotation.RotateVector(BudApicalDirectionVector);
+
+					if (Idx == 0)
+					{
+						ComputedUpVectors.Add(BudApicalDirectionVector);
+						ComputedNormalVectors.Add(BaseNormalVector);
+					}
+					else [[likely]]
+					{
+						float AxillaryRotationDegree = FMath::DegreesToRadians(PhyllotaxyFormationLeaf * AxillaryRotationIteration);
+						const FQuat4f RotationQuat(InitialApicalDirection.GetSafeNormal(), AxillaryRotationDegree);
+						FVector3f NewAxillaryDirection = RotationQuat.RotateVector(InitialAxillaryDirection.GetSafeNormal());
+						FQuat4f ApicalCorrection = FQuat4f::FindBetweenNormals(InitialApicalDirection.GetSafeNormal(),
+							BudApicalDirectionVector.GetSafeNormal());
+						NewAxillaryDirection = ApicalCorrection.RotateVector(NewAxillaryDirection).GetSafeNormal();
+
+						NumBuds = RandomStream.RandRange(MinBudsLeaf, MaxBudsLeaf);
+						float RotationPerBud = 360.0f / static_cast<float>(NumBuds);
+
+						for (int32 i = 0; i < NumBuds; ++i)
+						{
+							float RotationAngle = FMath::DegreesToRadians(RotationPerBud * i);
+							const FQuat4f AxillaryRotation(InitialApicalDirection.GetSafeNormal(), RotationAngle);
+							FVector3f AxillaryDirectionModified = AxillaryRotation.RotateVector(NewAxillaryDirection.GetSafeNormal());
+
+							if (VectorSettings.OverrideAxilAngle)
+							{
+								float AxilAngleRampValue = VectorSettings.AxilAngleRamp
+									? FMath::Clamp(VectorSettings.AxilAngleRamp->GetRichCurveConst()->Eval(1.0f - BranchGradient), 0.0f, 1.0f)
+									: 1.0f;
+								AxilAngleRampValue *= VectorSettings.AxilAngleRampUpperValue;
+								const float AxilAngleBlended = FMath::Lerp(VectorSettings.AxilAngle, AxilAngleRampValue,
+									VectorSettings.AxilAngleRampEffect);
+
+								const FVector3f RotationAxis = FVector3f::CrossProduct(BudApicalDirectionVector.GetSafeNormal(),
+									AxillaryDirectionModified.GetSafeNormal());
+								const FQuat4f Rotation(RotationAxis.GetSafeNormal(), FMath::DegreesToRadians(AxilAngleBlended));
+								ComputedUpVectors.Add(Rotation.RotateVector(BudApicalDirectionVector.GetSafeNormal()));
+							}
+							else
+							{
+								ComputedUpVectors.Add(AxillaryDirectionModified);
+							}
+							ComputedNormalVectors.Add(BaseNormalVector);
+						}
+						AxillaryRotationIteration++;
+					}
+
+					// Create FAttachmentPoint for each generated bud
+					for (int32 BudIdx = 0; BudIdx < NumBuds; ++BudIdx)
+					{
+						DistributionHelper::FAttachmentPoint Point;
+						Point.BranchIndex = BranchIndex;
+						Point.PointIndexA = CurrentPointIndex;
+						Point.PointIndexB = NextPointIndex;
+						Point.PointAlpha = DistanceAlpha;
+						Point.Position = Position;
+						Point.UpDirection = ComputedUpVectors[BudIdx];
+						Point.NormalDirection = ComputedNormalVectors[BudIdx];
+						Point.Scale = Scale;
+						Point.LengthFromRoot = AttachmentPointLfr;
+						Point.bIsTipInstance = bIsTipInstance;
+						Point.TipPointIndex = TipPointIndex;
+						Point.ConditionSample = ConditionSample;
+
+						OutAttachmentPoints.Add(Point);
+					}
+				}
+			}
+		}
+	}
+}
diff --git a/Engine/Plugins/Experimental/ProceduralVegetationEditor/Source/ProceduralVegetation/Private/Helpers/PVHormoneDistributionHelper.h b/Engine/Plugins/Experimental/ProceduralVegetationEditor/Source/ProceduralVegetation/Private/Helpers/PVHormoneDistributionHelper.h
new file mode 100644
--- /dev/null
+++ b/Engine/Plugins/Experimental/ProceduralVegetationEditor/Source/ProceduralVegetation/Private/Helpers/PVHormoneDistributionHelper.h
@@ -0,0 +1,65 @@
+// Copyright Epic Games, Inc. All Rights Reserved.
+
+#pragma once
+
+#include "CoreMinimal.h"
+#include "PVFoliageHelper.h"
+#include "PVPhyllotaxyHelper.h"
+#include "Utils/PVAttributes.h"
+#include "Utils/PVFloatRamp.h"
+
+namespace PV::HormoneDistributionHelper
+{
+	namespace FHormoneSettings
+	{
+		struct FDistributionSettings
+		{
+			const float EthyleneThreshold;
+			const float InstanceSpacing;
+			const FPVFloatRamp* InstanceSpacingRamp;
+			const float InstanceSpacingRampEffect;
+			const int32 MaxPerBranch;
+		};
+
+		struct FScaleSettings
+		{
+			const float BaseScale;
+			const float BranchScaleImpact;
+			const float MinScale;
+			const float MaxScale;
+			const float RandomScaleMin;
+			const float RandomScaleMax;
+			const FPVFloatRamp* ScaleRamp;
+		};
+
+		struct FVectorSettings
+		{
+			const bool OverrideAxilAngle;
+			const float AxilAngle;
+			const FPVFloatRamp* AxilAngleRamp;
+			const float AxilAngleRampUpperValue;
+			const float AxilAngleRampEffect;
+		};
+
+		struct FPhyllotaxySettings
+		{
+			const EPhyllotaxyType PhyllotaxyType;
+			const EPhyllotaxyFormation PhyllotaxyFormation;
+			const int32 MinimumNodeBuds;
+			const int32 MaximumNodeBuds;
+			const float PhyllotaxyAdditionalAngle;
+			const bool ResetPhyllotaxy;
+			const float PhyllotaxyOffset;
+		};
+	}
+	
+	void ComputeHormoneBasedAttachmentPoints(
+		const FManagedArrayCollection& OutCollection,
+		const DistributionConditionUtils::FNormalizedPointCaches& NormalizedPointCaches,
+		const FHormoneSettings::FDistributionSettings& DistributionSettings,
+		const FHormoneSettings::FScaleSettings& ScaleSettings,
+		const FHormoneSettings::FVectorSettings& VectorSettings,
+		const FHormoneSettings::FPhyllotaxySettings& PhyllotaxySettings,
+		FRandomStream& RandomStream,
+		TArray<DistributionHelper::FAttachmentPoint>& OutAttachmentPoints);
+}
\ No newline at end of file
diff --git a/Engine/Plugins/Experimental/ProceduralVegetationEditor/Source/ProceduralVegetation/Private/Helpers/PVParametricDistributionHelper.cpp b/Engine/Plugins/Experimental/ProceduralVegetationEditor/Source/ProceduralVegetation/Private/Helpers/PVParametricDistributionHelper.cpp
--- a/Engine/Plugins/Experimental/ProceduralVegetationEditor/Source/ProceduralVegetation/Private/Helpers/PVParametricDistributionHelper.cpp
+++ b/Engine/Plugins/Experimental/ProceduralVegetationEditor/Source/ProceduralVegetation/Private/Helpers/PVParametricDistributionHelper.cpp
@@ -6,13 +6,14 @@
 #include "Facades/PVBranchFacade.h"
 #include "Helpers/PVUtilities.h"
 
-void PV::ParametricDistributionHelper::ComputeAttachmentPoints(FManagedArrayCollection& OutCollection,
-                                                               const FNormalizedPointCaches& NormalizedPointCaches,
-                                                               const FParametricSettings::FSpacingSettings& SpacingSettings,
-                                                               const FParametricSettings::FAngleSettings& AngleSettings,
-                                                               const FParametricSettings::FScaleSettings& ScaleSettings,
-                                                               FRandomStream& RandomStream,
-                                                               TArray<FAttachmentPoint>& OutAttachmentPoints)
+void PV::ParametricDistributionHelper::ComputeAttachmentPoints(
+	FManagedArrayCollection& OutCollection,
+	const FNormalizedPointCaches& NormalizedPointCaches,
+	const FParametricSettings::FSpacingSettings& SpacingSettings,
+	const FParametricSettings::FAngleSettings& AngleSettings,
+	const FParametricSettings::FScaleSettings& ScaleSettings,
+	FRandomStream& RandomStream,
+	TArray<DistributionHelper::FAttachmentPoint>& OutAttachmentPoints)
 {
 	auto PointPlantGradientAttribute = FPointPlantGradientAttribute::GetAttribute(OutCollection);
 	auto PointBranchGradientAttribute = FPointBranchGradientAttribute::GetAttribute(OutCollection);
@@ -46,11 +47,11 @@ void PV::ParametricDistributionHelper::ComputeAttachmentPoints(FManagedArrayColl
 			BranchNumberAttribute,
 			BranchChildrenAttribute
 		});
-	
+
 	TMap<int32, float> BranchIndicesToBranchLengths = ComputeBranchLengths(PointLengthFromRootAttribute, BranchPointsAttribute);
 	float MaxBranchLength = -1.0f;
 	float AverageBranchLength = 0.0f;
-	
+
 	for (const auto IndexLengthPair : BranchIndicesToBranchLengths)
 	{
 		const float BranchLength = IndexLengthPair.Value;
@@ -66,7 +67,7 @@ void PV::ParametricDistributionHelper::ComputeAttachmentPoints(FManagedArrayColl
 
 	AverageBranchLength = AverageBranchLength / BranchIndicesToBranchLengths.Num();
 	AverageBranchLength = FMath::Lerp(AverageBranchLength, MaxBranchLength, 0.33f);
-	
+
 	for (int32 BranchIndex = 0; BranchIndex < NumOfBranches; BranchIndex++)
 	{
 		// Check generation limit
@@ -274,7 +275,7 @@ void PV::ParametricDistributionHelper::ComputeAttachmentPoints(FManagedArrayColl
 				Sample.Health = 0.0f;
 				Sample.Scale = 0.0f;
 			}
-			
+
 			OutAttachmentPoints.Add({
 				.BranchIndex = BranchIndex,
 				.PointIndexA = PointIndexA,
diff --git a/Engine/Plugins/Experimental/ProceduralVegetationEditor/Source/ProceduralVegetation/Private/Helpers/PVParametricDistributionHelper.h b/Engine/Plugins/Experimental/ProceduralVegetationEditor/Source/ProceduralVegetation/Private/Helpers/PVParametricDistributionHelper.h
--- a/Engine/Plugins/Experimental/ProceduralVegetationEditor/Source/ProceduralVegetation/Private/Helpers/PVParametricDistributionHelper.h
+++ b/Engine/Plugins/Experimental/ProceduralVegetationEditor/Source/ProceduralVegetation/Private/Helpers/PVParametricDistributionHelper.h
@@ -52,24 +52,6 @@ namespace PV::ParametricDistributionHelper
 		};
 	};
 
-	struct FAttachmentPoint
-	{
-		int32 BranchIndex;
-		int32 PointIndexA;
-		int32 PointIndexB;
-		float PointAlpha;
-		FVector3f Position;
-		FVector3f UpDirection;
-		FVector3f NormalDirection;
-
-		float Scale;
-		float LengthFromRoot;
-		bool bIsTipInstance;
-		int32 TipPointIndex;
-
-		ConditionPicker::FPointConditionSample ConditionSample;
-	};
-
 	void ComputeAttachmentPoints(
 		FManagedArrayCollection& OutCollection,
 		const FNormalizedPointCaches& NormalizedPointCaches,
@@ -77,7 +59,7 @@ namespace PV::ParametricDistributionHelper
 		const FParametricSettings::FAngleSettings& AngleSettings,
 		const FParametricSettings::FScaleSettings& ScaleSettings,
 		FRandomStream& RandomStream,
-		TArray<FAttachmentPoint>& OutAttachmentPoints);
+		TArray<DistributionHelper::FAttachmentPoint>& OutAttachmentPoints);
 
 	TMap<int32, float> ComputeBranchLengths(
 		FPointLengthFromRootAttributeConstView PointLengthFromRootAttribute,
diff --git a/Engine/Plugins/Experimental/ProceduralVegetationEditor/Source/ProceduralVegetation/Private/Implementations/PVFoliage.cpp b/Engine/Plugins/Experimental/ProceduralVegetationEditor/Source/ProceduralVegetation/Private/Implementations/PVFoliage.cpp
--- a/Engine/Plugins/Experimental/ProceduralVegetationEditor/Source/ProceduralVegetation/Private/Implementations/PVFoliage.cpp
+++ b/Engine/Plugins/Experimental/ProceduralVegetationEditor/Source/ProceduralVegetation/Private/Implementations/PVFoliage.cpp
@@ -1,7 +1,7 @@
 // Copyright Epic Games, Inc. All Rights Reserved.
- 
+
 #include "PVFoliage.h"
- 
+
 #include "ProceduralVegetationModule.h"
 #include "Utils/PVFloatRamp.h"
 #include "Algo/RandomShuffle.h"
@@ -20,23 +20,23 @@ using PV::DistributionConditionUtils::FConditionAttributes;
 void FPVFoliage::DistributeFoliageWithHormoneBasedSettings(
 	FManagedArrayCollection& OutCollection,
 	const FManagedArrayCollection& FoliageCollection,
-	const PV::HormoneBasedSettings::FDistributionSettings& DistributionSettings,
-	const PV::HormoneBasedSettings::FScaleSettings& ScaleSettings,
-	const PV::HormoneBasedSettings::FVectorSettings& VectorSettings,
-	const PV::HormoneBasedSettings::FPhyllotaxySettings& PhyllotaxySettings,
+	const FHormoneSettings::FDistributionSettings& DistributionSettings,
+	const FHormoneSettings::FScaleSettings& ScaleSettings,
+	const FHormoneSettings::FVectorSettings& VectorSettings,
+	const FHormoneSettings::FPhyllotaxySettings& PhyllotaxySettings,
 	const int32 RandomSeed,
 	const FPVFoliageDistributionConditions& DistributionConditions)
 {
 	TRACE_CPUPROFILER_EVENT_SCOPE(PVFoliage::DistributeFoliageWithHormoneBasedSettings);
-	
+
 	if (!(DistributionSettings.InstanceSpacingRamp && ScaleSettings.ScaleRamp && VectorSettings.AxilAngleRamp))
 	{
 		return;
 	}
-	
+
 	PV::Facades::FFoliageFacade FoliageFacadeOutput(OutCollection);
 	PV::Facades::FFoliageFacade FoliagePaletteFacade(FoliageCollection);
- 
+
 	if (const int32 NumFoliageMeshes = FoliagePaletteFacade.NumFoliageInfo(); NumFoliageMeshes == 0)
 	{
 		UE_LOG(LogProceduralVegetation, Warning, TEXT("There are no foliage meshes available in input for distribution"));
@@ -44,560 +44,107 @@ void FPVFoliage::DistributeFoliageWithHormoneBasedSettings(
 	}
 
 	FoliageFacadeOutput.SetFoliageInfos(FoliagePaletteFacade.GetFoliageInfos());
- 
+
 	FRandomStream RandomStream;
 	RandomStream.Initialize(RandomSeed);
- 
-	PV::Facades::FBranchFacade BranchFacadeSource(OutCollection);
-	PV::Facades::FPointFacade PointFacadeSource(OutCollection);
-
-	TManagedArrayAccessor<TArray<FVector3f>> BudDirectionAttribute(OutCollection, PV::AttributeNames::BudDirection,
-			PV::GroupNames::PointGroup);
-	
-	TManagedArrayAccessor<float> BranchGradientsAttribute(OutCollection, PV::AttributeNames::BranchGradient,
-		PV::GroupNames::PointGroup);
-	
-	TManagedArrayAccessor<TArray<float>> LeafPhyllotaxyAttribute(OutCollection, PV::AttributeNames::LeafPhyllotaxy,
-		PV::GroupNames::DetailsGroup);
-	
-	TManagedArrayAccessor<TArray<float>> BudHormoneLevels(OutCollection, PV::AttributeNames::BudHormoneLevels,
-		PV::GroupNames::PointGroup);
- 
-	TManagedArrayAccessor<TArray<int>> BudDevelopment(OutCollection, PV::AttributeNames::BudDevelopment,
-		PV::GroupNames::PointGroup);
-	
-	float PhyllotaxyFormationLeaf = 180;
-	float MinBudsLeaf = 1;
-	float MaxBudsLeaf = 1;
-	int32 ResetPhyllotaxyLeaf = 0;
-	float PhyllotaxyOffsetLeaf = 0;
- 
-	if (PhyllotaxySettings.OverridePhyllotaxy)
+
+	FNormalizedPointCaches NormalizedPointCaches;
+	FBuildOptions BuildOptions;
+	BuildOptions.bComputePScales = DistributionConditions.IsActiveCondition(EFoliageDistributionCondition::Scale);
+	BuildOptions.bComputeLight = DistributionConditions.IsActiveCondition(EFoliageDistributionCondition::Light);
+	BuildOptions.bComputeHealth = DistributionConditions.IsActiveCondition(EFoliageDistributionCondition::Health);
+
+	BuildNormalizedPointCaches(
+		OutCollection,
+		BuildOptions,
+		NormalizedPointCaches);
+
+	OutCollection.EmptyGroup(PV::GroupNames::FoliageGroup);
+
+	// Collect non-empty mesh name indices
+	const TArray<FString>& FoliageMeshNames = FoliagePaletteFacade.GetFoliageNames();
+	TArray<int32> NonEmptyFoliageMeshNameIndices;
+	NonEmptyFoliageMeshNameIndices.Reserve(FoliageMeshNames.Num());
+
+	for (int32 i = 0; i < FoliageMeshNames.Num(); ++i)
 	{
-		switch (PhyllotaxySettings.PhyllotaxyType)
+		if (!FoliageMeshNames[i].IsEmpty())
 		{
-		case EPhyllotaxyType::Alternate: PhyllotaxyFormationLeaf = 180 + PhyllotaxySettings.PhyllotaxyAdditionalAngle;
-			MinBudsLeaf = MaxBudsLeaf = 1;
-			break;
-		case EPhyllotaxyType::Opposite: PhyllotaxyFormationLeaf = PhyllotaxySettings.PhyllotaxyAdditionalAngle;
-			MinBudsLeaf = MaxBudsLeaf = 2;
-			break;
-		case EPhyllotaxyType::Decussate: PhyllotaxyFormationLeaf = 90 + PhyllotaxySettings.PhyllotaxyAdditionalAngle;
-			MinBudsLeaf = MaxBudsLeaf = 2;
-			break;
-		case EPhyllotaxyType::Whorled: PhyllotaxyFormationLeaf = 90 + PhyllotaxySettings.PhyllotaxyAdditionalAngle;
-			MinBudsLeaf = PhyllotaxySettings.MinimumNodeBuds;
-			MaxBudsLeaf = PhyllotaxySettings.MaximumNodeBuds;
-			break;
-		case EPhyllotaxyType::Spiral:
-			{
-				// These leaf arrangement configurations where Distichous is two leaves, Tristichous is three etc.
-				switch (PhyllotaxySettings.PhyllotaxyFormation)
-				{
-				case EPhyllotaxyFormation::Distichous: PhyllotaxyFormationLeaf = 180 + PhyllotaxySettings.PhyllotaxyAdditionalAngle;
-					break;
-				case EPhyllotaxyFormation::Tristichous: PhyllotaxyFormationLeaf = 120 + PhyllotaxySettings.PhyllotaxyAdditionalAngle;
-					break;
-				case EPhyllotaxyFormation::Pentastichous: PhyllotaxyFormationLeaf = 144 + PhyllotaxySettings.PhyllotaxyAdditionalAngle;
-					break;
-				case EPhyllotaxyFormation::Octastichous: PhyllotaxyFormationLeaf = 135 + PhyllotaxySettings.PhyllotaxyAdditionalAngle;
-					break;
-				case EPhyllotaxyFormation::Parastichous: PhyllotaxyFormationLeaf = PhyllotaxySettings.PhyllotaxyAdditionalAngle;
-					break;
-				}
- 
-				MinBudsLeaf = MaxBudsLeaf = 1;
-				break;
-			}
+			NonEmptyFoliageMeshNameIndices.Add(i);
 		}
 	}
-	else
+
+	if (NonEmptyFoliageMeshNameIndices.Num() == 0)
 	{
-		PhyllotaxyFormationLeaf = LeafPhyllotaxyAttribute[0][1];
-		MinBudsLeaf = LeafPhyllotaxyAttribute[0][3];
-		MaxBudsLeaf = LeafPhyllotaxyAttribute[0][4];
-		ResetPhyllotaxyLeaf = static_cast<int>(LeafPhyllotaxyAttribute[0][6]);
-		PhyllotaxyOffsetLeaf = LeafPhyllotaxyAttribute[0][7];
+		return;
 	}
- 
-	const int32 NumBranches = BranchFacadeSource.GetElementCount();
- 
-	if (DistributionSettings.OverrideDistribution)
-	{
-		TRACE_CPUPROFILER_EVENT_SCOPE(PVFoliageDistributorNode::OverrideDistribution);
-		
-		FNormalizedPointCaches NormalizedPointCaches;
-		FBuildOptions BuildOptions;
-		BuildOptions.bComputePScales = DistributionConditions.IsActiveCondition(EFoliageDistributionCondition::Scale);
-		BuildOptions.bComputeLight = DistributionConditions.IsActiveCondition(EFoliageDistributionCondition::Light);
-		BuildOptions.bComputeHealth = DistributionConditions.IsActiveCondition(EFoliageDistributionCondition::Health);
-
-		BuildNormalizedPointCaches(
-			OutCollection,
-			BuildOptions,
-			NormalizedPointCaches);
-		
-		OutCollection.EmptyGroup(PV::GroupNames::FoliageGroup);
- 
-		// Collect non-empty mesh name indices
-		const TArray<FString>& FoliageMeshNames = FoliagePaletteFacade.GetFoliageNames();
-		TArray<int32> NonEmptyFoliageMeshNameIndices;
-		NonEmptyFoliageMeshNameIndices.Reserve(FoliageMeshNames.Num());
- 
-		for (int32 i = 0; i < FoliageMeshNames.Num(); ++i)
-		{
-			if (!FoliageMeshNames[i].IsEmpty())
-			{
-				NonEmptyFoliageMeshNameIndices.Add(i);
-			}
-		}
- 
-		if (NonEmptyFoliageMeshNameIndices.Num() == 0)
-		{
-			return;
-		}
-		
-		const int32 NumFoliageInfo = FoliagePaletteFacade.NumFoliageInfo();
-		const TArray<FPVFoliageInfo>& FoliageInfos = FoliagePaletteFacade.GetFoliageInfos();
-
-		auto IsValidCandidate = [&](int32 CandidateIndex)
-			{
-				return FoliageInfos.IsValidIndex(CandidateIndex) && FoliageInfos[CandidateIndex].IsValid();
-			};
-
-		auto GetCandidateValue = [&](int32 CandidateIndex, EFoliageDistributionCondition Condition) -> float
-			{
-				return FoliageInfos[CandidateIndex].Attributes.GetAttributeValue(Condition);
-			};
-		
-		for (int32 BranchIndex = 0; BranchIndex < NumBranches; ++BranchIndex)
-		{
-			FConditionAttributes ConditionAttrs;
-			TArray<float> NormalizedAttachmentPoints;
-			TArray<FVector3f> AttachmentPointsPositions;
-			TArray<float> AttachmentPointsLengthsFromRoot;
-			TArray<FVector3f> AttachmentPointsUpVectors;
-			TArray<FVector3f> AttachmentPointsNormalVectors;
-			TArray<float> AttachmentPointsScales;
-			
-			TArray<int32> BranchPoints = BranchFacadeSource.GetPoints(BranchIndex);
- 
-			if (BranchPoints.Num() < 2)
-			{
-				continue;
-			}
-			
-			const float FirstPointLengthFromRoot = PointFacadeSource.GetLengthFromRoot(BranchPoints[0]);
-			const float LastPointLengthFromRoot = PointFacadeSource.GetLengthFromRoot(BranchPoints.Last());
-			const float BranchLength = LastPointLengthFromRoot - FirstPointLengthFromRoot;
-			
-			if (BranchLength < KINDA_SMALL_NUMBER)
-			{
-				continue;
-			}
-			
-			float Increment = DistributionSettings.InstanceSpacing / BranchLength;
-			const float AdjustedMaxPerBranch = DistributionSettings.MaxPerBranch == -1
-				? BranchPoints.Num()
-				: DistributionSettings.MaxPerBranch;
-			float LookUp = 0.0f;
- 
-			for (int32 i = 0; i < AdjustedMaxPerBranch; ++i)
-			{
-				if (LookUp > 1.0f)
-				{
-					break;
-				}
- 
-				float InstanceSpacingRampValue = DistributionSettings.InstanceSpacingRamp
-					? DistributionSettings.InstanceSpacingRamp->GetRichCurveConst()->Eval(1.0f - LookUp)
-					: Increment;
-				// The increment is nudged slightly before blending with the weighted ramp value, the number is chosen based on observation
-				float AdjustedIncrement = Increment + (InstanceSpacingRampValue * (Increment * 10.0f));
-				AdjustedIncrement = FMath::Lerp(Increment, AdjustedIncrement, DistributionSettings.InstanceSpacingRampEffect);
- 
-				NormalizedAttachmentPoints.Add(1.0f - LookUp);
-				LookUp += AdjustedIncrement;
-			}
- 
-			// Compute Position, Angles, Scale and LengthFromRoot for attachment points
- 
-			// Compute based on first point attributes
-			const FVector3f InitialApicalDirection(BudDirectionAttribute[BranchPoints[0]][0]);
-			FVector3f InitialAxillaryDirection(BudDirectionAttribute[BranchPoints[0]][1]);
-			const FVector3f InitialUpVector(BudDirectionAttribute[BranchPoints[0]][5]);
- 
-			if (ResetPhyllotaxyLeaf == 1)
-			{
-				InitialAxillaryDirection = FVector3f::CrossProduct(InitialApicalDirection, InitialUpVector);
-			}
- 
-			if (PhyllotaxyOffsetLeaf > 0.01f)
-			{
-				const FQuat4f RotationQuat(InitialApicalDirection.GetSafeNormal(), FMath::DegreesToRadians(PhyllotaxyOffsetLeaf));
-				InitialAxillaryDirection = RotationQuat.RotateVector(InitialAxillaryDirection);
-			}
- 
-			int32 AxillaryRotationIteration = 0;
-			for (int32 Idx = 0; Idx < NormalizedAttachmentPoints.Num(); ++Idx)
-			{
-				for (int32 BranchPointIndex = 0; BranchPointIndex < BranchPoints.Num() - 1; ++BranchPointIndex)
-				{
-					float AxillaryRotationDegree = FMath::DegreesToRadians(PhyllotaxyFormationLeaf * AxillaryRotationIteration);
- 
-					const int32 CurrentPointIndex = BranchPoints[BranchPointIndex];
-					const int32 NextPointIndex = BranchPoints[BranchPointIndex + 1];
- 
-					const FVector3f CurrentPointPosition = PointFacadeSource.GetPosition(CurrentPointIndex);
-					const FVector3f NextPointPosition = PointFacadeSource.GetPosition(NextPointIndex);
- 
-					float EthyleneLevelAtCurrentPoint = BudHormoneLevels[CurrentPointIndex][4];
-					float EthyleneLevelAtNextPoint = BudHormoneLevels[NextPointIndex][4];
- 
-					const float CurrentPointLengthFromRoot = PointFacadeSource.GetLengthFromRoot(CurrentPointIndex);
-					const float NextPointLengthFromRoot = PointFacadeSource.GetLengthFromRoot(NextPointIndex);
-
-					const float CurrentPointLfrNormalized = FMath::GetMappedRangeValueClamped(
-						FVector2f(FirstPointLengthFromRoot, LastPointLengthFromRoot),
-						FVector2f(0.0f, 1.0f), CurrentPointLengthFromRoot);
-
-					const float NextPointLfrNormalized = FMath::GetMappedRangeValueClamped(
-						FVector2f(FirstPointLengthFromRoot, LastPointLengthFromRoot),
-						FVector2f(0.0f, 1.0f),
-						NextPointLengthFromRoot);
-
-					if (NormalizedAttachmentPoints[Idx] >= CurrentPointLfrNormalized && NormalizedAttachmentPoints[Idx] <=
-						NextPointLfrNormalized)
-					{
-						const float DistanceAlpha = (NextPointLfrNormalized - CurrentPointLfrNormalized) < KINDA_SMALL_NUMBER
-							? 1.0f
-							: (NormalizedAttachmentPoints[Idx] -
-								CurrentPointLfrNormalized) / (NextPointLfrNormalized - CurrentPointLfrNormalized);
- 
-						const float EthyleneLevel = FMath::Lerp(EthyleneLevelAtCurrentPoint, EthyleneLevelAtNextPoint, DistanceAlpha);
- 
-						if ((EthyleneLevel - 0.001f) >= DistributionSettings.EthyleneThreshold)
-						{
-							continue;
-						}
- 
-						const FVector3f Position = FMath::Lerp(CurrentPointPosition, NextPointPosition, DistanceAlpha);
-						AttachmentPointsPositions.Add(Position);
- 
-						const float AttachmentPointLfr = FMath::GetMappedRangeValueClamped(
-							FVector2f(0.0f, 1.0f),
-							FVector2f(FirstPointLengthFromRoot, LastPointLengthFromRoot),
-							NormalizedAttachmentPoints[Idx]);
-						AttachmentPointsLengthsFromRoot.Add(AttachmentPointLfr);
- 
-						// Compute Scale
-						const float CurrentPointBranchGradient = BranchGradientsAttribute[CurrentPointIndex];
-						const float NextPointBranchGradient = BranchGradientsAttribute[NextPointIndex];
- 
-						const float BranchGradient = FMath::Clamp(
-							FMath::Lerp(CurrentPointBranchGradient, NextPointBranchGradient, DistanceAlpha), 0.0f, 1.0f);
-						float BranchGradientRampValue = ScaleSettings.ScaleRamp
-							? ScaleSettings.ScaleRamp->GetRichCurveConst()->Eval(1.0f - BranchGradient)
-							: 1.0f;
- 
-						const float CurrentPointScale = PointFacadeSource.GetPointScale(CurrentPointIndex);
-						const float NextPointScale = PointFacadeSource.GetPointScale(NextPointIndex);
-						const float BranchScale = FMath::Lerp(CurrentPointScale, NextPointScale, DistanceAlpha);
- 
-						const float BranchScaledValue = BranchGradientRampValue * (BranchScale + 0.1f);
-						BranchGradientRampValue = FMath::Lerp(BranchGradientRampValue, BranchScaledValue, ScaleSettings.BranchScaleImpact);
-						BranchGradientRampValue = FMath::Clamp(BranchGradientRampValue, 0.0f, 1.0f);
- 
-						float Scale = FMath::GetMappedRangeValueClamped(
-							FVector2f(0.0f, 1.0f),
-							FVector2f(ScaleSettings.MinScale, ScaleSettings.MaxScale),
-							BranchGradientRampValue
-						) * ScaleSettings.BaseScale;
- 
-						const float RandomValue = RandomStream.FRandRange(ScaleSettings.RandomScaleMin, ScaleSettings.RandomScaleMax);
-						Scale = Scale * RandomValue;
- 
-						AttachmentPointsScales.Add(Scale);
-						
-						const FVector3f CurrentPointApicalDirectionVector(BudDirectionAttribute[CurrentPointIndex][0]);
-						const FVector3f NextPointApicalDirectionVector(BudDirectionAttribute[NextPointIndex][0]);
-						const FVector3f BudApicalDirectionVector = FMath::Lerp(CurrentPointApicalDirectionVector,
-							NextPointApicalDirectionVector, DistanceAlpha);
- 
-						// Condition - Y Dot
-						float YDot = UpAlignmentNormalized(BudApicalDirectionVector);
-						ConditionAttrs.UpAlignment.Add(YDot);
- 
-						// Condition - Light
-						// const TArray<float>& CurrentBudLightDetected = PointFacadeSource.GetBudLightDetected(CurrentPointIndex);
-						// const float CurrentLightAvailable = CurrentBudLightDetected.IsValidIndex(PV::Facades::BudLightDetectedAvailableIndex) ? CurrentBudLightDetected[PV::Facades::BudLightDetectedAvailableIndex] : 0.0f;
-						//
-						// const TArray<float>& NextBudLightDetected = PointFacadeSource.GetBudLightDetected(NextPointIndex);
-						// const float NextLightAvailable = NextBudLightDetected.IsValidIndex(PV::Facades::BudLightDetectedAvailableIndex) ? NextBudLightDetected[PV::Facades::BudLightDetectedAvailableIndex] : 0.0f;	
- 
-						const float CurrentLightAvailable = NormalizedPointCaches.Light[CurrentPointIndex];
-						const float NextLightAvailable = NormalizedPointCaches.Light[NextPointIndex];
-						const float BudLightReceived = FMath::Lerp(CurrentLightAvailable,NextLightAvailable, DistanceAlpha);
-						ConditionAttrs.Light.Add(BudLightReceived);
- 
-						// Condition - Scale
-						// const float CurrentPointScaleGradient = PointFacadeSource.GetPointScaleGradient(CurrentPointIndex);
-						// const float NextPointScaleGradient = PointFacadeSource.GetPointScaleGradient(NextPointIndex);
-						// const float PointScaleGradient = FMath::Lerp(CurrentPointScaleGradient, NextPointScaleGradient, DistanceAlpha);
- 
-						const float CurrentPScale = NormalizedPointCaches.PScales[CurrentPointIndex];
-						const float NextPScale = NormalizedPointCaches.PScales[NextPointIndex];
-						const float PointPScale = FMath::Lerp(CurrentPScale, NextPScale, DistanceAlpha);
-						ConditionAttrs.Scale.Add(PointPScale);
- 
-						// Health
-						const float CurrentHealth = NormalizedPointCaches.Health[CurrentPointIndex];
-						const float NextHealth = NormalizedPointCaches.Health[NextPointIndex];
-						const float BudHealth = FMath::Lerp(CurrentHealth,NextHealth, DistanceAlpha);
-						ConditionAttrs.Health.Add(BudHealth);
- 
-						// Condition - Tip
-						const bool bIsTip = (BranchPointIndex == BranchPoints.Num() - 2) && FMath::IsNearlyEqual(DistanceAlpha, 1.0f, 1e-3f);
-						ConditionAttrs.Tip.Add(Tip(bIsTip));
- 
-						int32 BudsToAdd = 1;
- 
-						// Compute Up Vector
-						// Last point (order of the array is last to first)
-						if (Idx == 0)
-						{
-							AttachmentPointsUpVectors.Add(BudApicalDirectionVector);
-						}
-						else // Points other than last
-							[[likely]]
-							{
-								const FQuat4f RotationQuat(InitialApicalDirection.GetSafeNormal(), AxillaryRotationDegree);
-								FVector3f NewAxillaryDirection = RotationQuat.RotateVector(InitialAxillaryDirection.GetSafeNormal());
-								FQuat4f ApicalCorrection = FQuat4f::FindBetweenNormals(
-									InitialApicalDirection.GetSafeNormal(), BudApicalDirectionVector.GetSafeNormal());
-								NewAxillaryDirection = ApicalCorrection.RotateVector(NewAxillaryDirection).GetSafeNormal();
- 
-								int32 NumBuds = RandomStream.RandRange(MinBudsLeaf, MaxBudsLeaf);
-								float RotationPerBud = 360.0f / static_cast<float>(NumBuds);
- 
-								for (int32 i = 0; i < NumBuds; ++i)
-								{
-									float RotationAngle = FMath::DegreesToRadians(RotationPerBud * i);
-									const FQuat4f AxillaryRotation(InitialApicalDirection.GetSafeNormal(), RotationAngle);
-									FVector3f AxillaryDirectionModified = AxillaryRotation.RotateVector(NewAxillaryDirection.GetSafeNormal());
- 
-									if (VectorSettings.OverrideAxilAngle)
-									{
-										float AxilAngleRampValue = VectorSettings.AxilAngleRamp
-											? FMath::Clamp(
-												VectorSettings.AxilAngleRamp->GetRichCurveConst()->Eval(1.0f - BranchGradient),
-												0.0f,
-												1.0f)
-											: 1.0f;
- 
-										AxilAngleRampValue = AxilAngleRampValue * VectorSettings.AxilAngleRampUpperValue;
-										const float AxilAngleBlended = FMath::Lerp(VectorSettings.AxilAngle, AxilAngleRampValue,
-											VectorSettings.AxilAngleRampEffect);
- 
-										const FVector3f RotationAxis = FVector3f::CrossProduct(
-											BudApicalDirectionVector.GetSafeNormal(), AxillaryDirectionModified.GetSafeNormal());
- 
-										const FQuat4f Rotation(RotationAxis.GetSafeNormal(), FMath::DegreesToRadians(AxilAngleBlended));
-										const FVector3f AxilVector = Rotation.RotateVector(BudApicalDirectionVector.GetSafeNormal());
- 
-										AttachmentPointsUpVectors.Add(AxilVector);
-									}
-									else
-									{
-										AttachmentPointsUpVectors.Add(AxillaryDirectionModified);
-									}
- 
-									if (i > 0)
-									{
-										FVector3f LastPositionAdded = AttachmentPointsPositions.Last();
-										float LastScaleAdded = AttachmentPointsScales.Last();
-										float LastLengthFromRootAdded = AttachmentPointsLengthsFromRoot.Last();
-										AttachmentPointsPositions.Add(LastPositionAdded);
-										AttachmentPointsScales.Add(LastScaleAdded);
-										AttachmentPointsLengthsFromRoot.Add(LastLengthFromRootAdded);
- 
-										// Conditions
-										float LastConditionScale = ConditionAttrs.Scale.Last();
-										ConditionAttrs.Scale.Add(LastConditionScale);
-									
-										float LastConditionLight = ConditionAttrs.Light.Last();
-										ConditionAttrs.Light.Add(LastConditionLight);
-									
-										float LastConditionUpAlignment  = ConditionAttrs.UpAlignment.Last();
-										ConditionAttrs.UpAlignment.Add(LastConditionUpAlignment);
- 
-										float LastConditionTip = ConditionAttrs.Tip.Last();
-										ConditionAttrs.Tip.Add(LastConditionTip);
- 
-										float LastConditionHealth = ConditionAttrs.Health.Last();
-										ConditionAttrs.Health.Add(LastConditionHealth);
-									}
-								}
- 
-								BudsToAdd = NumBuds;
-								AxillaryRotationIteration++;
-							}
- 
-						// Compute Normal vectors based on light optimal direction
-						const FVector3f CurrentPointLightOptimalVector(BudDirectionAttribute[CurrentPointIndex][2]);
-						const FVector3f NextPointLightOptimalVector(BudDirectionAttribute[NextPointIndex][2]);
-						const FVector3f BudLightOptimalVector = FMath::Lerp(CurrentPointLightOptimalVector, NextPointLightOptimalVector,
-							DistanceAlpha);
- 
-						FQuat4f QuatRotation = FQuat4f::FindBetweenNormals(BudApicalDirectionVector.GetSafeNormal(),
-							BudLightOptimalVector.GetSafeNormal());
-						const FVector3f ApicalDirectionRotated = QuatRotation.RotateVector(BudApicalDirectionVector);
- 
-						for (int32 i = 0; i < BudsToAdd; ++i)
-						{
-							AttachmentPointsNormalVectors.Add(ApicalDirectionRotated);
-						}
-					}
-				}
-			}
- 
-			// Set foliage entries
-			TManagedArrayAccessor<FVector3f> PivotPointsAttribute(OutCollection, PV::AttributeNames::FoliagePivotPoint,
-				PV::GroupNames::FoliageGroup);
-			
-			for (int32 Idx = 0; Idx < AttachmentPointsLengthsFromRoot.Num(); ++Idx)
-			{
-				ConditionPicker::FPointConditionSample Sample;
-				Sample.Light = ConditionAttrs.Light[Idx];
-				Sample.Scale = ConditionAttrs.Scale[Idx];
-				Sample.UpAlignment = ConditionAttrs.UpAlignment[Idx];
-				Sample.Tip = ConditionAttrs.Tip[Idx];
-				Sample.Health = ConditionAttrs.Health[Idx];
-
-				const int32 PickedFoliageIndex = ConditionPicker::PickCandidateIndex(
-					NumFoliageInfo,
-					DistributionConditions,
-					Sample,
-					RandomStream,
-					GetCandidateValue,
-					IsValidCandidate);
-
-				if (PickedFoliageIndex == INDEX_NONE)
-				{
-					continue;
-				}
- 
-				const int32 NewIndex = PivotPointsAttribute.AddElements(1);
-				FoliageFacadeOutput.SetFoliageEntry(NewIndex, {
-					.NameId = PickedFoliageIndex,
-					.BranchId = BranchIndex,
-					.PivotPoint = AttachmentPointsPositions[Idx],
-					.UpVector = AttachmentPointsUpVectors[Idx],
-					.NormalVector = AttachmentPointsNormalVectors[Idx],
-					.Scale = AttachmentPointsScales[Idx],
-					.LengthFromRoot = AttachmentPointsLengthsFromRoot[Idx],
-					.ConditionUpAlignment = ConditionAttrs.UpAlignment[Idx],
-					.ConditionTip = ConditionAttrs.Tip[Idx],
-					.ConditionLight = ConditionAttrs.Light[Idx],
-					.ConditionScale = ConditionAttrs.Scale[Idx],
-					.ConditionHealth = ConditionAttrs.Health[Idx]
-				});
-			}
-		}
-		
-		// Re-calculate Foliage IDs for each branch
-		TMap<int32, TArray<int32>> BranchIndexToFoliageIDs;
-		for (int32 i = 0; i < FoliageFacadeOutput.NumFoliageEntries(); ++i)
+
+	TArray<PV::DistributionHelper::FAttachmentPoint> AttachmentPoints;
+	ComputeHormoneBasedAttachmentPoints(
+		OutCollection,
+		NormalizedPointCaches,
+		DistributionSettings,
+		ScaleSettings,
+		VectorSettings,
+		PhyllotaxySettings,
+		RandomStream,
+		AttachmentPoints);
+
+	const int32 NumFoliageInfo = FoliagePaletteFacade.NumFoliageInfo();
+	const TArray<FPVFoliageInfo>& FoliageInfos = FoliagePaletteFacade.GetFoliageInfos();
+
+	auto IsValidCandidate = [&](int32 CandidateIndex)
 		{
-			const int32 BranchIndex = FoliageFacadeOutput.GetFoliageEntry(i).BranchId;
-			TArray<int32>& IDs = BranchIndexToFoliageIDs.FindOrAdd(BranchIndex);
-			IDs.Add(i);
-		}
-		for (const TPair<int32, TArray<int32>>& Pair : BranchIndexToFoliageIDs)
+			return FoliageInfos.IsValidIndex(CandidateIndex) && FoliageInfos[CandidateIndex].IsValid();
+		};
+
+	auto GetCandidateValue = [&](int32 CandidateIndex, EFoliageDistributionCondition Condition) -> float
 		{
-			int32 BranchId = Pair.Key;
-			const TArray<int32>& FoliageIDs = Pair.Value;
-			FoliageFacadeOutput.SetFoliageIdsArray(BranchId, FoliageIDs);
-		}
+			return FoliageInfos[CandidateIndex].Attributes.GetAttributeValue(Condition);
+		};
+
+
+	// Set foliage entries
+	TManagedArrayAccessor<FVector3f> PivotPointsAttribute(OutCollection, PV::AttributeNames::FoliagePivotPoint,
+		PV::GroupNames::FoliageGroup);
+
+	for (const auto& AttachmentPoint : AttachmentPoints)
+	{
+		const int32 PickedFoliageIndex = ConditionPicker::PickCandidateIndex(
+			NumFoliageInfo, DistributionConditions, AttachmentPoint.ConditionSample, RandomStream, GetCandidateValue, IsValidCandidate);
+
+		if (PickedFoliageIndex == INDEX_NONE) continue;
+
+		const int32 NewIndex = PivotPointsAttribute.AddElements(1);
+		FoliageFacadeOutput.SetFoliageEntry(NewIndex, {
+			.NameId = PickedFoliageIndex,
+			.BranchId = AttachmentPoint.BranchIndex,
+			.PivotPoint = AttachmentPoint.Position,
+			.UpVector = AttachmentPoint.UpDirection,
+			.NormalVector = AttachmentPoint.NormalDirection,
+			.Scale = AttachmentPoint.Scale,
+			.LengthFromRoot = AttachmentPoint.LengthFromRoot,
+			.ConditionUpAlignment = AttachmentPoint.ConditionSample.UpAlignment,
+			.ConditionTip = AttachmentPoint.ConditionSample.Tip,
+			.ConditionLight = AttachmentPoint.ConditionSample.Light,
+			.ConditionScale = AttachmentPoint.ConditionSample.Scale,
+			.ConditionHealth = AttachmentPoint.ConditionSample.Health
+		});
 	}
- 
+
+	// Re-calculate Foliage IDs for each branch
+	TMap<int32, TArray<int32>> BranchIndexToFoliageIDs;
+	for (int32 i = 0; i < FoliageFacadeOutput.NumFoliageEntries(); ++i)
 	{
-		TRACE_CPUPROFILER_EVENT_SCOPE(PVFoliageDistributorNode::EthyleneRemove);
- 
-		// No need to apply Ethylene threshold if the distribution was overriden since that's already been accounted for
-		if (DistributionSettings.OverrideDistribution)
-		{
-			return;
-		}
- 
-		// Compute Foliage instances to remove based on Ethylene threshold
-		TManagedArrayAccessor<FVector3f> PivotPointsAttribute(
-			OutCollection,
-			PV::AttributeNames::FoliagePivotPoint,
-			PV::GroupNames::FoliageGroup);
-		TSet<int32> FoliageEntriesToRemove;
- 
-		for (int32 BranchIndex = 0; BranchIndex < NumBranches; ++BranchIndex)
-		{
-			TArray<int32> BranchPoints = BranchFacadeSource.GetPoints(BranchIndex);
-			TArray<int32> FoliageIDs = FoliageFacadeOutput.GetFoliageEntryIdsForBranch(BranchIndex);
-			TSet<int32> FoliageIDsToRemove;
- 
-			for (int32 BranchPointIndex = 0; BranchPointIndex < BranchPoints.Num() - 1; ++BranchPointIndex)
-			{
-				for (int32 FoliageIDsIndex = 0; FoliageIDsIndex < FoliageIDs.Num(); ++FoliageIDsIndex)
-				{
-					const float CurrentPointLengthFromRoot = PointFacadeSource.GetLengthFromRoot(BranchPoints[BranchPointIndex]);
-					const float NextPointLengthFromRoot = PointFacadeSource.GetLengthFromRoot(BranchPoints[BranchPointIndex + 1]);
-					const float FoliageLengthFromRoot = FoliageFacadeOutput.GetFoliageEntry(FoliageIDs[FoliageIDsIndex]).LengthFromRoot;
- 
-					if (FoliageLengthFromRoot >= CurrentPointLengthFromRoot && FoliageLengthFromRoot <= NextPointLengthFromRoot)
-					{
-						float EthyleneLevelAtCurrentPoint = BudHormoneLevels[BranchPoints[BranchPointIndex]][4];
-						float EthyleneLevelAtNextPoint = BudHormoneLevels[BranchPoints[BranchPointIndex + 1]][4];
- 
-						float EthyleneLevel = FMath::GetMappedRangeValueClamped(
-							FVector2f(CurrentPointLengthFromRoot, NextPointLengthFromRoot),
-							FVector2f(EthyleneLevelAtNextPoint, EthyleneLevelAtCurrentPoint),
-							FoliageLengthFromRoot);
- 
-						if ((EthyleneLevel - 0.001f) >= DistributionSettings.EthyleneThreshold)
-						{
-							FoliageIDsToRemove.Add(FoliageIDs[FoliageIDsIndex]);
-						}
-					}
-				}
-			}
- 
-			for (int32 Fid : FoliageIDsToRemove)
-			{
-				FoliageEntriesToRemove.Add(Fid);
-			}
-		}
- 
-		// Remove foliage instances
-		for (int32 i = PivotPointsAttribute.Num() - 1; i >= 0; --i)
-		{
-			if (FoliageEntriesToRemove.Contains(i))
-			{
-				PivotPointsAttribute.RemoveElements(1, i);
-			}
-		}
- 
-		// Re-calculate Foliage IDs for each branch
-		TMap<int32, TArray<int32>> BranchIndexToFoliageIDs;
-		for (int32 i = 0; i < FoliageFacadeOutput.NumFoliageEntries(); ++i)
-		{
-			const int32 BranchIndex = FoliageFacadeOutput.GetFoliageEntry(i).BranchId;
-			TArray<int32>& IDs = BranchIndexToFoliageIDs.FindOrAdd(BranchIndex);
-			IDs.Add(i);
-		}
-		for (const TPair<int32, TArray<int32>>& Pair : BranchIndexToFoliageIDs)
-		{
-			int32 BranchId = Pair.Key;
-			const TArray<int32>& FoliageIDs = Pair.Value;
-			FoliageFacadeOutput.SetFoliageIdsArray(BranchId, FoliageIDs);
-		}
+		const int32 BranchIndex = FoliageFacadeOutput.GetFoliageEntry(i).BranchId;
+		TArray<int32>& IDs = BranchIndexToFoliageIDs.FindOrAdd(BranchIndex);
+		IDs.Add(i);
+	}
+	for (const TPair<int32, TArray<int32>>& Pair : BranchIndexToFoliageIDs)
+	{
+		int32 BranchId = Pair.Key;
+		const TArray<int32>& FoliageIDs = Pair.Value;
+		FoliageFacadeOutput.SetFoliageIdsArray(BranchId, FoliageIDs);
 	}
 }
 
@@ -611,7 +158,7 @@ void FPVFoliage::DistributeFoliageWithParametricSettings(
 	const FPVFoliageDistributionConditions& DistributionConditions)
 {
 	TRACE_CPUPROFILER_EVENT_SCOPE(PVFoliage::DistributeFoliageWithParametricSettings);
-	
+
 	if (!(SpacingSettings.SpacingRamp && AngleSettings.AxilAngleRamp && ScaleSettings.ScaleRamp))
 	{
 		return;
@@ -619,7 +166,7 @@ void FPVFoliage::DistributeFoliageWithParametricSettings(
 
 	PV::Facades::FFoliageFacade FoliageFacadeOutput(OutCollection);
 	PV::Facades::FFoliageFacade FoliagePaletteFacade(FoliageCollection);
-	
+
 	if (const int32 NumFoliageMeshes = FoliagePaletteFacade.NumFoliageInfo(); NumFoliageMeshes == 0)
 	{
 		UE_LOG(LogProceduralVegetation, Warning, TEXT("There are no foliage meshes available in input for distribution"));
@@ -627,10 +174,10 @@ void FPVFoliage::DistributeFoliageWithParametricSettings(
 	}
 
 	FoliageFacadeOutput.SetFoliageInfos(FoliagePaletteFacade.GetFoliageInfos());
-	
+
 	FRandomStream RandomStream;
 	RandomStream.Initialize(RandomSeed);
-	
+
 	FNormalizedPointCaches NormalizedPointCaches;
 	FBuildOptions BuildOptions;
 	BuildOptions.bComputePScales = DistributionConditions.IsActiveCondition(EFoliageDistributionCondition::Scale);
@@ -641,19 +188,19 @@ void FPVFoliage::DistributeFoliageWithParametricSettings(
 		OutCollection,
 		BuildOptions,
 		NormalizedPointCaches);
-	
-	TArray<FAttachmentPoint> AttachmentPoints;
+
+	TArray<PV::DistributionHelper::FAttachmentPoint> AttachmentPoints;
 	ComputeAttachmentPoints(
-		OutCollection, 
-		NormalizedPointCaches, 
-		SpacingSettings, 
-		AngleSettings, 
-		ScaleSettings, 
-		RandomStream, 
+		OutCollection,
+		NormalizedPointCaches,
+		SpacingSettings,
+		AngleSettings,
+		ScaleSettings,
+		RandomStream,
 		AttachmentPoints);
-	
+
 	OutCollection.EmptyGroup(PV::GroupNames::FoliageGroup);
-	
+
 	const int32 NumFoliageInfo = FoliagePaletteFacade.NumFoliageInfo();
 	const TArray<FPVFoliageInfo>& FoliageInfos = FoliagePaletteFacade.GetFoliageInfos();
 
@@ -666,10 +213,10 @@ void FPVFoliage::DistributeFoliageWithParametricSettings(
 		{
 			return FoliageInfos[CandidateIndex].Attributes.GetAttributeValue(Condition);
 		};
-	
+
 	TManagedArrayAccessor<FVector3f> PivotPointsAttribute(OutCollection, PV::AttributeNames::FoliagePivotPoint,
-				PV::GroupNames::FoliageGroup);
-	
+		PV::GroupNames::FoliageGroup);
+
 	for (const auto& AttachmentPoint : AttachmentPoints)
 	{
 		const int32 PickedFoliageIndex = ConditionPicker::PickCandidateIndex(
@@ -679,7 +226,7 @@ void FPVFoliage::DistributeFoliageWithParametricSettings(
 			RandomStream,
 			GetCandidateValue,
 			IsValidCandidate);
-		
+
 		if (PickedFoliageIndex == INDEX_NONE)
 		{
 			continue;
@@ -701,7 +248,7 @@ void FPVFoliage::DistributeFoliageWithParametricSettings(
 			.ConditionHealth = AttachmentPoint.ConditionSample.Health
 		});
 	}
-	
+
 	// Re-calculate Foliage IDs for each branch
 	TMap<int32, TArray<int32>> BranchIndexToFoliageIDs;
 	for (int32 i = 0; i < FoliageFacadeOutput.NumFoliageEntries(); ++i)
diff --git a/Engine/Plugins/Experimental/ProceduralVegetationEditor/Source/ProceduralVegetation/Private/Implementations/PVFoliage.h b/Engine/Plugins/Experimental/ProceduralVegetationEditor/Source/ProceduralVegetation/Private/Implementations/PVFoliage.h
--- a/Engine/Plugins/Experimental/ProceduralVegetationEditor/Source/ProceduralVegetation/Private/Implementations/PVFoliage.h
+++ b/Engine/Plugins/Experimental/ProceduralVegetationEditor/Source/ProceduralVegetation/Private/Implementations/PVFoliage.h
@@ -4,67 +4,23 @@
 
 #include "CoreMinimal.h"
 #include "GeometryCollection/ManagedArrayCollection.h"
+#include "Helpers/PVHormoneDistributionHelper.h"
 #include "Helpers/PVParametricDistributionHelper.h"
 
 struct FPVFoliageDistributionConditions;
-enum class EPhyllotaxyFormation : uint8;
-enum class EPhyllotaxyType : uint8;
-struct FPVFloatRamp;
 
 using namespace PV::ParametricDistributionHelper;
-
-namespace PV::HormoneBasedSettings
-{
-	struct FDistributionSettings
-	{
-		const float EthyleneThreshold;
-		const bool OverrideDistribution;
-		const float InstanceSpacing;
-		const FPVFloatRamp* InstanceSpacingRamp;
-		const float InstanceSpacingRampEffect;
-		const int32 MaxPerBranch;
-	};
-
-	struct FScaleSettings
-	{
-		const float BaseScale;
-		const float BranchScaleImpact;
-		const float MinScale;
-		const float MaxScale;
-		const float RandomScaleMin;
-		const float RandomScaleMax;
-		const FPVFloatRamp* ScaleRamp;
-	};
-
-	struct FVectorSettings
-	{
-		const bool OverrideAxilAngle;
-		const float AxilAngle;
-		const FPVFloatRamp* AxilAngleRamp;
-		const float AxilAngleRampUpperValue;
-		const float AxilAngleRampEffect;
-	};
-
-	struct FPhyllotaxySettings
-	{
-		const bool OverridePhyllotaxy;
-		const EPhyllotaxyType PhyllotaxyType;
-		const EPhyllotaxyFormation PhyllotaxyFormation;
-		const int32 MinimumNodeBuds;
-		const int32 MaximumNodeBuds;
-		const float PhyllotaxyAdditionalAngle;
-	};
-}
+using namespace PV::HormoneDistributionHelper;
 
 struct FPVFoliage
 {
 	static void DistributeFoliageWithHormoneBasedSettings(
 		FManagedArrayCollection& OutCollection,
 		const FManagedArrayCollection& FoliageCollection,
-		const PV::HormoneBasedSettings::FDistributionSettings& DistributionSettings,
-		const PV::HormoneBasedSettings::FScaleSettings& ScaleSettings,
-		const PV::HormoneBasedSettings::FVectorSettings& VectorSettings,
-		const PV::HormoneBasedSettings::FPhyllotaxySettings& PhyllotaxySettings,
+		const FHormoneSettings::FDistributionSettings& DistributionSettings,
+		const FHormoneSettings::FScaleSettings& ScaleSettings,
+		const FHormoneSettings::FVectorSettings& VectorSettings,
+		const FHormoneSettings::FPhyllotaxySettings& PhyllotaxySettings,
 		const int32 RandomSeed,
 		const FPVFoliageDistributionConditions& DistributionConditions);
 
diff --git a/Engine/Plugins/Experimental/ProceduralVegetationEditor/Source/ProceduralVegetation/Private/Implementations/PVGrafter.cpp b/Engine/Plugins/Experimental/ProceduralVegetationEditor/Source/ProceduralVegetation/Private/Implementations/PVGrafter.cpp
--- a/Engine/Plugins/Experimental/ProceduralVegetationEditor/Source/ProceduralVegetation/Private/Implementations/PVGrafter.cpp
+++ b/Engine/Plugins/Experimental/ProceduralVegetationEditor/Source/ProceduralVegetation/Private/Implementations/PVGrafter.cpp
@@ -650,14 +650,19 @@ TArray<FManagedArrayCollection> FPVGrafter::ExtractIndividualPlants(const FManag
 void FPVGrafter::DistributeGraftWithHormoneBasedSettings(
 	FManagedArrayCollection& OutCollection,
 	TArray<FManagedArrayCollection>& SourceGraftCollections,
-	const FHormoneBasedSettings::FDistributionSettings& DistributionSettings,
-	const FHormoneBasedSettings::FPhyllotaxySettings& PhyllotaxySettings,
-	const FHormoneBasedSettings::FScaleSettings& ScaleSettings,
-	const FHormoneBasedSettings::FVectorSettings& VectorSettings,
+	const FHormoneSettings::FDistributionSettings& DistributionSettings,
+	const FHormoneSettings::FPhyllotaxySettings& PhyllotaxySettings,
+	const FHormoneSettings::FScaleSettings& ScaleSettings,
+	const FHormoneSettings::FVectorSettings& VectorSettings,
 	const int32 RandomSeed,
 	const FPVFoliageDistributionConditions& DistributionConditions)
 {
 	TRACE_CPUPROFILER_EVENT_SCOPE(PVGrafter::DistributeGraftWithHormoneBasedSettings);
+	
+	if (!(DistributionSettings.InstanceSpacingRamp && ScaleSettings.ScaleRamp && VectorSettings.AxilAngleRamp))
+	{
+		return;
+	}
 
 	FRandomStream RandomStream;
 	RandomStream.Initialize(RandomSeed);
@@ -710,52 +715,16 @@ void FPVGrafter::DistributeGraftWithHormoneBasedSettings(
 		CandidateGraftData[i] = MetaInfoFacade.GetGraftAttributes();
 	}
 
-	auto BudDirectionAttribute = PV::FBudDirectionAttribute::GetAttribute(OutCollection);
-	auto BranchGradientsAttribute = PV::FPointBranchGradientAttribute::GetAttribute(OutCollection);
-	auto BudHormoneLevels = PV::FBudHormoneLevelsAttribute::GetAttribute(OutCollection);
-
-	float PhyllotaxyFormationLeaf = 180;
-	float MinBudsLeaf = 1;
-	float MaxBudsLeaf = 1;
-	bool ResetPhyllotaxyLeaf = PhyllotaxySettings.ResetPhyllotaxy;
-	float PhyllotaxyOffsetLeaf = PhyllotaxySettings.PhyllotaxyOffset;
-
-	switch (PhyllotaxySettings.PhyllotaxyType)
-	{
-	case EPhyllotaxyType::Alternate: PhyllotaxyFormationLeaf = 180 + PhyllotaxySettings.PhyllotaxyAdditionalAngle;
-		MinBudsLeaf = MaxBudsLeaf = 1;
-		break;
-	case EPhyllotaxyType::Opposite: PhyllotaxyFormationLeaf = PhyllotaxySettings.PhyllotaxyAdditionalAngle;
-		MinBudsLeaf = MaxBudsLeaf = 2;
-		break;
-	case EPhyllotaxyType::Decussate: PhyllotaxyFormationLeaf = 90 + PhyllotaxySettings.PhyllotaxyAdditionalAngle;
-		MinBudsLeaf = MaxBudsLeaf = 2;
-		break;
-	case EPhyllotaxyType::Whorled: PhyllotaxyFormationLeaf = 90 + PhyllotaxySettings.PhyllotaxyAdditionalAngle;
-		MinBudsLeaf = PhyllotaxySettings.MinimumNodeBuds;
-		MaxBudsLeaf = PhyllotaxySettings.MaximumNodeBuds;
-		break;
-	case EPhyllotaxyType::Spiral:
-		{
-			// These leaf arrangement configurations where Distichous is two leaves, Tristichous is three etc.
-			switch (PhyllotaxySettings.PhyllotaxyFormation)
-			{
-			case EPhyllotaxyFormation::Distichous: PhyllotaxyFormationLeaf = 180 + PhyllotaxySettings.PhyllotaxyAdditionalAngle;
-				break;
-			case EPhyllotaxyFormation::Tristichous: PhyllotaxyFormationLeaf = 120 + PhyllotaxySettings.PhyllotaxyAdditionalAngle;
-				break;
-			case EPhyllotaxyFormation::Pentastichous: PhyllotaxyFormationLeaf = 144 + PhyllotaxySettings.PhyllotaxyAdditionalAngle;
-				break;
-			case EPhyllotaxyFormation::Octastichous: PhyllotaxyFormationLeaf = 135 + PhyllotaxySettings.PhyllotaxyAdditionalAngle;
-				break;
-			case EPhyllotaxyFormation::Parastichous: PhyllotaxyFormationLeaf = PhyllotaxySettings.PhyllotaxyAdditionalAngle;
-				break;
-			}
-
-			MinBudsLeaf = MaxBudsLeaf = 1;
-			break;
-		}
-	}
+	TArray<PV::DistributionHelper::FAttachmentPoint> AttachmentPoints;
+	ComputeHormoneBasedAttachmentPoints(
+		OutCollection,
+		NormalizedPointCaches,
+		DistributionSettings,
+		ScaleSettings,
+		VectorSettings,
+		PhyllotaxySettings,
+		RandomStream,
+		AttachmentPoints);
 
 	auto BudNumberAttribute = PV::FPointBudNumberAttribute::GetAttribute(OutCollection);
 	check(BudNumberAttribute.Num() > 0);
@@ -791,340 +760,44 @@ void FPVGrafter::DistributeGraftWithHormoneBasedSettings(
 			}
 		};
 
-	// Compute attachment points in terms of normalized distances
-	auto PointLengthFromRootAttribute = PV::FPointLengthFromRootAttribute::GetAttribute(OutCollection);
-	auto PointScaleAttribute = PV::FPointScaleAttribute::GetAttribute(OutCollection);
-	auto PointPositionAttribute = PV::FPointPositionAttribute::GetAttribute(OutCollection);
-	const int32 NumBranches = BranchPointsAttribute.Num();
-	for (int32 BranchIndex = 0; BranchIndex < NumBranches; ++BranchIndex)
+	for (const auto& AttachmentPoint : AttachmentPoints)
 	{
-		FConditionAttributes ConditionAttrs;
-		TArray<float> NormalizedAttachmentPoints;
-		TArray<float> NormalizedAttachmentPointsFinal;
-		TArray<FVector3f> AttachmentPointsPositions;
-		TArray<float> AttachmentPointsLengthsFromRoot;
-		TArray<FVector3f> AttachmentPointsUpVectors;
-		TArray<FVector3f> AttachmentPointsNormalVectors;
-		TArray<float> AttachmentPointsScales;
-		TArray<int32> AttachmentPointsPointIndices;
-		const TArray<int32>& BranchPoints = BranchPointsAttribute[BranchIndex];
-
-		if (BranchPoints.Num() < 2)
-		{
-			continue;
-		}
-
-		const float FirstPointLengthFromRoot = PointLengthFromRootAttribute[BranchPoints[0]];
-		const float LastPointLengthFromRoot = PointLengthFromRootAttribute[BranchPoints.Last()];
-		const float BranchLength = LastPointLengthFromRoot - FirstPointLengthFromRoot;
-
-		if (BranchLength <= KINDA_SMALL_NUMBER)
-		{
-			continue;
-		}
-
-		float Increment = DistributionSettings.InstanceSpacing / BranchLength;
-		const float AdjustedMaxPerBranch = DistributionSettings.MaxPerBranch == -1
-			? BranchPoints.Num()
-			: DistributionSettings.MaxPerBranch;
-		float LookUp = 0.0f;
-
-		for (int32 i = 0; i < AdjustedMaxPerBranch; ++i)
-		{
-			if (LookUp > 1.0f)
-			{
-				break;
-			}
-
-			float InstanceSpacingRampValue = DistributionSettings.InstanceSpacingRamp
-				? DistributionSettings.InstanceSpacingRamp->GetRichCurveConst()->Eval(1.0f - LookUp)
-				: Increment;
-			// The increment is nudged slightly before blending with the weighted ramp value, the number is chosen based on observation
-			float AdjustedIncrement = Increment + (InstanceSpacingRampValue * (Increment * 10.0f));
-			AdjustedIncrement = FMath::Lerp(Increment, AdjustedIncrement, DistributionSettings.InstanceSpacingRampEffect);
-
-			NormalizedAttachmentPoints.Add(1.0f - LookUp);
-			LookUp += AdjustedIncrement;
-		}
-
-		// Compute Position, Angles, Scale and LengthFromRoot for attachment points
+		const int32 PickedCandidateIndex = ConditionPicker::PickCandidateIndex(
+			GraftSourcesExtracted.Num(), DistributionConditions, AttachmentPoint.ConditionSample, RandomStream, GetCandidateValue, IsValidCandidate);
 
-		// Compute based on first point attributes
-		const FVector3f InitialApicalDirection(BudDirectionAttribute[BranchPoints[0]].Apical);
-		FVector3f InitialAxillaryDirection(BudDirectionAttribute[BranchPoints[0]].Axillary);
-		const FVector3f InitialUpVector(BudDirectionAttribute[BranchPoints[0]].UpVector);
+		if (PickedCandidateIndex == INDEX_NONE) continue;
 
-		if (ResetPhyllotaxyLeaf)
-		{
-			InitialAxillaryDirection = FVector3f::CrossProduct(InitialApicalDirection, InitialUpVector);
-		}
+		int32 AttachmentPointIndex = INDEX_NONE;
 
-		if (PhyllotaxyOffsetLeaf > 0.01f)
+		if (AttachmentPoint.bIsTipInstance)
 		{
-			const FQuat4f RotationQuat(InitialApicalDirection.GetSafeNormal(), FMath::DegreesToRadians(PhyllotaxyOffsetLeaf));
-			InitialAxillaryDirection = RotationQuat.RotateVector(InitialAxillaryDirection);
+			AttachmentPointIndex = AttachmentPoint.TipPointIndex;
 		}
-
-		int32 AxillaryRotationIteration = 0;
-		for (int32 Idx = 0; Idx < NormalizedAttachmentPoints.Num(); ++Idx)
+		else
 		{
-			for (int32 BranchPointIndex = 0; BranchPointIndex < BranchPoints.Num() - 1; ++BranchPointIndex)
-			{
-				float AxillaryRotationDegree = FMath::DegreesToRadians(PhyllotaxyFormationLeaf * AxillaryRotationIteration);
-
-				const int32 CurrentPointIndex = BranchPoints[BranchPointIndex];
-				const int32 NextPointIndex = BranchPoints[BranchPointIndex + 1];
-				const FVector3f CurrentPointPosition = PointPositionAttribute[CurrentPointIndex];
-				const FVector3f NextPointPosition = PointPositionAttribute[NextPointIndex];
-				float EthyleneLevelAtCurrentPoint = BudHormoneLevels[CurrentPointIndex].Ethylene;
-				float EthyleneLevelAtNextPoint = BudHormoneLevels[NextPointIndex].Ethylene;
-				const float CurrentPointLengthFromRoot = PointLengthFromRootAttribute[CurrentPointIndex];
-				const float NextPointLengthFromRoot = PointLengthFromRootAttribute[NextPointIndex];
-				const float CurrentPointLfrNormalized = FMath::GetMappedRangeValueClamped(
-					FVector2f(FirstPointLengthFromRoot, LastPointLengthFromRoot),
-					FVector2f(0.0f, 1.0f), CurrentPointLengthFromRoot);
-				const float NextPointLfrNormalized = FMath::GetMappedRangeValueClamped(
-					FVector2f(FirstPointLengthFromRoot, LastPointLengthFromRoot),
-					FVector2f(0.0f, 1.0f),
-					NextPointLengthFromRoot);
-
-				if (NormalizedAttachmentPoints[Idx] >= CurrentPointLfrNormalized && NormalizedAttachmentPoints[Idx] <=
-					NextPointLfrNormalized)
-				{
-					const float DistanceAlpha = (NextPointLfrNormalized - CurrentPointLfrNormalized) < KINDA_SMALL_NUMBER
-						? 1.0f
-						: (NormalizedAttachmentPoints[Idx] -
-							CurrentPointLfrNormalized) / (NextPointLfrNormalized - CurrentPointLfrNormalized);
-
-					const float EthyleneLevel = FMath::Lerp(EthyleneLevelAtCurrentPoint, EthyleneLevelAtNextPoint, DistanceAlpha);
-
-					if ((EthyleneLevel - 0.001f) >= DistributionSettings.EthyleneThreshold)
-					{
-						continue;
-					}
-
-					const FVector3f Position = FMath::Lerp(CurrentPointPosition, NextPointPosition, DistanceAlpha);
-					AttachmentPointsPositions.Add(Position);
-
-					const float AttachmentPointLfr = FMath::GetMappedRangeValueClamped(
-						FVector2f(0.0f, 1.0f),
-						FVector2f(FirstPointLengthFromRoot, LastPointLengthFromRoot),
-						NormalizedAttachmentPoints[Idx]);
-					AttachmentPointsLengthsFromRoot.Add(AttachmentPointLfr);
-
-					int32 IndexAdded;
-					// Don't add a new point for a tip instance
-					if (Idx == 0 && FMath::IsNearlyEqual(AttachmentPointLfr, LastPointLengthFromRoot, 0.001f))
-					{
-						IndexAdded = NextPointIndex;
-					}
-					else
-					{
-						// Setting NjordPixelIdx to an integer value so that it is not removed
-						// in the mesher when it removes non bud points
-						IndexAdded = PV::Utilities::AddInterpolatedPointToCollection(
-							OutCollection,
-							CurrentPointIndex,
-							NextPointIndex,
-							DistanceAlpha,
-							CurrentBudNumber,
-							CurrentBudNumber);
-
-						CurrentBudNumber++;
-					}
-					AttachmentPointsPointIndices.Add(IndexAdded);
-
-					// Compute Scale
-					const float CurrentPointBranchGradient = BranchGradientsAttribute[CurrentPointIndex];
-					const float NextPointBranchGradient = BranchGradientsAttribute[NextPointIndex];
-
-					const float BranchGradient = FMath::Clamp(
-						FMath::Lerp(CurrentPointBranchGradient, NextPointBranchGradient, DistanceAlpha), 0.0f, 1.0f);
-					float BranchGradientRampValue = ScaleSettings.ScaleRamp
-						? ScaleSettings.ScaleRamp->GetRichCurveConst()->Eval(1.0f - BranchGradient)
-						: 1.0f;
-
-					const float CurrentPointScale = PointScaleAttribute[CurrentPointIndex];
-					const float NextPointScale = PointScaleAttribute[NextPointIndex];
-					const float BranchScale = FMath::Lerp(CurrentPointScale, NextPointScale, DistanceAlpha);
-
-					const float BranchScaledValue = BranchGradientRampValue * (BranchScale + 0.1f);
-					BranchGradientRampValue = FMath::Lerp(BranchGradientRampValue, BranchScaledValue, ScaleSettings.BranchScaleImpact);
-					BranchGradientRampValue = FMath::Clamp(BranchGradientRampValue, 0.0f, 1.0f);
-
-					float Scale = FMath::GetMappedRangeValueClamped(
-						FVector2f(0.0f, 1.0f),
-						FVector2f(ScaleSettings.MinScale, ScaleSettings.MaxScale),
-						BranchGradientRampValue
-					) * ScaleSettings.BaseScale;
-
-					const float RandomValue = RandomStream.FRandRange(ScaleSettings.RandomScaleMin, ScaleSettings.RandomScaleMax);
-					Scale = Scale * RandomValue;
-
-					AttachmentPointsScales.Add(Scale);
-
-					const FVector3f CurrentPointApicalDirectionVector(BudDirectionAttribute[CurrentPointIndex].Apical);
-					const FVector3f NextPointApicalDirectionVector(BudDirectionAttribute[NextPointIndex].Apical);
-					const FVector3f BudApicalDirectionVector = FMath::Lerp(CurrentPointApicalDirectionVector,
-						NextPointApicalDirectionVector, DistanceAlpha);
-
-					ConditionAttrs.UpAlignment.Add(UpAlignmentNormalized(BudApicalDirectionVector));
-
-					const float CurrentLightAvailable = NormalizedPointCaches.Light[CurrentPointIndex];
-					const float NextLightAvailable = NormalizedPointCaches.Light[NextPointIndex];
-					const float BudLightReceived = FMath::Lerp(CurrentLightAvailable, NextLightAvailable, DistanceAlpha);
-					ConditionAttrs.Light.Add(BudLightReceived);
-
-					const float CurrentHealth = NormalizedPointCaches.Health[CurrentPointIndex];
-					const float NextHealth = NormalizedPointCaches.Health[NextPointIndex];
-					const float BudHealth = FMath::Lerp(CurrentHealth, NextHealth, DistanceAlpha);
-					ConditionAttrs.Health.Add(BudHealth);
-
-					const float CurrentPScale = NormalizedPointCaches.PScales[CurrentPointIndex];
-					const float NextPScale = NormalizedPointCaches.PScales[NextPointIndex];
-					const float PointPScale = FMath::Lerp(CurrentPScale, NextPScale, DistanceAlpha);
-					ConditionAttrs.Scale.Add(PointPScale);
-
-					const bool bIsTip =
-						(BranchPointIndex == BranchPoints.Num() - 2) &&
-						FMath::IsNearlyEqual(DistanceAlpha, 1.0f, 1e-3f);
-					ConditionAttrs.Tip.Add(Tip(bIsTip));
-
-					int32 BudsToAdd = 1;
-
-					// Compute Up Vector
-					// Last point (order of the array is last to first)
-					if (Idx == 0)
-					{
-						AttachmentPointsUpVectors.Add(BudApicalDirectionVector);
-					}
-					else // Points other than last
-					[[likely]]
-					{
-						const FQuat4f RotationQuat(InitialApicalDirection.GetSafeNormal(), AxillaryRotationDegree);
-						FVector3f NewAxillaryDirection = RotationQuat.RotateVector(InitialAxillaryDirection.GetSafeNormal());
-						FQuat4f ApicalCorrection = FQuat4f::FindBetweenNormals(
-							InitialApicalDirection.GetSafeNormal(), BudApicalDirectionVector.GetSafeNormal());
-						NewAxillaryDirection = ApicalCorrection.RotateVector(NewAxillaryDirection).GetSafeNormal();
-
-						int32 NumBuds = RandomStream.RandRange(MinBudsLeaf, MaxBudsLeaf);
-						float RotationPerBud = 360.0f / static_cast<float>(NumBuds);
-
-						for (int32 i = 0; i < NumBuds; ++i)
-						{
-							float RotationAngle = FMath::DegreesToRadians(RotationPerBud * i);
-							const FQuat4f AxillaryRotation(InitialApicalDirection.GetSafeNormal(), RotationAngle);
-							FVector3f AxillaryDirectionModified = AxillaryRotation.RotateVector(NewAxillaryDirection.GetSafeNormal());
-
-							if (VectorSettings.OverrideAxilAngle)
-							{
-								float AxilAngleRampValue = VectorSettings.AxilAngleRamp
-									? FMath::Clamp(
-										VectorSettings.AxilAngleRamp->GetRichCurveConst()->Eval(1.0f - BranchGradient),
-										0.0f,
-										1.0f)
-									: 1.0f;
-
-								AxilAngleRampValue = AxilAngleRampValue * VectorSettings.AxilAngleRampUpperValue;
-								const float AxilAngleBlended = FMath::Lerp(VectorSettings.AxilAngle, AxilAngleRampValue,
-									VectorSettings.AxilAngleRampEffect);
-
-								const FVector3f RotationAxis = FVector3f::CrossProduct(
-									BudApicalDirectionVector.GetSafeNormal(), AxillaryDirectionModified.GetSafeNormal());
-
-								const FQuat4f Rotation(RotationAxis.GetSafeNormal(), FMath::DegreesToRadians(AxilAngleBlended));
-								const FVector3f AxilVector = Rotation.RotateVector(BudApicalDirectionVector.GetSafeNormal());
-
-								AttachmentPointsUpVectors.Add(AxilVector);
-							}
-							else
-							{
-								AttachmentPointsUpVectors.Add(AxillaryDirectionModified);
-							}
-
-							if (i > 0)
-							{
-								FVector3f LastPositionAdded = AttachmentPointsPositions.Last();
-								float LastScaleAdded = AttachmentPointsScales.Last();
-								float LastLengthFromRootAdded = AttachmentPointsLengthsFromRoot.Last();
-								int32 LastPointIndexAdded = AttachmentPointsPointIndices.Last();
-								AttachmentPointsPositions.Add(LastPositionAdded);
-								AttachmentPointsPointIndices.Add(LastPointIndexAdded);
-								AttachmentPointsScales.Add(LastScaleAdded);
-								AttachmentPointsLengthsFromRoot.Add(LastLengthFromRootAdded);
-
-								ConditionAttrs.Light.Add(ConditionAttrs.Light.Last());
-								ConditionAttrs.Health.Add(ConditionAttrs.Health.Last());
-								ConditionAttrs.UpAlignment.Add(ConditionAttrs.UpAlignment.Last());
-								ConditionAttrs.Scale.Add(ConditionAttrs.Scale.Last());
-								ConditionAttrs.Tip.Add(ConditionAttrs.Tip.Last());
-							}
-						}
-
-						BudsToAdd = NumBuds;
-						AxillaryRotationIteration++;
-					}
-
-					// Compute Normal vectors based on light optimal direction
-					const FVector3f CurrentPointLightOptimalVector(BudDirectionAttribute[CurrentPointIndex].LightOptimal);
-					const FVector3f NextPointLightOptimalVector(BudDirectionAttribute[NextPointIndex].LightOptimal);
-					const FVector3f BudLightOptimalVector = FMath::Lerp(CurrentPointLightOptimalVector, NextPointLightOptimalVector,
-						DistanceAlpha);
-
-					FQuat4f QuatRotation = FQuat4f::FindBetweenNormals(BudApicalDirectionVector.GetSafeNormal(),
-						BudLightOptimalVector.GetSafeNormal());
-					const FVector3f ApicalDirectionRotated = QuatRotation.RotateVector(BudApicalDirectionVector);
-
-					for (int32 i = 0; i < BudsToAdd; ++i)
-					{
-						AttachmentPointsNormalVectors.Add(ApicalDirectionRotated);
-					}
-				}
-			}
+			AttachmentPointIndex = PV::Utilities::AddInterpolatedPointToCollection(
+				OutCollection,
+				AttachmentPoint.PointIndexA,
+				AttachmentPoint.PointIndexB,
+				AttachmentPoint.PointAlpha,
+				CurrentBudNumber,
+				CurrentBudNumber);
+			CurrentBudNumber++;
 		}
 
-		for (int32 Idx = 0; Idx < AttachmentPointsLengthsFromRoot.Num(); Idx++)
-		{
-			bool bIsTipInstance =
-				Idx == 0 && FMath::IsNearlyEqual(AttachmentPointsLengthsFromRoot[Idx], LastPointLengthFromRoot, 0.001f);
-
-			FPVGraftAttributes Attributes{
-				.BranchIndex = BranchIndex,
-				.PointIndex = AttachmentPointsPointIndices[Idx],
-				.AttachmentPoint = AttachmentPointsPositions[Idx],
-				.UpVector = AttachmentPointsUpVectors[Idx],
-				.NormalVector = AttachmentPointsNormalVectors[Idx],
-				.Scale = AttachmentPointsScales[Idx],
-				.LengthFromRoot = AttachmentPointsLengthsFromRoot[Idx],
-				.bIsTipInstance = bIsTipInstance
-			};
-
-			check(GraftSourcesExtracted.Num() > 0);
-
-			ConditionPicker::FPointConditionSample Sample;
-			Sample.Light = ConditionAttrs.Light[Idx];
-			Sample.UpAlignment = ConditionAttrs.UpAlignment[Idx];
-			Sample.Health = ConditionAttrs.Health[Idx];
-			Sample.Tip = ConditionAttrs.Tip[Idx];
-			Sample.Scale = ConditionAttrs.Scale[Idx];
-
-			const int32 PickedCandidateIndex = ConditionPicker::PickCandidateIndex(
-				GraftSourcesExtracted.Num(),
-				DistributionConditions,
-				Sample,
-				RandomStream,
-				GetCandidateValue,
-				IsValidCandidate);
-
-			if (PickedCandidateIndex == INDEX_NONE)
-			{
-				continue;
-			}
-
-			const FManagedArrayCollection& PickedSourceGraft = GraftSourcesExtracted[PickedCandidateIndex];
+		FPVGraftAttributes Attributes{
+			.BranchIndex = AttachmentPoint.BranchIndex,
+			.PointIndex = AttachmentPointIndex,
+			.AttachmentPoint = AttachmentPoint.Position,
+			.UpVector = AttachmentPoint.UpDirection,
+			.NormalVector = AttachmentPoint.NormalDirection,
+			.Scale = AttachmentPoint.Scale,
+			.LengthFromRoot = AttachmentPoint.LengthFromRoot,
+			.bIsTipInstance = AttachmentPoint.bIsTipInstance
+		};
 
-			SetGraft(OutCollection, PickedSourceGraft, Attributes);
-		}
+		const FManagedArrayCollection& PickedSourceGraft = GraftSourcesExtracted[PickedCandidateIndex];
+		SetGraft(OutCollection, PickedSourceGraft, Attributes);
 	}
 }
 
@@ -1170,7 +843,7 @@ void FPVGrafter::DistributeGraftWithParametricSettings(FManagedArrayCollection&
 		CandidateGraftData[i] = MetaInfoFacade.GetGraftAttributes();
 	}
 
-	TArray<FAttachmentPoint> AttachmentPoints;
+	TArray<PV::DistributionHelper::FAttachmentPoint> AttachmentPoints;
 	ComputeAttachmentPoints(
 		OutCollection, 
 		NormalizedPointCaches, 
diff --git a/Engine/Plugins/Experimental/ProceduralVegetationEditor/Source/ProceduralVegetation/Private/Implementations/PVGrafter.h b/Engine/Plugins/Experimental/ProceduralVegetationEditor/Source/ProceduralVegetation/Private/Implementations/PVGrafter.h
--- a/Engine/Plugins/Experimental/ProceduralVegetationEditor/Source/ProceduralVegetation/Private/Implementations/PVGrafter.h
+++ b/Engine/Plugins/Experimental/ProceduralVegetationEditor/Source/ProceduralVegetation/Private/Implementations/PVGrafter.h
@@ -5,59 +5,17 @@
 #include "CoreMinimal.h"
 #include "Utils/PVAttributes.h"
 #include "GeometryCollection/ManagedArrayCollection.h"
+#include "Helpers/PVHormoneDistributionHelper.h"
 #include "Helpers/PVParametricDistributionHelper.h"
-#include "Helpers/PVPhyllotaxyHelper.h"
 
 using namespace PV::ParametricDistributionHelper;
+using namespace PV::HormoneDistributionHelper;
 
 struct FPVFoliageDistributionConditions;
 struct FPVFloatRamp;
 
 struct FPVGrafter
 {
-	struct FHormoneBasedSettings
-	{
-		struct FDistributionSettings
-		{
-			const float EthyleneThreshold;
-			const float InstanceSpacing;
-			const FPVFloatRamp* InstanceSpacingRamp;
-			const float InstanceSpacingRampEffect;
-			const int32 MaxPerBranch;
-		};
-
-		struct FPhyllotaxySettings
-		{
-			const bool ResetPhyllotaxy;
-			const EPhyllotaxyType PhyllotaxyType;
-			const EPhyllotaxyFormation PhyllotaxyFormation;
-			const int32 MinimumNodeBuds;
-			const int32 MaximumNodeBuds;
-			const float PhyllotaxyAdditionalAngle;
-			const float PhyllotaxyOffset;
-		};
-
-		struct FScaleSettings
-		{
-			const float BaseScale;
-			const float BranchScaleImpact;
-			const float MinScale;
-			const float MaxScale;
-			const float RandomScaleMin;
-			const float RandomScaleMax;
-			const FPVFloatRamp* ScaleRamp;
-		};
-
-		struct FVectorSettings
-		{
-			const bool OverrideAxilAngle;
-			const float AxilAngle;
-			const FPVFloatRamp* AxilAngleRamp;
-			const float AxilAngleRampUpperValue;
-			const float AxilAngleRampEffect;
-		};
-	};
-
 	struct FPVGraftAttributes
 	{
 		int32 BranchIndex;
@@ -96,10 +54,10 @@ struct FPVGrafter
 	static void DistributeGraftWithHormoneBasedSettings(
 		FManagedArrayCollection& OutCollection,
 		TArray<FManagedArrayCollection>& SourceGraftCollections,
-		const FHormoneBasedSettings::FDistributionSettings& DistributionSettings,
-		const FHormoneBasedSettings::FPhyllotaxySettings& PhyllotaxySettings,
-		const FHormoneBasedSettings::FScaleSettings& ScaleSettings,
-		const FHormoneBasedSettings::FVectorSettings& VectorSettings,
+		const FHormoneSettings::FDistributionSettings& DistributionSettings,
+		const FHormoneSettings::FPhyllotaxySettings& PhyllotaxySettings,
+		const FHormoneSettings::FScaleSettings& ScaleSettings,
+		const FHormoneSettings::FVectorSettings& VectorSettings,
 		int32 RandomSeed,
 		const FPVFoliageDistributionConditions& DistributionConditions);
 
diff --git a/Engine/Plugins/Experimental/ProceduralVegetationEditor/Source/ProceduralVegetation/Private/Nodes/PVFoliageDistributorSettings.cpp b/Engine/Plugins/Experimental/ProceduralVegetationEditor/Source/ProceduralVegetation/Private/Nodes/PVFoliageDistributorSettings.cpp
--- a/Engine/Plugins/Experimental/ProceduralVegetationEditor/Source/ProceduralVegetation/Private/Nodes/PVFoliageDistributorSettings.cpp
+++ b/Engine/Plugins/Experimental/ProceduralVegetationEditor/Source/ProceduralVegetation/Private/Nodes/PVFoliageDistributorSettings.cpp
@@ -11,6 +11,7 @@
 #include "GeometryCollection/ManagedArrayCollection.h"
 #include "Helpers/PCGHelpers.h"
 #include "Helpers/PCGSettingsHelpers.h"
+#include "Helpers/PVHormoneDistributionHelper.h"
  
 #define LOCTEXT_NAMESPACE "PVFoliageDistributorSettings"
  
@@ -136,16 +137,15 @@ bool FPVFoliageDistributorElement::ExecuteInternal(FPCGContext* InContext) const
 			
 			if (Settings->Mode == EPVFoliageDistributionSettingsMode::HormoneBasedSettings)
 			{
-				const PV::HormoneBasedSettings::FDistributionSettings DistributionSettings = {
+				const FHormoneSettings::FDistributionSettings DistributionSettings = {
 					.EthyleneThreshold = Settings->HormoneBasedSettings.EthyleneThreshold,
-					.OverrideDistribution = true,
 					.InstanceSpacing = Settings->HormoneBasedSettings.InstanceSpacing,
 					.InstanceSpacingRamp = &Settings->HormoneBasedSettings.InstanceSpacingRamp,
 					.InstanceSpacingRampEffect = Settings->HormoneBasedSettings.InstanceSpacingRampEffect,
 					.MaxPerBranch = Settings->HormoneBasedSettings.MaxPerBranch
 				};
  
-				const PV::HormoneBasedSettings::FScaleSettings ScaleSettings = {
+				const FHormoneSettings::FScaleSettings ScaleSettings = {
 					.BaseScale = Settings->HormoneBasedSettings.BaseScale,
 					.BranchScaleImpact = Settings->HormoneBasedSettings.BranchScaleImpact,
 					.MinScale = Settings->HormoneBasedSettings.MinScale,
@@ -155,7 +155,7 @@ bool FPVFoliageDistributorElement::ExecuteInternal(FPCGContext* InContext) const
 					.ScaleRamp = &Settings->HormoneBasedSettings.ScaleRamp
 				};
  
-				const PV::HormoneBasedSettings::FVectorSettings VectorSettings = {
+				const FHormoneSettings::FVectorSettings VectorSettings = {
 					.OverrideAxilAngle = Settings->HormoneBasedSettings.OverrideAxilAngle,
 					.AxilAngle = Settings->HormoneBasedSettings.AxilAngle,
 					.AxilAngleRamp = &Settings->HormoneBasedSettings.AxilAngleRamp,
@@ -163,13 +163,14 @@ bool FPVFoliageDistributorElement::ExecuteInternal(FPCGContext* InContext) const
 					.AxilAngleRampEffect = Settings->HormoneBasedSettings.AxilAngleRampEffect
 				};
  
-				const PV::HormoneBasedSettings::FPhyllotaxySettings PhyllotaxySettings = {
-					.OverridePhyllotaxy = Settings->HormoneBasedSettings.OverridePhyllotaxy,
+				const FHormoneSettings::FPhyllotaxySettings PhyllotaxySettings = {
 					.PhyllotaxyType = Settings->HormoneBasedSettings.PhyllotaxyType,
 					.PhyllotaxyFormation = Settings->HormoneBasedSettings.PhyllotaxyFormation,
 					.MinimumNodeBuds = Settings->HormoneBasedSettings.MinimumNodeBuds,
 					.MaximumNodeBuds = Settings->HormoneBasedSettings.MaximumNodeBuds,
-					.PhyllotaxyAdditionalAngle = Settings->HormoneBasedSettings.PhyllotaxyAdditionalAngle
+					.PhyllotaxyAdditionalAngle = Settings->HormoneBasedSettings.PhyllotaxyAdditionalAngle,
+					.ResetPhyllotaxy = Settings->HormoneBasedSettings.ResetPhyllotaxy,
+					.PhyllotaxyOffset = Settings->HormoneBasedSettings.PhyllotaxyOffset
 				};
 
 				FPVFoliage::DistributeFoliageWithHormoneBasedSettings(
diff --git a/Engine/Plugins/Experimental/ProceduralVegetationEditor/Source/ProceduralVegetation/Private/Nodes/PVFoliageDistributorSettings.h b/Engine/Plugins/Experimental/ProceduralVegetationEditor/Source/ProceduralVegetation/Private/Nodes/PVFoliageDistributorSettings.h
--- a/Engine/Plugins/Experimental/ProceduralVegetationEditor/Source/ProceduralVegetation/Private/Nodes/PVFoliageDistributorSettings.h
+++ b/Engine/Plugins/Experimental/ProceduralVegetationEditor/Source/ProceduralVegetation/Private/Nodes/PVFoliageDistributorSettings.h
@@ -8,8 +8,6 @@
 #include "Helpers/PVFoliageHelper.h"
 #include "Helpers/PVPhyllotaxyHelper.h"
 
-#include "Implementations/PVFoliage.h"
- 
 #include "Nodes/PVBaseSettings.h"
  
 #include "PVFoliageDistributorSettings.generated.h"
@@ -83,24 +81,27 @@ struct FPVFoliageDistributionHormoneBasedSettings
  
 	UPROPERTY(EditAnywhere, Category="Vector Settings", meta=(PCG_Overridable, EditCondition="OverrideAxilAngle", ClampMin=0.0f, ClampMax=1.0f, UIMin=0.0f, UIMax=1.0f, Tooltip="Adjusts this parameter for finer control.\n\nDefines the upper axil angle reached when AxilAngleRamp evaluates to 1. Together with AxilAngle, it forms the min–max range for ramp-driven orientation."))
 	float AxilAngleRampEffect = 0.0f;
+	
+	UPROPERTY(EditAnywhere, Category="Phyllotaxy Settings", meta=(PCG_Overridable, Tooltip="Reset phyllotaxy per branch.\n\nBranch orientation is used to align the initial phyllotaxy. This setting will reset it per branch so that its no longer tied to parent orientation."))
+	bool ResetPhyllotaxy = false;
  
-	UPROPERTY(VisibleAnywhere, Category="Phyllotaxy Settings", meta=(Tooltip="Use custom phyllotaxy settings instead of defaults.\n\nOveride the phyllotaxy settings loaded from the procedural vegetation preset."))
-	bool OverridePhyllotaxy = true;
- 
-	UPROPERTY(EditAnywhere, Category="Phyllotaxy Settings", meta=(PCG_Overridable, EditCondition="OverridePhyllotaxy", Tooltip="Pattern for arranging buds/leaves around the stem.\n\nSelects the phyllotactic pattern used for placement (e.g., Alternate, Opposite, Whorled, Spiral). This defines the sequence and angular spacing of buds/leaves around the stem and along its length."))
+	UPROPERTY(EditAnywhere, Category="Phyllotaxy Settings", meta=(PCG_Overridable, Tooltip="Pattern for arranging buds/leaves around the stem.\n\nSelects the phyllotactic pattern used for placement (e.g., Alternate, Opposite, Whorled, Spiral). This defines the sequence and angular spacing of buds/leaves around the stem and along its length."))
 	EPhyllotaxyType PhyllotaxyType = EPhyllotaxyType::Alternate;
  
-	UPROPERTY(EditAnywhere, Category="Phyllotaxy Settings", meta=(PCG_Overridable, EditCondition="OverridePhyllotaxy && PhyllotaxyType == EPhyllotaxyType::Spiral", EditConditionHides))
+	UPROPERTY(EditAnywhere, Category="Phyllotaxy Settings", meta=(PCG_Overridable, EditCondition="PhyllotaxyType == EPhyllotaxyType::Spiral", EditConditionHides))
 	EPhyllotaxyFormation PhyllotaxyFormation = EPhyllotaxyFormation::Distichous;
  
-	UPROPERTY(EditAnywhere, Category="Phyllotaxy Settings", meta=(PCG_Overridable, EditCondition="OverridePhyllotaxy", ClampMin=1, ClampMax=10, UIMin=1, UIMax=10, Tooltip="Minimum buds/leaves per node.\n\nSets the lower bound on the number of buds/leaves generated at each node. Guarantees at least this many instances (subject to spacing and other constraints). Should be ≤ MaximumNodeBuds. Applied when Whorled phyllotaxy is selected."))
+	UPROPERTY(EditAnywhere, Category="Phyllotaxy Settings", meta=(PCG_Overridable, ClampMin=1, ClampMax=10, UIMin=1, UIMax=10, Tooltip="Minimum buds/leaves per node.\n\nSets the lower bound on the number of buds/leaves generated at each node. Guarantees at least this many instances (subject to spacing and other constraints). Should be ≤ MaximumNodeBuds. Applied when Whorled phyllotaxy is selected."))
 	int32 MinimumNodeBuds = 2;
  
-	UPROPERTY(EditAnywhere, Category="Phyllotaxy Settings", meta=(PCG_Overridable, EditCondition="OverridePhyllotaxy", ClampMin=1, ClampMax=10, UIMin=1, UIMax=10, Tooltip="Maximum buds/leaves per node.\n\nCaps the number of buds/leaves spawned. Use with MinimumNodeBuds to bound density and keep arrangements predictable. Applied when Whorled phyllotaxy is selected."))
+	UPROPERTY(EditAnywhere, Category="Phyllotaxy Settings", meta=(PCG_Overridable, ClampMin=1, ClampMax=10, UIMin=1, UIMax=10, Tooltip="Maximum buds/leaves per node.\n\nCaps the number of buds/leaves spawned. Use with MinimumNodeBuds to bound density and keep arrangements predictable. Applied when Whorled phyllotaxy is selected."))
 	int32 MaximumNodeBuds = 3;
  
-	UPROPERTY(EditAnywhere, Category="Phyllotaxy Settings", meta=(PCG_Overridable, EditCondition="OverridePhyllotaxy", ClampMin=0.0f, ClampMax=360.0f, UIMin=0.0f, UIMax=360.0f, Tooltip="Extra divergence angle (°) added to the phyllotaxy pattern\n\nAdds an angular offset to the base phyllotactic divergence."))
+	UPROPERTY(EditAnywhere, Category="Phyllotaxy Settings", meta=(PCG_Overridable, ClampMin=0.0f, ClampMax=360.0f, UIMin=0.0f, UIMax=360.0f, Tooltip="Extra divergence angle (°) added to the phyllotaxy pattern\n\nAdds an angular offset to the base phyllotactic divergence."))
 	float PhyllotaxyAdditionalAngle = 0.0f;
+	
+	UPROPERTY(EditAnywhere, Category="Phyllotaxy Settings", meta=(PCG_Overridable, ClampMin=0.0f, ClampMax=360.0f, UIMin=0.0f, UIMax=360.0f, Tooltip="Initial offset value for the phyllotaxy (rotation around parent).\n\nValues are in degrees."))
+	float PhyllotaxyOffset = 0.0f;
 };
 
 USTRUCT()
diff --git a/Engine/Plugins/Experimental/ProceduralVegetationEditor/Source/ProceduralVegetation/Private/Nodes/PVGrafterSettings.cpp b/Engine/Plugins/Experimental/ProceduralVegetationEditor/Source/ProceduralVegetation/Private/Nodes/PVGrafterSettings.cpp
--- a/Engine/Plugins/Experimental/ProceduralVegetationEditor/Source/ProceduralVegetation/Private/Nodes/PVGrafterSettings.cpp
+++ b/Engine/Plugins/Experimental/ProceduralVegetationEditor/Source/ProceduralVegetation/Private/Nodes/PVGrafterSettings.cpp
@@ -174,7 +174,7 @@ bool FPVGrafterElement::ExecuteInternal(FPCGContext* InContext) const
 
 	if (Settings->Mode == EPVGrafterDistributionSettingsMode::HormoneBasedSettings)
 	{
-		const FPVGrafter::FHormoneBasedSettings::FDistributionSettings DistributionSettings = {
+		const FHormoneSettings::FDistributionSettings DistributionSettings = {
 			.EthyleneThreshold = Settings->HormoneBasedSettings.EthyleneThreshold,
 			.InstanceSpacing = Settings->HormoneBasedSettings.InstanceSpacing,
 			.InstanceSpacingRamp = &Settings->HormoneBasedSettings.InstanceSpacingRamp,
@@ -182,17 +182,17 @@ bool FPVGrafterElement::ExecuteInternal(FPCGContext* InContext) const
 			.MaxPerBranch = Settings->HormoneBasedSettings.MaxPerBranch
 		};
 		
-		const FPVGrafter::FHormoneBasedSettings::FPhyllotaxySettings PhyllotaxySettings = {
-			.ResetPhyllotaxy = Settings->HormoneBasedSettings.ResetPhyllotaxy,
+		const FHormoneSettings::FPhyllotaxySettings PhyllotaxySettings = {
 			.PhyllotaxyType = Settings->HormoneBasedSettings.PhyllotaxyType,
 			.PhyllotaxyFormation = Settings->HormoneBasedSettings.PhyllotaxyFormation,
 			.MinimumNodeBuds = Settings->HormoneBasedSettings.MinimumNodeBuds,
 			.MaximumNodeBuds = Settings->HormoneBasedSettings.MaximumNodeBuds,
 			.PhyllotaxyAdditionalAngle = Settings->HormoneBasedSettings.PhyllotaxyAdditionalAngle,
+			.ResetPhyllotaxy = Settings->HormoneBasedSettings.ResetPhyllotaxy,
 			.PhyllotaxyOffset = Settings->HormoneBasedSettings.PhyllotaxyOffset
 		}; 
 
-		const FPVGrafter::FHormoneBasedSettings::FScaleSettings ScaleSettings = {
+		const FHormoneSettings::FScaleSettings ScaleSettings = {
 			.BaseScale = Settings->HormoneBasedSettings.BaseScale,
 			.BranchScaleImpact = Settings->HormoneBasedSettings.BranchScaleImpact,
 			.MinScale = Settings->HormoneBasedSettings.MinScale,
@@ -202,7 +202,7 @@ bool FPVGrafterElement::ExecuteInternal(FPCGContext* InContext) const
 			.ScaleRamp = &Settings->HormoneBasedSettings.ScaleRamp
 		};
 
-		const FPVGrafter::FHormoneBasedSettings::FVectorSettings VectorSettings = {
+		const FHormoneSettings::FVectorSettings VectorSettings = {
 			.OverrideAxilAngle = Settings->HormoneBasedSettings.OverrideAxilAngle,
 			.AxilAngle = Settings->HormoneBasedSettings.AxilAngle,
 			.AxilAngleRamp = &Settings->HormoneBasedSettings.AxilAngleRamp,
