[PVE] Crash fix on Branch Removal in Skeleton Edit
Issue: StartIndex was being set as StartIndex+1 which was an invalid array index in BranchPoints
Steps to reproduce: In SkeletonEdit Node, try to remove a last point on a branch, in BranchRemoval mode

#Synced-CL 51373704

diff --git a/Engine/Plugins/Experimental/ProceduralVegetationEditor/Source/ProceduralVegetationEditor/Private/Tools/NodeTools/PVSkeletonEditTool.cpp b/Engine/Plugins/Experimental/ProceduralVegetationEditor/Source/ProceduralVegetationEditor/Private/Tools/NodeTools/PVSkeletonEditTool.cpp
--- a/Engine/Plugins/Experimental/ProceduralVegetationEditor/Source/ProceduralVegetationEditor/Private/Tools/NodeTools/PVSkeletonEditTool.cpp
+++ b/Engine/Plugins/Experimental/ProceduralVegetationEditor/Source/ProceduralVegetationEditor/Private/Tools/NodeTools/PVSkeletonEditTool.cpp
@@ -587,9 +587,25 @@ void UPVSkeletonEditTool::OnClicked(const FInputDeviceRay& ClickPos)
 				SkeletonEditSettings->SkeletonEditSettings.BranchRemovalParams.BranchRemovalMode == EBranchRemovalMode::TrimBranch
 				? SkeletonPointProxy->BranchPointIndex
 				: 0;
-			StartIndex = StartIndex == 0
-				? 0
-				: StartIndex + 1;
+
+			bool bNextPointRemoved = false;
+			if (BranchPoints.IsValidIndex(StartIndex + 1))
+			{
+				const int32 NextPointBudNumber = PointFacade.GetBudNumber(BranchPoints[StartIndex + 1]);
+				bNextPointRemoved = RemovedPoints[NextPointBudNumber];
+			}
+			if (StartIndex > 0 && StartIndex < BranchPoints.Num() - 1 && !bNextPointRemoved)
+			{
+				// Remove from the next point onwards from currently selected point
+				// unless it's the first or the last point, we want to remove the first or last point if selected
+				StartIndex += 1;
+			}
+			if (StartIndex <= 1)
+			{
+				// Remove the complete branch if first point was selected (can't have a branch with 1 point)
+				StartIndex = 0;
+			}
+			
 			for (int32 BranchPointIndex = StartIndex; BranchPointIndex < BranchPoints.Num(); BranchPointIndex++)
 			{
 				const int32 PointIndex = BranchPoints[BranchPointIndex];
