[PVE] Added use Splitpoint control on grower

#Synced-CL 50970602

diff --git a/Engine/Plugins/Experimental/ProceduralVegetationEditor/Source/ProceduralVegetation/Private/Implementations/PVGrower.cpp b/Engine/Plugins/Experimental/ProceduralVegetationEditor/Source/ProceduralVegetation/Private/Implementations/PVGrower.cpp
--- a/Engine/Plugins/Experimental/ProceduralVegetationEditor/Source/ProceduralVegetation/Private/Implementations/PVGrower.cpp
+++ b/Engine/Plugins/Experimental/ProceduralVegetationEditor/Source/ProceduralVegetation/Private/Implementations/PVGrower.cpp
@@ -1,6 +1,7 @@
 // Copyright Epic Games, Inc. All Rights Reserved.
 
 #include "PVGrower.h"
+
 #include <random>
 #include "ProceduralVegetationModule.h"
 #include "Algo/Accumulate.h"
@@ -77,26 +78,32 @@ void FPVGrowerParams::GetRandomAngle(FPVRandomAngleParams& OutRandomAngle, int G
 	OutRandomAngle = RandomAngle;
 }
 
-void FPVGrower::Grow(const FPVSeedPoint& SeedPoint, const FPVGrowerParams& InGrowerParams, FManagedArrayCollection& InCollection)
+void FPVGrower::Grow(const FPVSeedPoint& SeedPoint, const FPVGrowerParams& InGrowerParams, FManagedArrayCollection& InCollection, bool bUseSplitPoints)
 {
 	TRACE_CPUPROFILER_EVENT_SCOPE(PV::Grower::Grow);
 	TObjectPtr<UPVGrowerData> Skeleton = NewObject<UPVGrowerData>(GetTransientPackage(), UPVGrowerData::StaticClass(),FName("Skeleton"), RF_Public | RF_Transactional);
 	InitializeGrowth(InGrowerParams, SeedPoint, Skeleton);
 	IterateGrowthCycles(InGrowerParams, Skeleton);
 	ConvertPointsToCentimeters(Skeleton);
-	SplitBranchSourcePoints(Skeleton);
+	if (bUseSplitPoints)
+	{
+		SplitBranchSourcePoints(Skeleton);
+	}
 	FillCollection(InCollection, Skeleton);
 	//PrintPoints();
 }
 
-void FPVGrower::Grow(const FManagedArrayCollection& InCollection, const FPVGrowerParams& InGrowerParams, FManagedArrayCollection& Collection)
+void FPVGrower::Grow(const FManagedArrayCollection& InCollection, const FPVGrowerParams& InGrowerParams, FManagedArrayCollection& Collection, bool bUseSplitPoints)
 {
 	TRACE_CPUPROFILER_EVENT_SCOPE(PV::Grower::Grow);
 	TObjectPtr<UPVGrowerData> Skeleton = NewObject<UPVGrowerData>(GetTransientPackage(), UPVGrowerData::StaticClass(),FName("Skeleton"), RF_Public | RF_Transactional);
-	InitializeGrowth(InCollection, Skeleton);
+	InitializeGrowth(InCollection, Skeleton, bUseSplitPoints);
 	IterateGrowthCycles(InGrowerParams, Skeleton);
 	ConvertPointsToCentimeters(Skeleton);
-	SplitBranchSourcePoints(Skeleton);
+	if (bUseSplitPoints)
+	{
+		SplitBranchSourcePoints(Skeleton);
+	}
 	FillCollection(Collection, Skeleton);
 	//PrintPoints();
 }
@@ -229,10 +236,10 @@ void FPVGrower::FillCollection(FManagedArrayCollection& Collection, UPVGrowerDat
 	check(BranchFacade.IsValid());
 }
 
-void FPVGrower::FillFromCollection(const FManagedArrayCollection& Collection, UPVGrowerData* Skeleton)
+void FPVGrower::FillFromCollection(const FManagedArrayCollection& Collection, UPVGrowerData* Skeleton, bool bUseSplitPoints)
 {
 	FillPointsFromCollection(Collection, Skeleton);
-	FillBranchesFromCollection(Collection, Skeleton);
+	FillBranchesFromCollection(Collection, Skeleton, bUseSplitPoints);
 }
 
 void FPVGrower::FillPointsFromCollection(const FManagedArrayCollection& Collection, UPVGrowerData* Skeleton)
@@ -300,7 +307,7 @@ void FPVGrower::FillPointsFromCollection(const FManagedArrayCollection& Collecti
 	}
 }
 
-void FPVGrower::FillBranchesFromCollection(const FManagedArrayCollection& Collection, UPVGrowerData* Skeleton)
+void FPVGrower::FillBranchesFromCollection(const FManagedArrayCollection& Collection, UPVGrowerData* Skeleton, bool bUseSplitPoints /*= true*/)
 {
 	PV::Facades::FBranchFacade BranchFacade = PV::Facades::FBranchFacade(Collection);
 
@@ -335,7 +342,10 @@ void FPVGrower::FillBranchesFromCollection(const FManagedArrayCollection& Collec
 
 		for (int j=0; j< BranchPoints.Num(); j++)
 		{
-			SwapSplitPointWithFusedPoint(BranchPoints, j, BranchSourceBudNumber);
+			if (bUseSplitPoints)
+			{
+				SwapSplitPointWithFusedPoint(BranchPoints, j, BranchSourceBudNumber);
+			}
 			
 			UPVGrowerPoint* Point = Skeleton->GetPointFromIndex(BranchPoints[j]);
 			if (Point)
@@ -385,11 +395,14 @@ void FPVGrower::FillBranchesFromCollection(const FManagedArrayCollection& Collec
 		Skeleton->AddPrimitive(Primitive, BranchNumber);
 	}
 
-	for (UPVGrowerPoint* Point : SplitPoints)
+	if (bUseSplitPoints)
 	{
-		if (Point)
+		for (UPVGrowerPoint* Point : SplitPoints)
 		{
-			Skeleton->RemovePoint(Point);
+			if (Point)
+			{
+				Skeleton->RemovePoint(Point);
+			}
 		}
 	}
 }
@@ -421,12 +434,12 @@ void FPVGrower::InitializeGrowth(const FPVGrowerParams& InGrowerParams, const FP
 	//Calculate MinRootDist from seed points
 }
 
-void FPVGrower::InitializeGrowth(const FManagedArrayCollection& InCollection, UPVGrowerData* Skeleton)
+void FPVGrower::InitializeGrowth(const FManagedArrayCollection& InCollection, UPVGrowerData* Skeleton, bool bUseSplitPoints)
 {
 	TRACE_CPUPROFILER_EVENT_SCOPE(PV::Grower::InitializeGrowth);
 	check(Skeleton);
 
-	FillFromCollection(InCollection, Skeleton);
+	FillFromCollection(InCollection, Skeleton, bUseSplitPoints);
 }
 
 void FPVGrower::IterateGrowthCycles(const FPVGrowerParams& InGrowerParams, UPVGrowerData* Skeleton)
@@ -463,16 +476,18 @@ void FPVGrower::SplitBranchSourcePoints(UPVGrowerData* Skeleton)
 {
 	for (auto Primitive : Skeleton->Primitives)
 	{
-		check(Primitive->BranchBuds.IsValidIndex(0));
-		UPVGrowerPoint* Point = Skeleton->GetPoint(Primitive->BranchBuds[0]);
-		if (Point)
+		if (Primitive->BranchBuds.IsValidIndex(0))
 		{
-			if (Point->Bud.Status.Seed != 1)
+			UPVGrowerPoint* Point = Skeleton->GetPoint(Primitive->BranchBuds[0]);
+			if (Point)
 			{
-				UPVGrowerPoint* SplitPoint = DuplicateObject<UPVGrowerPoint>(Point, Skeleton);
-				SplitPoint->Bud = Point->Bud;
-				Skeleton->AddPoint(SplitPoint, Primitive, false);
-				Primitive->BranchBuds[0] = Skeleton->MaxBudNumber;
+				if (Point->Bud.Status.Seed != 1)
+				{
+					UPVGrowerPoint* SplitPoint = DuplicateObject<UPVGrowerPoint>(Point, Skeleton);
+					SplitPoint->Bud = Point->Bud;
+					Skeleton->AddPoint(SplitPoint, Primitive, false);
+					Primitive->BranchBuds[0] = Skeleton->MaxBudNumber;
+				}
 			}
 		}
 	}
diff --git a/Engine/Plugins/Experimental/ProceduralVegetationEditor/Source/ProceduralVegetation/Private/Implementations/PVGrower.h b/Engine/Plugins/Experimental/ProceduralVegetationEditor/Source/ProceduralVegetation/Private/Implementations/PVGrower.h
--- a/Engine/Plugins/Experimental/ProceduralVegetationEditor/Source/ProceduralVegetation/Private/Implementations/PVGrower.h
+++ b/Engine/Plugins/Experimental/ProceduralVegetationEditor/Source/ProceduralVegetation/Private/Implementations/PVGrower.h
@@ -15,23 +15,23 @@ struct FPVPointGravityAttribute;
 
 struct FPVGrower
 {
-	static void Grow(const FPVSeedPoint& SeedPoint, const FPVGrowerParams& InGrowerParams, FManagedArrayCollection& InCollection);
-	
-	static void Grow(const FManagedArrayCollection& InCollection, const FPVGrowerParams& InGrowerParams, FManagedArrayCollection& Collection);
+	static void Grow(const FPVSeedPoint& SeedPoint, const FPVGrowerParams& InGrowerParams, FManagedArrayCollection& InCollection, bool bUseSplitPoints = true);
+
+	static void Grow(const FManagedArrayCollection& InCollection, const FPVGrowerParams& InGrowerParams, FManagedArrayCollection& Collection, bool bUseSplitPoints = true);
 
 	static void FillCollection(FManagedArrayCollection& Collection, UPVGrowerData* Skeleton);
-	
-	static void FillFromCollection(const FManagedArrayCollection& Collection, UPVGrowerData* Skeleton);
-	
+
+	static void FillFromCollection(const FManagedArrayCollection& Collection, UPVGrowerData* Skeleton, bool bUseSplitPoints = true);
+
 	static void FillPointsFromCollection(const FManagedArrayCollection& Collection, UPVGrowerData* Skeleton);
-	
-	static void FillBranchesFromCollection(const FManagedArrayCollection& Collection, UPVGrowerData* Skeleton);
+
+	static void FillBranchesFromCollection(const FManagedArrayCollection& Collection, UPVGrowerData* Skeleton, bool bUseSplitPoints = true);
 
 	static float RandomValue(uint32 Seed);
 
 	static void InitializeGrowth(const FPVGrowerParams& InGrowerParams, const FPVSeedPoint& SeedPoint, UPVGrowerData* Skeleton);
 	
-	static void InitializeGrowth(const FManagedArrayCollection& InCollection, UPVGrowerData* Skeleton);
+	static void InitializeGrowth(const FManagedArrayCollection& InCollection, UPVGrowerData* Skeleton, bool bUseSplitPoints = true);
 
 	static void IterateGrowthCycles(const FPVGrowerParams& InGrowerParams, UPVGrowerData* Skeleton);
 
diff --git a/Engine/Plugins/Experimental/ProceduralVegetationEditor/Source/ProceduralVegetation/Private/Nodes/PVGrowerSettings.cpp b/Engine/Plugins/Experimental/ProceduralVegetationEditor/Source/ProceduralVegetation/Private/Nodes/PVGrowerSettings.cpp
--- a/Engine/Plugins/Experimental/ProceduralVegetationEditor/Source/ProceduralVegetation/Private/Nodes/PVGrowerSettings.cpp
+++ b/Engine/Plugins/Experimental/ProceduralVegetationEditor/Source/ProceduralVegetation/Private/Nodes/PVGrowerSettings.cpp
@@ -287,7 +287,7 @@ bool FPVGrowerElement::ExecuteInternal(FPCGContext* InContext) const
 
 				FManagedArrayCollection OutCollection;
 		
-				FPVGrower::Grow(SkeletonCollection, GrowerParams, OutCollection);
+				FPVGrower::Grow(SkeletonCollection, GrowerParams, OutCollection, Settings->bUseSplitPoints);
 			
 				UPVGrowthData* OutManagedArrayCollectionData = FPCGContext::NewObject_AnyThread<UPVGrowthData>(InContext);
 				OutManagedArrayCollectionData->Initialize(MoveTemp(OutCollection));
@@ -304,7 +304,7 @@ bool FPVGrowerElement::ExecuteInternal(FPCGContext* InContext) const
 			FManagedArrayCollection OutCollection;
 		
 			FPVSeedPoint SeedPoint;
-			FPVGrower::Grow(SeedPoint, GrowerParams, OutCollection);
+			FPVGrower::Grow(SeedPoint, GrowerParams, OutCollection, Settings->bUseSplitPoints);
 			
 			UPVGrowthData* OutManagedArrayCollectionData = FPCGContext::NewObject_AnyThread<UPVGrowthData>(InContext);
 			OutManagedArrayCollectionData->Initialize(MoveTemp(OutCollection));
diff --git a/Engine/Plugins/Experimental/ProceduralVegetationEditor/Source/ProceduralVegetation/Private/Nodes/PVGrowerSettings.h b/Engine/Plugins/Experimental/ProceduralVegetationEditor/Source/ProceduralVegetation/Private/Nodes/PVGrowerSettings.h
--- a/Engine/Plugins/Experimental/ProceduralVegetationEditor/Source/ProceduralVegetation/Private/Nodes/PVGrowerSettings.h
+++ b/Engine/Plugins/Experimental/ProceduralVegetationEditor/Source/ProceduralVegetation/Private/Nodes/PVGrowerSettings.h
@@ -48,6 +48,9 @@ public:
 	UFUNCTION(CallInEditor, Category = Preset, meta = (DisplayName = "Apply to selected preset"))
 	void ApplyToSelectedPreset() const;
 	
+	UPROPERTY(EditAnywhere, Category = Settings)
+	bool bUseSplitPoints = false;
+
 	UPROPERTY(EditAnywhere, Category = Settings, meta=( ShowOnlyInnerProperties, PCG_Overridable))
 	FPVGrowerParams GrowerParams;
 };
