Fix incorrect point scale smoothing when importing trees via the 2D/3D importers in the Procedural Vegetation Editor
Replace the interpolation algorithm with a taper-rate-aware smoothing pass that rejects noisy scale dips and correctly handles fused-point branch topologies
Add a helper template to validate that attribute collections share consistent group sizes, and refactor all existing IsValid checks to use it

#Synced-CL 51620377

diff --git a/Engine/Plugins/Experimental/ProceduralVegetationEditor/Source/ProceduralVegetation/Private/Helpers/PVAttributesHelper.cpp b/Engine/Plugins/Experimental/ProceduralVegetationEditor/Source/ProceduralVegetation/Private/Helpers/PVAttributesHelper.cpp
--- a/Engine/Plugins/Experimental/ProceduralVegetationEditor/Source/ProceduralVegetation/Private/Helpers/PVAttributesHelper.cpp
+++ b/Engine/Plugins/Experimental/ProceduralVegetationEditor/Source/ProceduralVegetation/Private/Helpers/PVAttributesHelper.cpp
@@ -564,76 +564,102 @@ void PV::AttributesHelper::ComputeBranchHierarchyNumbers(const FComputeBranchHie
 	);
 }
 
-void PV::AttributesHelper::InterpolatePointsScales(
-	PV::FPointScaleAttributeView PointScaleAttribute,
-	PV::FPointLengthFromRootAttributeConstView PointLengthFromRootAttribute,
-	PV::FBranchPointsAttributeConstView BranchPointsAttribute,
-	PV::FBranchParentNumberAttributeConstView BranchParentNumberAttribute,
-	PV::FBranchChildrenAttributeConstView BranchChildrenAttribute,
-	PV::FBranchNumberAttributeConstView BranchNumberAttribute
-)
+void PV::AttributesHelper::ComputeSmoothedPointScales(const FComputeSmoothedPointScalesAttributes& InAttributes)
 {
 	using namespace PV::PlantTraversalHelper;
 
-	RecursiveWalkBranches(
-		BranchParentNumberAttribute,
-		BranchChildrenAttribute,
-		BranchNumberAttribute,
-		[&](int32 BranchIndex)->EForEachResult
+	const bool bHasFusedPoints = PV::PlantTraversalHelper::HasFusedPoints(InAttributes.BranchPoints, InAttributes.BranchParentNumber, InAttributes.BranchNumber);
+
+	RecursiveWalkBranches<int32>(
+		InAttributes.BranchParentNumber,
+		InAttributes.BranchChildren,
+		InAttributes.BranchNumber,
+		INDEX_NONE,
+		[&](int32 BranchIndex, int32& ParentBranch)->EForEachResult
 		{
-			const TArray<int32>& BranchPoints = BranchPointsAttribute[BranchIndex];
+			const TArray<int32>& BranchPoints = InAttributes.BranchPoints[BranchIndex];
+			const bool bIsRootBranch = ParentBranch == INDEX_NONE;
+			const bool bShouldSkipFirstPoint = bHasFusedPoints && !bIsRootBranch;
 
-			float CurrentPointScale = PointScaleAttribute[BranchPoints.Last()];
-			float CurrentLengthFromRoot = PointLengthFromRootAttribute[BranchPoints.Last()];
+			int32 StartIndex = bShouldSkipFirstPoint ? 1 : 0;
 
-			for (int32 i = BranchPoints.Num() - 1; i >= 0; --i)
+			float CurrentPointScale = InAttributes.PointScale[BranchPoints[StartIndex]];
+			float CurrentLengthFromRoot = InAttributes.PointLengthFromRoot[BranchPoints[StartIndex]];
+			if (CurrentPointScale <= 0)
 			{
-				float NextValidScale = -1;
-				float NextValidScale_LengthFromRoot = -1;
-				int32 NextValidScale_Index = INDEX_NONE;
-
-				for (int32 j = i - 1; j >= 0; --j)
+				for (int32 i = StartIndex; i < BranchPoints.Num(); ++i)
 				{
-					const int32 PointIndex = BranchPoints[j];
-					const float PointScale = PointScaleAttribute[PointIndex];
+					const float PointScale = InAttributes.PointScale[BranchPoints[i]];
 					if (PointScale > 0)
 					{
-						NextValidScale = PointScale;
-						NextValidScale_LengthFromRoot = PointLengthFromRootAttribute[PointIndex];
-						NextValidScale_Index = j;
+						CurrentPointScale = PointScale;
+						CurrentLengthFromRoot = InAttributes.PointLengthFromRoot[BranchPoints[i]];
+						for (int32 j = StartIndex; j < i; ++j)
+						{
+							InAttributes.PointScale[BranchPoints[j]] = CurrentPointScale;
+						}
+						StartIndex = i;
 						break;
 					}
 				}
+			}
 
-				if (NextValidScale_Index == INDEX_NONE)
+			const float TipTargetScale = 0.01f;
+			const float TipLengthFromRoot = InAttributes.PointLengthFromRoot[BranchPoints.Last()];
+			constexpr float MaxTaperRateMultiplier = 3.0f;
+
+			int32 BranchPointIndex = StartIndex;
+			while (BranchPointIndex < BranchPoints.Num() - 1)
+			{
+				float NextValidScale = -1;
+				float NextValidScale_LengthFromRoot = -1;
+				int32 NextValidScale_Index = INDEX_NONE;
+
+				for (int32 i = BranchPointIndex + 1; i < BranchPoints.Num(); ++i)
 				{
-					for (int32 j = i - 1; j >= 0; --j)
+					const int32 PointIndex = BranchPoints[i];
+					const float PointScale = InAttributes.PointScale[PointIndex];
+					if (PointScale > 0 && PointScale <= CurrentPointScale)
 					{
-						PointScaleAttribute[BranchPoints[j]] = CurrentPointScale;
+						const float CandidateLengthFromRoot = InAttributes.PointLengthFromRoot[PointIndex];
+						const float RemainingLength = FMath::Max(TipLengthFromRoot - CurrentLengthFromRoot, UE_SMALL_NUMBER);
+						const float DistToCandidate = FMath::Max(CandidateLengthFromRoot - CurrentLengthFromRoot, UE_SMALL_NUMBER);
+						const float ExpectedTaperRate = (CurrentPointScale - TipTargetScale) / RemainingLength;
+						const float ImpliedTaperRate = (CurrentPointScale - PointScale) / DistToCandidate;
+						if (ImpliedTaperRate <= MaxTaperRateMultiplier * ExpectedTaperRate)
+						{
+							NextValidScale = PointScale;
+							NextValidScale_LengthFromRoot = CandidateLengthFromRoot;
+							NextValidScale_Index = i;
+							break;
+						}
 					}
-					break;
 				}
 
-				if (NextValidScale_Index == i - 1)
+				if (NextValidScale_Index == INDEX_NONE)
 				{
-					CurrentPointScale = PointScaleAttribute[BranchPoints[NextValidScale_Index]];
-					continue;
+					NextValidScale = 0.01f;
+					NextValidScale_Index = BranchPoints.Num() - 1;
+					NextValidScale_LengthFromRoot = InAttributes.PointLengthFromRoot[BranchPoints[NextValidScale_Index]];
+					InAttributes.PointScale[BranchPoints[NextValidScale_Index]] = NextValidScale;
 				}
 
-				for (int32 j = i - 1; j >= NextValidScale_Index; --j)
+				for (int32 i = BranchPointIndex + 1; i < NextValidScale_Index; ++i)
 				{
-					const int32 PointIndex = BranchPoints[j];
-					const float PointLengthFromRoot = PointLengthFromRootAttribute[PointIndex];
+					const int32 PointIndex = BranchPoints[i];
+					const float PointLengthFromRoot = InAttributes.PointLengthFromRoot[PointIndex];
 					const float Alpha = FMath::GetMappedRangeValueClamped(FVector2D(CurrentLengthFromRoot, NextValidScale_LengthFromRoot), FVector2D(0, 1), PointLengthFromRoot);
 					const float PointScale = FMath::Lerp(CurrentPointScale, NextValidScale, Alpha);
-					PointScaleAttribute[PointIndex] = PointScale;
+					InAttributes.PointScale[PointIndex] = PointScale;
 				}
 
 				CurrentPointScale = NextValidScale;
 				CurrentLengthFromRoot = NextValidScale_LengthFromRoot;
-				i = NextValidScale_Index + 1;
+				BranchPointIndex = NextValidScale_Index;
 			}
 
+			ParentBranch = BranchIndex;
+
 			return EForEachResult::Continue;
 		}
 	);
@@ -741,14 +767,14 @@ bool PV::AttributesHelper::RecomputeAllGrowthDataAttributes(const FRecomputeGrow
 
 	ComputeBudLightDetected(InAttributes.BudLightDetected);
 
-	InterpolatePointsScales(
+	ComputeSmoothedPointScales({
 		InAttributes.PointScale,
 		InAttributes.PointLengthFromRoot,
 		InAttributes.BranchPoints,
 		InAttributes.BranchParentNumber,
 		InAttributes.BranchChildren,
 		InAttributes.BranchNumber
-	);
+	});
 
 	ComputeBudLateralMeristem({
 		InAttributes.BudLateralMeristem,
diff --git a/Engine/Plugins/Experimental/ProceduralVegetationEditor/Source/ProceduralVegetation/Private/Helpers/PVAttributesHelper.h b/Engine/Plugins/Experimental/ProceduralVegetationEditor/Source/ProceduralVegetation/Private/Helpers/PVAttributesHelper.h
--- a/Engine/Plugins/Experimental/ProceduralVegetationEditor/Source/ProceduralVegetation/Private/Helpers/PVAttributesHelper.h
+++ b/Engine/Plugins/Experimental/ProceduralVegetationEditor/Source/ProceduralVegetation/Private/Helpers/PVAttributesHelper.h
@@ -5,6 +5,38 @@
 
 namespace PV::AttributesHelper
 {
+	template<typename... TAttributeViews>
+	bool ValidateAttributeCollection(TAttributeViews... Attributes)
+	{
+		TMap<FName, int32> GroupSizes;
+		GroupSizes.Reserve(sizeof...(TAttributeViews));
+
+		bool bIsValid = true;
+		auto AddEntry = [&GroupSizes, &bIsValid](auto Attr)
+		{
+			if (!Attr.IsValid())
+			{
+				bIsValid = false;
+				return;
+			}
+
+			const FName GroupName = std::remove_reference_t<decltype(Attr)>::GroupName;
+			const int32 Size = Attr.Num();
+			if (const int32* ExistingSize = GroupSizes.Find(GroupName))
+			{
+				bIsValid &= (*ExistingSize == Size);
+			}
+			else
+			{
+				GroupSizes.Add(GroupName, Size);
+			}
+		};
+
+		(AddEntry(Attributes), ...);
+
+		return bIsValid;
+	}
+
 	struct FComputeBudNumbersAttributes
 	{
 		FPointBudNumberAttributeView PointBudNumber;
@@ -28,9 +60,7 @@ namespace PV::AttributesHelper
 
 		bool IsValid() const
 		{
-			return PointBudNumber.IsValid()
-				&& BranchSourceBudNumber.IsValid()
-				&& BranchPoints.IsValid();
+			return ValidateAttributeCollection(PointBudNumber, BranchSourceBudNumber, BranchPoints);
 		}
 	};
 
@@ -70,11 +100,7 @@ namespace PV::AttributesHelper
 
 		bool IsValid() const
 		{
-			return PointNjordPixelIndex.IsValid()
-				&& BranchPoints.IsValid()
-				&& BranchParentNumber.IsValid()
-				&& BranchNumber.IsValid()
-				&& BranchChildren.IsValid();
+			return ValidateAttributeCollection(PointNjordPixelIndex, BranchPoints, BranchParentNumber, BranchNumber, BranchChildren);
 		}
 	};
 
@@ -122,13 +148,7 @@ namespace PV::AttributesHelper
 
 		bool IsValid() const
 		{
-			return PointLengthFromRoot.IsValid()
-				&& PointLengthFromSeed.IsValid()
-				&& PointPosition.IsValid()
-				&& BranchPoints.IsValid()
-				&& BranchParentNumber.IsValid()
-				&& BranchNumber.IsValid()
-				&& BranchChildren.IsValid();
+			return ValidateAttributeCollection(PointLengthFromRoot, PointLengthFromSeed, PointPosition, BranchPoints, BranchParentNumber, BranchNumber, BranchChildren);
 		}
 	};
 
@@ -170,12 +190,7 @@ namespace PV::AttributesHelper
 
 		bool IsValid() const
 		{
-			return BudDevelopment.IsValid()
-				&& PointPlantGradient.IsValid()
-				&& BranchPoints.IsValid()
-				&& BranchParentNumber.IsValid()
-				&& BranchParents.IsValid()
-				&& BranchNumber.IsValid();
+			return ValidateAttributeCollection(BudDevelopment, PointPlantGradient, BranchPoints, BranchParentNumber, BranchParents, BranchNumber);
 		}
 	};
 
@@ -221,13 +236,7 @@ namespace PV::AttributesHelper
 
 		bool IsValid() const
 		{
-			return BudDirection.IsValid()
-				&& PointPosition.IsValid()
-				&& PointLengthFromRoot.IsValid()
-				&& BranchChildren.IsValid()
-				&& BranchParentNumber.IsValid()
-				&& BranchNumber.IsValid()
-				&& BranchPoints.IsValid();
+			return ValidateAttributeCollection(BudDirection, PointPosition, PointLengthFromRoot, BranchChildren, BranchParentNumber, BranchNumber, BranchPoints);
 		}
 	};
 
@@ -241,7 +250,7 @@ namespace PV::AttributesHelper
 
 		bool IsValid() const
 		{
-			return BudHormoneLevels && PointPlantGradient;
+			return ValidateAttributeCollection(BudHormoneLevels, PointPlantGradient);
 		}
 	};
 
@@ -300,16 +309,7 @@ namespace PV::AttributesHelper
 
 		bool IsValid() const
 		{
-			return PointGroundGradient
-				&& PointScaleGradient.IsValid()
-				&& PointPlantGradient.IsValid()
-				&& PointBranchGradient.IsValid()
-				&& PointLengthFromRoot.IsValid()
-				&& PointScale.IsValid()
-				&& BranchPoints.IsValid()
-				&& BranchParentNumber.IsValid()
-				&& BranchNumber.IsValid()
-				&& BranchChildren.IsValid();
+			return ValidateAttributeCollection(PointGroundGradient, PointScaleGradient, PointPlantGradient, PointBranchGradient, PointLengthFromRoot, PointScale, BranchPoints, BranchParentNumber, BranchNumber, BranchChildren);
 		}
 	};
 
@@ -341,7 +341,7 @@ namespace PV::AttributesHelper
 		{
 		}
 
-		bool IsValid() const { return PointGroundGradient && PointLengthFromRoot; }
+		bool IsValid() const { return ValidateAttributeCollection(PointGroundGradient, PointLengthFromRoot); }
 	};
 
 	PROCEDURALVEGETATION_API extern void ComputePointGroundGradient(const FComputePointGroundGradientAttributes& InAttributes);
@@ -371,7 +371,7 @@ namespace PV::AttributesHelper
 		{
 		}
 
-		bool IsValid() const { return PointScaleGradient && PointScale; }
+		bool IsValid() const { return ValidateAttributeCollection(PointScaleGradient, PointScale); }
 	};
 
 	PROCEDURALVEGETATION_API extern void ComputePointScaleGradient(const FComputePointScaleGradientAttributes& InAttributes);
@@ -423,12 +423,7 @@ namespace PV::AttributesHelper
 
 		bool IsValid() const
 		{
-			return PointPlantGradient
-				&& PointBranchGradient
-				&& BranchPoints
-				&& BranchParentNumber
-				&& BranchNumber
-				&& BranchChildren;
+			return ValidateAttributeCollection(PointPlantGradient, PointBranchGradient, BranchPoints, BranchParentNumber, BranchNumber, BranchChildren);
 		}
 	};
 
@@ -470,12 +465,7 @@ namespace PV::AttributesHelper
 
 		bool IsValid() const
 		{
-			return BudStatus
-				&& BranchPoints
-				&& BranchParentNumber
-				&& BranchNumber
-				&& BranchChildren
-				&& PointLengthFromRoot;
+			return ValidateAttributeCollection(BudStatus, BranchPoints, BranchParentNumber, BranchNumber, BranchChildren, PointLengthFromRoot);
 		}
 	};
 
@@ -531,15 +521,7 @@ namespace PV::AttributesHelper
 
 		bool IsValid() const
 		{
-			return BudLateralMeristem
-				&& PointPosition
-				&& PointScale
-				&& PointLengthFromRoot
-				&& BudDirection
-				&& BranchPoints
-				&& BranchParentNumber
-				&& BranchChildren
-				&& BranchNumber;
+			return ValidateAttributeCollection(BudLateralMeristem, PointPosition, PointScale, PointLengthFromRoot, BudDirection, BranchPoints, BranchParentNumber, BranchChildren, BranchNumber);
 		}
 	};
 
@@ -578,27 +560,57 @@ namespace PV::AttributesHelper
 
 		bool IsValid() const
 		{
-			return BranchParentNumber.IsValid()
-				&& BranchNumber.IsValid()
-				&& BranchChildren.IsValid()
-				&& BranchHierarchyNumber.IsValid();
+			return ValidateAttributeCollection(BranchParentNumber, BranchNumber, BranchChildren, BranchHierarchyNumber);
 		}
 	};
 
 	// Re-computes the BranchHierarchyNumber attribute given an already existing branch hierarchy.
 	PROCEDURALVEGETATION_API extern void ComputeBranchHierarchyNumbers(const FComputeBranchHierarchyNumberAttributes& InAttributes);
 
-	// This function ensures all points have a valid point scale by walking the point hierarchy
-	// and finding points with missing point scales. It then interpolates the missing point scales
-	// using the point scale of the parent/child points.
-	PROCEDURALVEGETATION_API extern void InterpolatePointsScales(
-		PV::FPointScaleAttributeView PointScaleAttribute,
-		PV::FPointLengthFromRootAttributeConstView PointLengthFromRootAttribute,
-		PV::FBranchPointsAttributeConstView BranchPointsAttribute,
-		PV::FBranchParentNumberAttributeConstView BranchParentNumberAttribute,
-		PV::FBranchChildrenAttributeConstView BranchChildrenAttribute,
-		PV::FBranchNumberAttributeConstView BranchNumberAttribute
-	);
+	struct FComputeSmoothedPointScalesAttributes
+	{
+		PV::FPointScaleAttributeView PointScale;
+		PV::FPointLengthFromRootAttributeConstView PointLengthFromRoot;
+		PV::FBranchPointsAttributeConstView BranchPoints;
+		PV::FBranchParentNumberAttributeConstView BranchParentNumber;
+		PV::FBranchChildrenAttributeConstView BranchChildren;
+		PV::FBranchNumberAttributeConstView BranchNumber;
+
+		FComputeSmoothedPointScalesAttributes(
+			PV::FPointScaleAttributeView InPointScale,
+			PV::FPointLengthFromRootAttributeConstView InPointLengthFromRoot,
+			PV::FBranchPointsAttributeConstView InBranchPoints,
+			PV::FBranchParentNumberAttributeConstView InBranchParentNumber,
+			PV::FBranchChildrenAttributeConstView InBranchChildren,
+			PV::FBranchNumberAttributeConstView InBranchNumber)
+			: PointScale(InPointScale)
+			, PointLengthFromRoot(InPointLengthFromRoot)
+			, BranchPoints(InBranchPoints)
+			, BranchParentNumber(InBranchParentNumber)
+			, BranchChildren(InBranchChildren)
+			, BranchNumber(InBranchNumber)
+		{
+		}
+
+		FComputeSmoothedPointScalesAttributes(FManagedArrayCollection& InCollection)
+			: PointScale(PV::FPointScaleAttribute::FindAttribute(InCollection))
+			, PointLengthFromRoot(PV::FPointLengthFromRootAttribute::FindAttribute(InCollection))
+			, BranchPoints(PV::FBranchPointsAttribute::FindAttribute(InCollection))
+			, BranchParentNumber(PV::FBranchParentNumberAttribute::FindAttribute(InCollection))
+			, BranchChildren(PV::FBranchChildrenAttribute::FindAttribute(InCollection))
+			, BranchNumber(PV::FBranchNumberAttribute::FindAttribute(InCollection))
+		{
+		}
+
+		bool IsValid() const
+		{
+			return ValidateAttributeCollection(PointScale, PointLengthFromRoot, BranchPoints, BranchParentNumber, BranchChildren, BranchNumber);
+		}
+	};
+
+	// Smooths the existing point scales across the branch hierarchy and writes the results back to the PointScale attribute.
+	// Expects the PointScale attribute to already contain valid scale values prior to calling this function.
+	PROCEDURALVEGETATION_API extern void ComputeSmoothedPointScales(const FComputeSmoothedPointScalesAttributes& InAttributes);
 
 	struct FSplitFusedPointsAttributes
 	{
@@ -652,27 +664,29 @@ namespace PV::AttributesHelper
 
 		bool IsValid() const
 		{
-			return BranchParentNumber.IsValid()
-				&& BranchNumber.IsValid()
-				&& BranchPoints.IsValid()
-				&& PointPosition.IsValid()
-				&& PointLengthFromRoot.IsValid()
-				&& PointLengthFromSeed.IsValid()
-				&& PointScaleGradient.IsValid()
-				&& PointHullGradient.IsValid()
-				&& PointMainTrunkGradient.IsValid()
-				&& PointGroundGradient.IsValid()
-				&& PointScale.IsValid()
-				&& BudDirection.IsValid()
-				&& BudStatus.IsValid()
-				&& BudLateralMeristem.IsValid()
-				&& PointBranchGradient.IsValid()
-				&& BudHormoneLevels.IsValid()
-				&& PointPlantGradient.IsValid()
-				&& PointBudNumber.IsValid()
-				&& PointNjordPixelIndex.IsValid()
-				&& BudLightDetected.IsValid()
-				&& BudDevelopment.IsValid();
+			return ValidateAttributeCollection(
+				BranchParentNumber,
+				BranchNumber,
+				BranchPoints,
+				PointPosition.GetConstView(),
+				PointLengthFromRoot.GetConstView(),
+				PointLengthFromSeed.GetConstView(),
+				PointScaleGradient.GetConstView(),
+				PointHullGradient.GetConstView(),
+				PointMainTrunkGradient.GetConstView(),
+				PointGroundGradient.GetConstView(),
+				PointScale.GetConstView(),
+				BudDirection.GetConstView(),
+				BudStatus.GetConstView(),
+				BudLateralMeristem.GetConstView(),
+				PointBranchGradient.GetConstView(),
+				BudHormoneLevels.GetConstView(),
+				PointPlantGradient.GetConstView(),
+				PointBudNumber.GetConstView(),
+				PointNjordPixelIndex.GetConstView(),
+				BudLightDetected.GetConstView(),
+				BudDevelopment.GetConstView()
+			);
 		}
 	};
 
@@ -738,30 +752,32 @@ namespace PV::AttributesHelper
 
 		bool IsValid() const
 		{
-			return PointPosition.IsValid()
-				&& BranchPoints.IsValid()
-				&& BranchParentNumber.IsValid()
-				&& BranchParents.IsValid()
-				&& BranchNumber.IsValid()
-				&& BranchChildren.IsValid()
-				&& BranchHierarchyNumber.IsValid()
-				&& BranchPlantNumber.IsValid()
-				&& PointScale.IsValid()
-				&& PointBudNumber.IsValid()
-				&& BranchSourceBudNumber.IsValid()
-				&& PointNjordPixelIndex.IsValid()
-				&& PointLengthFromRoot.IsValid()
-				&& PointLengthFromSeed.IsValid()
-				&& BudDirection.IsValid()
-				&& BudLightDetected.IsValid()
-				&& BudLateralMeristem.IsValid()
-				&& PointGroundGradient.IsValid()
-				&& PointScaleGradient.IsValid()
-				&& PointPlantGradient.IsValid()
-				&& PointBranchGradient.IsValid()
-				&& BudDevelopment.IsValid()
-				&& BudHormoneLevels.IsValid()
-				&& BudStatus.IsValid();
+			return ValidateAttributeCollection(
+				PointPosition,
+				BranchPoints,
+				BranchParentNumber,
+				BranchParents,
+				BranchNumber,
+				BranchChildren,
+				BranchHierarchyNumber,
+				BranchPlantNumber,
+				PointScale,
+				PointBudNumber,
+				BranchSourceBudNumber,
+				PointNjordPixelIndex,
+				PointLengthFromRoot,
+				PointLengthFromSeed,
+				BudDirection,
+				BudLightDetected,
+				BudLateralMeristem,
+				PointGroundGradient,
+				PointScaleGradient,
+				PointPlantGradient,
+				PointBranchGradient,
+				BudDevelopment,
+				BudHormoneLevels,
+				BudStatus
+			);
 		}
 	};
 
diff --git a/Engine/Plugins/Experimental/ProceduralVegetationEditor/Source/ProceduralVegetation/Private/Tests/PVAttributesTests.cpp b/Engine/Plugins/Experimental/ProceduralVegetationEditor/Source/ProceduralVegetation/Private/Tests/PVAttributesTests.cpp
--- a/Engine/Plugins/Experimental/ProceduralVegetationEditor/Source/ProceduralVegetation/Private/Tests/PVAttributesTests.cpp
+++ b/Engine/Plugins/Experimental/ProceduralVegetationEditor/Source/ProceduralVegetation/Private/Tests/PVAttributesTests.cpp
@@ -683,4 +683,221 @@ PV_ATTRIBUTES_TEST(SplitFusedPoints)
 	return true;
 }
 
+
+PV_ATTRIBUTES_TEST(ComputeSmoothedPointScales)
+{
+	// Branch 0 (points 0-5) of the mock tree is used to cover all scenarios:
+	//   P0 = 100  (base anchor)
+	//   P1 = 5    (noisy dip: close to base, ImpliedTaperRate >> 3x ExpectedTaperRate => rejected)
+	//   P2 = 200  (too large: > CurrentPointScale => skipped as not a valid anchor)
+	//   P3 = 70   (normal taper anchor, should be preserved)
+	//   P4 = 0    (zero scale, should be interpolated between P3 and P5)
+	//   P5 = 20   (tip anchor, should be preserved)
+	const TSharedRef<FManagedArrayCollection> Collection = PVMockTreeCollection::CreateCollection();
+
+	PV::FPointScaleAttributeView PointScaleAttribute = PV::FPointScaleAttribute::FindAttribute(*Collection);
+	PV::FPointLengthFromRootAttributeConstView PointLengthFromRootAttribute = PV::FPointLengthFromRootAttribute::FindAttribute(*Collection);
+	PV::FBranchPointsAttributeConstView BranchPointsAttribute = PV::FBranchPointsAttribute::FindAttribute(*Collection);
+	PV::FBranchParentNumberAttributeConstView BranchParentNumberAttribute = PV::FBranchParentNumberAttribute::FindAttribute(*Collection);
+	PV::FBranchChildrenAttributeConstView BranchChildrenAttribute = PV::FBranchChildrenAttribute::FindAttribute(*Collection);
+	PV::FBranchNumberAttributeConstView BranchNumberAttribute = PV::FBranchNumberAttribute::FindAttribute(*Collection);
+
+	UTEST_TRUE("HasScaleAttribute", PointScaleAttribute.IsValid());
+	UTEST_TRUE("HasLengthFromRootAttribute", PointLengthFromRootAttribute.IsValid());
+	UTEST_TRUE("HasBranchPointsAttribute", BranchPointsAttribute.IsValid());
+
+	// Set up test scale values for Branch 0 (point indices 0-5)
+	PointScaleAttribute[0] = 100.0f; // base anchor
+	PointScaleAttribute[1] =   5.0f; // noisy dip: ImpliedTaperRate ~205/m > 3 * ExpectedTaperRate ~151/m => rejected
+	PointScaleAttribute[2] = 200.0f; // too large: > CurrentPointScale, never accepted as anchor
+	PointScaleAttribute[3] =  70.0f; // valid taper anchor
+	PointScaleAttribute[4] =   0.0f; // zero scale, must be interpolated
+	PointScaleAttribute[5] =  20.0f; // tip anchor
+
+	PV::AttributesHelper::ComputeSmoothedPointScales({
+		PointScaleAttribute,
+		PointLengthFromRootAttribute,
+		BranchPointsAttribute,
+		BranchParentNumberAttribute,
+		BranchChildrenAttribute,
+		BranchNumberAttribute
+	});
+
+	// P0: base anchor must be preserved unchanged
+	UTEST_NEARLY_EQUAL("P0_BaseAnchorPreserved", PointScaleAttribute[0], 100.0f, 0.1f);
+
+	// P1: noisy dip was rejected; interpolated between P0 (100) and the first accepted
+	//     anchor P3 (70). At LFR alpha ~0.369: lerp(100, 70, 0.369) ~= 88.9
+	UTEST_NEARLY_EQUAL("P1_NoisyDipInterpolated", PointScaleAttribute[1], 88.9f, 0.5f);
+
+	// P2: too-large value was skipped; interpolated between P0 (100) and P3 (70).
+	//     At LFR alpha ~0.738: lerp(100, 70, 0.738) ~= 77.9
+	UTEST_NEARLY_EQUAL("P2_TooLargeInterpolated", PointScaleAttribute[2], 77.9f, 0.5f);
+
+	// P3: normal taper anchor must be preserved unchanged
+	UTEST_NEARLY_EQUAL("P3_NormalTaperPreserved", PointScaleAttribute[3], 70.0f, 0.1f);
+
+	// P4: zero scale was interpolated between P3 (70) and P5 (20).
+	//     At the exact midpoint (alpha = 0.5): lerp(70, 20, 0.5) = 45.0
+	UTEST_NEARLY_EQUAL("P4_ZeroScaleInterpolated", PointScaleAttribute[4], 45.0f, 0.5f);
+
+	// P5: tip anchor must be preserved unchanged
+	UTEST_NEARLY_EQUAL("P5_TipAnchorPreserved", PointScaleAttribute[5], 20.0f, 0.1f);
+
+	// Sanity check: the output must be a monotonically non-increasing taper
+	UTEST_TRUE("MonotonicTaper_P0_P1", PointScaleAttribute[1] <= PointScaleAttribute[0]);
+	UTEST_TRUE("MonotonicTaper_P1_P2", PointScaleAttribute[2] <= PointScaleAttribute[1]);
+	UTEST_TRUE("MonotonicTaper_P2_P3", PointScaleAttribute[3] <= PointScaleAttribute[2]);
+	UTEST_TRUE("MonotonicTaper_P3_P4", PointScaleAttribute[4] <= PointScaleAttribute[3]);
+	UTEST_TRUE("MonotonicTaper_P4_P5", PointScaleAttribute[5] <= PointScaleAttribute[4]);
+
+	// --- Fused-point regression ---
+	// When bHasFusedPoints is true, BranchPoints[0] of a child branch is the same point as
+	// the parent branch's attachment point. bShouldSkipFirstPoint / StartIndex=1 is supposed
+	// to protect it. The bug was that the initial backfill loop (triggered when the first own
+	// point of the child has scale 0) started at j=0 instead of j=StartIndex, so it would
+	// overwrite BranchPoints[0] — the fused/shared point — with the child branch's scale.
+	//
+	// Mock-tree fused topology (bSplitPoints=false):
+	//   Branch 0 (trunk): {0, 1, 2, 3, 4, 5}
+	//   Branch 1:         {2, 7, 8, 9}   — point 2 fused with branch 0
+	//   Branch 2:         {8, 10, 11}    — point 8 fused with branch 1
+	//   Branch 3:         {3, 6}         — point 3 fused with branch 0
+	//
+	// Setting point 7 (Branch1's BranchPoints[1]) to 0 forces the backfill path.
+	// With the old bug the backfill would start at j=0 and clobber point 2 with point 8's scale.
+	{
+		const TSharedRef<FManagedArrayCollection> FusedCollection = PVMockTreeCollection::CreateCollection(false);
+
+		PV::FPointScaleAttributeView FusedScaleAttr = PV::FPointScaleAttribute::FindAttribute(*FusedCollection);
+		PV::FPointLengthFromRootAttributeConstView FusedLFRAttr = PV::FPointLengthFromRootAttribute::FindAttribute(*FusedCollection);
+		PV::FBranchPointsAttributeConstView FusedBranchPointsAttr = PV::FBranchPointsAttribute::FindAttribute(*FusedCollection);
+		PV::FBranchParentNumberAttributeConstView FusedParentNumberAttr = PV::FBranchParentNumberAttribute::FindAttribute(*FusedCollection);
+		PV::FBranchChildrenAttributeConstView FusedBranchChildrenAttr = PV::FBranchChildrenAttribute::FindAttribute(*FusedCollection);
+		PV::FBranchNumberAttributeConstView FusedBranchNumberAttr = PV::FBranchNumberAttribute::FindAttribute(*FusedCollection);
+
+		UTEST_TRUE("FusedCollection_HasScaleAttribute", FusedScaleAttr.IsValid());
+
+		// Branch 0 (trunk): monotonically decreasing — all points kept as-is after processing.
+		// Point 2 = 80 is the fused attachment point for branch 1.
+		FusedScaleAttr[0]  = 100.0f;
+		FusedScaleAttr[1]  =  90.0f;
+		FusedScaleAttr[2]  =  80.0f; // fused: shared by branch 0 (index 2) and branch 1 (index 0)
+		FusedScaleAttr[3]  =  70.0f;
+		FusedScaleAttr[4]  =  60.0f;
+		FusedScaleAttr[5]  =  50.0f;
+		FusedScaleAttr[6]  =  40.0f; // branch 3 own point
+		FusedScaleAttr[7]  =   0.0f; // KEY: branch 1 BranchPoints[1] — triggers the backfill path
+		FusedScaleAttr[8]  =  30.0f; // branch 1 BranchPoints[2] — first valid value for backfill
+		FusedScaleAttr[9]  =  10.0f;
+		FusedScaleAttr[10] =  20.0f;
+		FusedScaleAttr[11] =   5.0f;
+
+		PV::AttributesHelper::ComputeSmoothedPointScales({
+			FusedScaleAttr,
+			FusedLFRAttr,
+			FusedBranchPointsAttr,
+			FusedParentNumberAttr,
+			FusedBranchChildrenAttr,
+			FusedBranchNumberAttr
+		});
+
+		// Point 2 is the fused point. Branch 1's backfill must not overwrite it;
+		// it must keep the value set by branch 0 (80).
+		UTEST_NEARLY_EQUAL("FusedPoint_P2_NotOverwrittenByChildBackfill", FusedScaleAttr[2], 80.0f, 0.1f);
+
+		// Point 7 is branch 1's first own point; it was 0 and must have been backfilled to
+		// the value of the first valid successor (point 8 = 30).
+		UTEST_NEARLY_EQUAL("FusedPoint_P7_BackfilledToP8Scale", FusedScaleAttr[7], 30.0f, 0.1f);
+	}
+
+	return true;
+}
+
+PV_ATTRIBUTES_TEST(ValidateAttributeCollection)
+{
+	using namespace PV::AttributesHelper;
+
+	// No attributes → always valid
+	UTEST_TRUE("NoAttributes", ValidateAttributeCollection());
+
+	// Single valid attribute → valid
+	{
+		FManagedArrayCollection Collection;
+		Collection.AddGroup(PV::GroupNames::PointGroup);
+		Collection.AddElements(5, PV::GroupNames::PointGroup);
+		PV::FPointLengthFromRootAttribute::AddAttribute(Collection);
+
+		const PV::FPointLengthFromRootAttributeConstView View = PV::FPointLengthFromRootAttribute::FindAttribute(Collection);
+		UTEST_TRUE("SingleValid", ValidateAttributeCollection(View));
+	}
+
+	// Single invalid (null) attribute → invalid
+	{
+		const PV::FPointLengthFromRootAttributeConstView InvalidView;
+		UTEST_FALSE("SingleInvalid", ValidateAttributeCollection(InvalidView));
+	}
+
+	// Two valid attributes in the same group with equal sizes → valid
+	{
+		FManagedArrayCollection Collection;
+		Collection.AddGroup(PV::GroupNames::PointGroup);
+		Collection.AddElements(5, PV::GroupNames::PointGroup);
+		PV::FPointLengthFromRootAttribute::AddAttribute(Collection);
+		PV::FPointScaleAttribute::AddAttribute(Collection);
+
+		const PV::FPointLengthFromRootAttributeConstView LengthView = PV::FPointLengthFromRootAttribute::FindAttribute(Collection);
+		const PV::FPointScaleAttributeConstView ScaleView = PV::FPointScaleAttribute::FindAttribute(Collection);
+		UTEST_TRUE("SameGroupSameSize", ValidateAttributeCollection(LengthView, ScaleView));
+	}
+
+	// Two valid attributes in the same group but from collections with different sizes → invalid
+	{
+		FManagedArrayCollection Collection5;
+		Collection5.AddGroup(PV::GroupNames::PointGroup);
+		Collection5.AddElements(5, PV::GroupNames::PointGroup);
+		PV::FPointLengthFromRootAttribute::AddAttribute(Collection5);
+
+		FManagedArrayCollection Collection10;
+		Collection10.AddGroup(PV::GroupNames::PointGroup);
+		Collection10.AddElements(10, PV::GroupNames::PointGroup);
+		PV::FPointScaleAttribute::AddAttribute(Collection10);
+
+		const PV::FPointLengthFromRootAttributeConstView LengthView = PV::FPointLengthFromRootAttribute::FindAttribute(Collection5);
+		const PV::FPointScaleAttributeConstView ScaleView = PV::FPointScaleAttribute::FindAttribute(Collection10);
+		UTEST_FALSE("SameGroupDifferentSizes", ValidateAttributeCollection(LengthView, ScaleView));
+	}
+
+	// Valid attributes in different groups → sizes between groups are independent, so this is valid
+	{
+		FManagedArrayCollection Collection;
+		Collection.AddGroup(PV::GroupNames::PointGroup);
+		Collection.AddElements(5, PV::GroupNames::PointGroup);
+		Collection.AddGroup(PV::GroupNames::BranchGroup);
+		Collection.AddElements(3, PV::GroupNames::BranchGroup);
+		PV::FPointLengthFromRootAttribute::AddAttribute(Collection);
+		PV::FBranchNumberAttribute::AddAttribute(Collection);
+
+		const PV::FPointLengthFromRootAttributeConstView LengthView = PV::FPointLengthFromRootAttribute::FindAttribute(Collection);
+		const PV::FBranchNumberAttributeConstView BranchNumberView = PV::FBranchNumberAttribute::FindAttribute(Collection);
+		UTEST_TRUE("DifferentGroupsDifferentSizes", ValidateAttributeCollection(LengthView, BranchNumberView));
+	}
+
+	// One invalid attribute mixed with valid ones → invalid
+	{
+		FManagedArrayCollection Collection;
+		Collection.AddGroup(PV::GroupNames::PointGroup);
+		Collection.AddElements(5, PV::GroupNames::PointGroup);
+		PV::FPointLengthFromRootAttribute::AddAttribute(Collection);
+		PV::FPointScaleAttribute::AddAttribute(Collection);
+
+		const PV::FPointLengthFromRootAttributeConstView LengthView = PV::FPointLengthFromRootAttribute::FindAttribute(Collection);
+		const PV::FPointScaleAttributeConstView ScaleView = PV::FPointScaleAttribute::FindAttribute(Collection);
+		const PV::FPointPositionAttributeConstView InvalidPositionView;
+		UTEST_FALSE("OneInvalidAmongValid", ValidateAttributeCollection(LengthView, InvalidPositionView, ScaleView));
+	}
+
+	return true;
+}
+
 #endif //WITH_DEV_AUTOMATION_TESTS
\ No newline at end of file
