Fix incorrect expected values in Procedural Vegetation attribute tests to align with updated bud status, hormone level, Davinci accumulation, and apical direction computation logic
Add clarifying comments explaining the computation rules for hormone levels, Davinci values, and apical direction per branch point

#Synced-CL 51620377

diff --git a/Engine/Plugins/Experimental/ProceduralVegetationEditor/Source/ProceduralVegetation/Private/Tests/PVAttributesTests.cpp b/Engine/Plugins/Experimental/ProceduralVegetationEditor/Source/ProceduralVegetation/Private/Tests/PVAttributesTests.cpp
--- a/Engine/Plugins/Experimental/ProceduralVegetationEditor/Source/ProceduralVegetation/Private/Tests/PVAttributesTests.cpp
+++ b/Engine/Plugins/Experimental/ProceduralVegetationEditor/Source/ProceduralVegetation/Private/Tests/PVAttributesTests.cpp
@@ -259,11 +259,11 @@ PV_ATTRIBUTES_TEST(ComputeBudStatus)
 	const FBudStatusDataLayout Point0_ExpectedBudStatus = {
 		.ApicalMeristem = 0,
 		.Codominant = 0,
-		.Axillary = 0,
+		.Axillary = 1,
 		.Seed = 1,
-		.Dormant = 0,
-		.Triggered = 1,
-		.NumTriggered = 1,
+		.Dormant = 1,
+		.Triggered = 0,
+		.NumTriggered = 0,
 		.Inactive = 1,
 		.BrokenTip = 0,
 		.Broken = 0,
@@ -457,9 +457,12 @@ PV_ATTRIBUTES_TEST(ComputeBudHormoneLevels)
 		}
 	};
 
+	// Apical is always 1 for all points.
+	// Axillary = GetMappedRangeValueClamped([0,1] -> [0.35, 1.0], InverseGradient), where InverseGradient = 1 - PlantGradient.
+	// Ethylene = PlantGradient.
 	const FBudHormoneLevelsDataLayout Point0_ExpectedBudHormoneLevels = {
-		.Apical = 0,
-		.Axillary = 0,
+		.Apical = 1,
+		.Axillary = 0.35f,  // Gradient=1, InverseGradient=0 => mapped to 0.35
 		.AxillaryInhibition = 0,
 		.Radical = 0,
 		.Ethylene = 1,
@@ -467,7 +470,7 @@ PV_ATTRIBUTES_TEST(ComputeBudHormoneLevels)
 	};
 	const FBudHormoneLevelsDataLayout Point5_ExpectedBudHormoneLevels = {
 		.Apical = 1,
-		.Axillary = 0,
+		.Axillary = 1.0f,   // Gradient=0, InverseGradient=1 => mapped to 1.0
 		.AxillaryInhibition = 0,
 		.Radical = 0,
 		.Ethylene = 0,
@@ -475,7 +478,7 @@ PV_ATTRIBUTES_TEST(ComputeBudHormoneLevels)
 	};
 	const FBudHormoneLevelsDataLayout Point6_ExpectedBudHormoneLevels = {
 		.Apical = 1,
-		.Axillary = 0,
+		.Axillary = 1.0f,   // Gradient=0, InverseGradient=1 => mapped to 1.0
 		.AxillaryInhibition = 0,
 		.Radical = 0,
 		.Ethylene = 0,
@@ -483,7 +486,7 @@ PV_ATTRIBUTES_TEST(ComputeBudHormoneLevels)
 	};
 	const FBudHormoneLevelsDataLayout Point9_ExpectedBudHormoneLevels = {
 		.Apical = 1,
-		.Axillary = 0,
+		.Axillary = 1.0f,   // Gradient=0, InverseGradient=1 => mapped to 1.0
 		.AxillaryInhibition = 0,
 		.Radical = 0,
 		.Ethylene = 0,
@@ -491,7 +494,7 @@ PV_ATTRIBUTES_TEST(ComputeBudHormoneLevels)
 	};
 	const FBudHormoneLevelsDataLayout Point11_ExpectedBudHormoneLevels = {
 		.Apical = 1,
-		.Axillary = 0,
+		.Axillary = 1.0f,   // Gradient=0, InverseGradient=1 => mapped to 1.0
 		.AxillaryInhibition = 0,
 		.Radical = 0,
 		.Ethylene = 0,
@@ -521,19 +524,20 @@ PV_ATTRIBUTES_TEST(ComputeBudLateralMeristem)
 	UTEST_TRUE("HasAttribute", BudLateralMeristemAttribute.IsValid());
 	UTEST_TRUE("Has 12 elements", BudLateralMeristemAttribute.Num() >= 12);
 
+	// All points receive a non-zero Davinci value because EstimatedLateralElongation (> 0) is
+	// accumulated starting from each branch's tip, so even tip points are non-zero.
 	UTEST_TRUE("Point0_DavinciIsNotZero",  !FMath::IsNearlyZero(BudLateralMeristemAttribute[0].Davinci));
 	UTEST_TRUE("Point1_DavinciIsNotZero",  !FMath::IsNearlyZero(BudLateralMeristemAttribute[1].Davinci));
 	UTEST_TRUE("Point2_DavinciIsNotZero",  !FMath::IsNearlyZero(BudLateralMeristemAttribute[2].Davinci));
 	UTEST_TRUE("Point3_DavinciIsNotZero",  !FMath::IsNearlyZero(BudLateralMeristemAttribute[3].Davinci));
 	UTEST_TRUE("Point4_DavinciIsNotZero",  !FMath::IsNearlyZero(BudLateralMeristemAttribute[4].Davinci));
+	UTEST_TRUE("Point5_DavinciIsNotZero",  !FMath::IsNearlyZero(BudLateralMeristemAttribute[5].Davinci));
+	UTEST_TRUE("Point6_DavinciIsNotZero",  !FMath::IsNearlyZero(BudLateralMeristemAttribute[6].Davinci));
 	UTEST_TRUE("Point7_DavinciIsNotZero",  !FMath::IsNearlyZero(BudLateralMeristemAttribute[7].Davinci));
 	UTEST_TRUE("Point8_DavinciIsNotZero",  !FMath::IsNearlyZero(BudLateralMeristemAttribute[8].Davinci));
+	UTEST_TRUE("Point9_DavinciIsNotZero",  !FMath::IsNearlyZero(BudLateralMeristemAttribute[9].Davinci));
 	UTEST_TRUE("Point10_DavinciIsNotZero", !FMath::IsNearlyZero(BudLateralMeristemAttribute[10].Davinci));
-
-	UTEST_TRUE("Point5_DavinciIsZero",  FMath::IsNearlyZero(BudLateralMeristemAttribute[5].Davinci));
-	UTEST_TRUE("Point6_DavinciIsZero",  FMath::IsNearlyZero(BudLateralMeristemAttribute[6].Davinci));
-	UTEST_TRUE("Point9_DavinciIsZero",  FMath::IsNearlyZero(BudLateralMeristemAttribute[9].Davinci));
-	UTEST_TRUE("Point11_DavinciIsZero", FMath::IsNearlyZero(BudLateralMeristemAttribute[11].Davinci));
+	UTEST_TRUE("Point11_DavinciIsNotZero", !FMath::IsNearlyZero(BudLateralMeristemAttribute[11].Davinci));
 
 	return true;
 }
@@ -588,18 +592,22 @@ PV_ATTRIBUTES_TEST(ComputeBudLightDetected)
 
 PV_ATTRIBUTES_TEST(ComputeBudDirections)
 {
-	const FVector3f ExpectedApical0 = (PVMockTree::Point1 - PVMockTree::Point0).GetSafeNormal();
-	const FVector3f ExpectedApical1 = (PVMockTree::Point2 - PVMockTree::Point1).GetSafeNormal();
-	const FVector3f ExpectedApical2 = (PVMockTree::Point3 - PVMockTree::Point2).GetSafeNormal();
-	const FVector3f ExpectedApical3 = (PVMockTree::Point4 - PVMockTree::Point3).GetSafeNormal();
-	const FVector3f ExpectedApical4 = (PVMockTree::Point5 - PVMockTree::Point4).GetSafeNormal();
-	const FVector3f ExpectedApical5 = (PVMockTree::Point5 - PVMockTree::Point4).GetSafeNormal();
-	const FVector3f ExpectedApical6 = (PVMockTree::Point6 - PVMockTree::Point3).GetSafeNormal();
-	const FVector3f ExpectedApical7 = (PVMockTree::Point8 - PVMockTree::Point7).GetSafeNormal();
-	const FVector3f ExpectedApical8 = (PVMockTree::Point9 - PVMockTree::Point8).GetSafeNormal();
-	const FVector3f ExpectedApical9 = (PVMockTree::Point9 - PVMockTree::Point8).GetSafeNormal();
-	const FVector3f ExpectedApical10 = (PVMockTree::Point11 - PVMockTree::Point10).GetSafeNormal();
-	const FVector3f ExpectedApical11 = (PVMockTree::Point11 - PVMockTree::Point10).GetSafeNormal();
+	// Apical direction rule: first point on a branch looks forward (Next - Current),
+	// all other points look backward along their incoming segment (Current - Prev).
+	// Points shared at branch junctions are split, so each branch start uses its own copy.
+	// Branch 0 [0,1,2,3,4,5], Branch 1 [copy(2),7,8,9], Branch 2 [copy(8),10,11], Branch 3 [copy(3),6]
+	const FVector3f ExpectedApical0  = (PVMockTree::Point1  - PVMockTree::Point0).GetSafeNormal(); // Branch0 BPI=0: Next-Current
+	const FVector3f ExpectedApical1  = (PVMockTree::Point1  - PVMockTree::Point0).GetSafeNormal(); // Branch0 BPI=1: Current-Prev
+	const FVector3f ExpectedApical2  = (PVMockTree::Point2  - PVMockTree::Point1).GetSafeNormal(); // Branch0 BPI=2: Current-Prev
+	const FVector3f ExpectedApical3  = (PVMockTree::Point3  - PVMockTree::Point2).GetSafeNormal(); // Branch0 BPI=3: Current-Prev
+	const FVector3f ExpectedApical4  = (PVMockTree::Point4  - PVMockTree::Point3).GetSafeNormal(); // Branch0 BPI=4: Current-Prev
+	const FVector3f ExpectedApical5  = (PVMockTree::Point5  - PVMockTree::Point4).GetSafeNormal(); // Branch0 BPI=5: Current-Prev
+	const FVector3f ExpectedApical6  = (PVMockTree::Point6  - PVMockTree::Point3).GetSafeNormal(); // Branch3 BPI=1: Current-Prev (prev is split copy of Point3)
+	const FVector3f ExpectedApical7  = (PVMockTree::Point7  - PVMockTree::Point2).GetSafeNormal(); // Branch1 BPI=1: Current-Prev (prev is split copy of Point2)
+	const FVector3f ExpectedApical8  = (PVMockTree::Point8  - PVMockTree::Point7).GetSafeNormal(); // Branch1 BPI=2: Current-Prev
+	const FVector3f ExpectedApical9  = (PVMockTree::Point9  - PVMockTree::Point8).GetSafeNormal(); // Branch1 BPI=3: Current-Prev
+	const FVector3f ExpectedApical10 = (PVMockTree::Point10 - PVMockTree::Point8).GetSafeNormal(); // Branch2 BPI=1: Current-Prev (prev is split copy of Point8)
+	const FVector3f ExpectedApical11 = (PVMockTree::Point11 - PVMockTree::Point10).GetSafeNormal(); // Branch2 BPI=2: Current-Prev
 
 	const int32 ExpectedNumBudDirectionValues = 6;
 
