[PVE] Support subgraph rendering

Issue: PVBaseSettings and PVSubgraphSettings don't share a common base class.
Fix: Added a new IPVRenderSettings interface, which allows the sharing of a common interface for visualization-related settings.

There was also a case where dragging an embedded graph was adding a PCG Subgraph node. To spawn our PVSubgraph node, we had to expose a hook in PCG codebase.

#Synced-CL 51621520

diff --git a/Engine/Plugins/Experimental/ProceduralVegetationEditor/Source/ProceduralVegetation/Private/Nodes/PVBaseSettings.cpp b/Engine/Plugins/Experimental/ProceduralVegetationEditor/Source/ProceduralVegetation/Private/Nodes/PVBaseSettings.cpp
--- a/Engine/Plugins/Experimental/ProceduralVegetationEditor/Source/ProceduralVegetation/Private/Nodes/PVBaseSettings.cpp
+++ b/Engine/Plugins/Experimental/ProceduralVegetationEditor/Source/ProceduralVegetation/Private/Nodes/PVBaseSettings.cpp
@@ -21,7 +21,7 @@ UPVBaseSettings::UPVBaseSettings()
 #if WITH_EDITOR
 void UPVBaseSettings::SetCurrentRenderType(TArray<EPVRenderType> InRenderTypes)
 {
-	CurrentRenderType = InRenderTypes;
+	Visualization.CurrentRenderType = InRenderTypes;
 
 	PostEditChange();
 }
diff --git a/Engine/Plugins/Experimental/ProceduralVegetationEditor/Source/ProceduralVegetation/Private/Nodes/PVSubgraphSettings.cpp b/Engine/Plugins/Experimental/ProceduralVegetationEditor/Source/ProceduralVegetation/Private/Nodes/PVSubgraphSettings.cpp
--- a/Engine/Plugins/Experimental/ProceduralVegetationEditor/Source/ProceduralVegetation/Private/Nodes/PVSubgraphSettings.cpp
+++ b/Engine/Plugins/Experimental/ProceduralVegetationEditor/Source/ProceduralVegetation/Private/Nodes/PVSubgraphSettings.cpp
@@ -13,6 +13,29 @@ UPCGGraphInterface* UPVSubgraphSettings::GetSubgraphInterface() const
 	return SubgraphInstance->GetGraph();
 }
 
+void FPVSubgraphElement::PostExecuteInternal(FPCGContext* InContext) const
+{
+	FPCGSubgraphElement::PostExecuteInternal(InContext);
+
+#if WITH_EDITORONLY_DATA
+	const UPVSubgraphSettings* Settings = InContext->GetInputSettings<UPVSubgraphSettings>();
+	check(Settings);
+
+	for (FPCGTaggedData& OutputData : InContext->OutputData.TaggedData)
+	{
+		if (OutputData.Data && OutputData.Data.IsA<UPVData>())
+		{
+			if (const UPVData* ProceduralVegetationData = Cast<UPVData>(OutputData.Data))
+			{
+				FPVDebugSettings DebugSettings = ProceduralVegetationData->GetDebugSettings();
+				DebugSettings.VisualizationSettings = Settings->GetDebugVisualizationSettings().VisualizationSettings;
+				ProceduralVegetationData->SetDebugSettings(MoveTemp(DebugSettings));
+			}
+		}
+	}
+#endif
+}
+
 #if WITH_EDITOR
 void UPVSubgraphSettings::PostEditChangeProperty(struct FPropertyChangedEvent& PropertyChangedEvent)
 {
@@ -39,6 +62,15 @@ UPVSubgraphSettings::UPVSubgraphSettings(const FObjectInitializer& InObjectIniti
 }
 
 
+#if WITH_EDITOR
+void UPVSubgraphSettings::SetCurrentRenderType(TArray<EPVRenderType> InRenderTypes)
+{
+	Visualization.CurrentRenderType = InRenderTypes;
+
+	PostEditChange();
+}
+#endif
+
 UPCGNode* UPVSubgraphSettings::CreateNode() const
 {
 	return NewObject<UPCGSubgraphNode>();
@@ -59,12 +91,24 @@ FString UPVSubgraphSettings::GetAdditionalTitleInformation() const
 
 FPCGElementPtr UPVSubgraphSettings::CreateElement() const
 {
-	return MakeShared<FPCGSubgraphElement>();
+	return MakeShared<FPVSubgraphElement>();
 }
 
 void UPVSubgraphSettings::SetSubgraphInternal(UPCGGraphInterface* InGraph)
 {
-	SubgraphInstance->GraphInstance->SetGraph(InGraph);
+	if (UProceduralVegetationGraph* ProceduralVegetationGraph = Cast<UProceduralVegetationGraph>(InGraph))
+	{
+		UProceduralVegetation* ProceduralVegetation = NewObject<UProceduralVegetation>(InGraph);
+		ProceduralVegetation->SetGraph(ProceduralVegetationGraph);
+	
+		SubgraphInstance->ProceduralVegetation = ProceduralVegetation;
+		SubgraphInstance->GraphInstance->SetGraph(InGraph);
+	}	
+	else
+	{
+		SubgraphInstance->ProceduralVegetation = nullptr;
+		SubgraphInstance->GraphInstance->SetGraph(nullptr);
+	}
 }
 
 #undef LOCTEXT_NAMESPACE
diff --git a/Engine/Plugins/Experimental/ProceduralVegetationEditor/Source/ProceduralVegetation/Private/Nodes/PVTexture2DImporterSettings.cpp b/Engine/Plugins/Experimental/ProceduralVegetationEditor/Source/ProceduralVegetation/Private/Nodes/PVTexture2DImporterSettings.cpp
--- a/Engine/Plugins/Experimental/ProceduralVegetationEditor/Source/ProceduralVegetation/Private/Nodes/PVTexture2DImporterSettings.cpp
+++ b/Engine/Plugins/Experimental/ProceduralVegetationEditor/Source/ProceduralVegetation/Private/Nodes/PVTexture2DImporterSettings.cpp
@@ -32,7 +32,7 @@ UPVTexture2DImporterSettings::UPVTexture2DImporterSettings()
 	static ConstructorHelpers::FObjectFinder<UMaterialInterface> TextureVisualizationMaterialRef(TEXT("/Script/Engine.Material'/Engine/Engine_MI_Shaders/M_Shader_SimpleTranslucent.M_Shader_SimpleTranslucent'"));
 	TextureVisualizationMaterial = TextureVisualizationMaterialRef.Object;
 
-	FPVVisualizationSettings& BranchHierarchyVisualization = DebugVisualizationSettings.VisualizationSettings.AddDefaulted_GetRef();
+	FPVVisualizationSettings& BranchHierarchyVisualization = Visualization.DebugSettings.VisualizationSettings.AddDefaulted_GetRef();
 	BranchHierarchyVisualization.DebugType = EPVDebugType::Custom;
 	BranchHierarchyVisualization.bShowAnchorPoints = false;
 	BranchHierarchyVisualization.bDrawPointAsMesh = false;
@@ -40,7 +40,7 @@ UPVTexture2DImporterSettings::UPVTexture2DImporterSettings()
 	BranchHierarchyVisualization.CustomGroupToFilter = PVImportNames::BranchHierarchyGroup;
 	BranchHierarchyVisualization.VisualizationMode = EPVDebugValueVisualizationMode::Curve;
 
-	FPVVisualizationSettings& PixelSamplingVisualization = DebugVisualizationSettings.VisualizationSettings.AddDefaulted_GetRef();
+	FPVVisualizationSettings& PixelSamplingVisualization = Visualization.DebugSettings.VisualizationSettings.AddDefaulted_GetRef();
 	PixelSamplingVisualization.DebugType = EPVDebugType::Custom;
 	PixelSamplingVisualization.bShowAnchorPoints = false;
 	PixelSamplingVisualization.bDrawPointAsMesh = false;
diff --git a/Engine/Plugins/Experimental/ProceduralVegetationEditor/Source/ProceduralVegetation/Public/DataTypes/PVData.h b/Engine/Plugins/Experimental/ProceduralVegetationEditor/Source/ProceduralVegetation/Public/DataTypes/PVData.h
--- a/Engine/Plugins/Experimental/ProceduralVegetationEditor/Source/ProceduralVegetation/Public/DataTypes/PVData.h
+++ b/Engine/Plugins/Experimental/ProceduralVegetationEditor/Source/ProceduralVegetation/Public/DataTypes/PVData.h
@@ -179,6 +179,18 @@ struct FPVVisualizationCollection
 	TArray<FPVVisualizationSettings> VisualizationSettings;
 };
 
+USTRUCT()
+struct FPVSettingsVisualization
+{
+	GENERATED_BODY()
+
+	UPROPERTY(EditAnywhere, Category = DebugSettings, meta=(EditConditionHides))
+	FPVDebugSettings DebugSettings;
+
+	UPROPERTY()
+	TArray<EPVRenderType> CurrentRenderType;
+};
+
 UCLASS(BlueprintType, ClassGroup = (Procedural))
 class PROCEDURALVEGETATION_API UPVData : public UPCGSpatialData
 {
diff --git a/Engine/Plugins/Experimental/ProceduralVegetationEditor/Source/ProceduralVegetation/Public/Nodes/PVBaseSettings.h b/Engine/Plugins/Experimental/ProceduralVegetationEditor/Source/ProceduralVegetation/Public/Nodes/PVBaseSettings.h
--- a/Engine/Plugins/Experimental/ProceduralVegetationEditor/Source/ProceduralVegetation/Public/Nodes/PVBaseSettings.h
+++ b/Engine/Plugins/Experimental/ProceduralVegetationEditor/Source/ProceduralVegetation/Public/Nodes/PVBaseSettings.h
@@ -3,12 +3,13 @@
 
 #include "CoreMinimal.h"
 #include "PCGSettings.h"
+#include "PVRenderSettings.h"
 #include "DataTypes/PVData.h"
 #include "Helpers/PVUtilities.h"
 #include "PVBaseSettings.generated.h"
 
 UCLASS(BlueprintType, Abstract, HideCategories=(Debug, AssetInfo), ClassGroup = (Procedural))
-class UPVBaseSettings : public UPCGSettings
+class UPVBaseSettings : public UPCGSettings, public IPVRenderSettings
 {
 	GENERATED_BODY()
 
@@ -19,14 +20,15 @@ public:
 	bool bLockInspection = false;
 	
 	virtual EPCGSettingsType GetType() const override { return EPCGSettingsType::Spatial; }
-	virtual TArray<EPVRenderType> GetDefaultRenderType() const { return TArray<EPVRenderType>(); }
-	virtual TArray<EPVRenderType> GetSupportedRenderTypes() const { return GetDefaultRenderType(); }
-	virtual TArray<EPVRenderType> GetCurrentRenderTypes() const { return CurrentRenderType; }
-	virtual TMap<UObject*, FTransform> GetViewportObjects() const { return TMap<UObject*, FTransform>(); }
-	virtual void SetCurrentRenderType(TArray<EPVRenderType> InRenderTypes);
+	virtual TArray<EPVRenderType> GetDefaultRenderType() const override { return TArray<EPVRenderType>(); }
+	virtual TArray<EPVRenderType> GetSupportedRenderTypes() const override { return GetDefaultRenderType(); }
+	virtual TArray<EPVRenderType> GetCurrentRenderTypes() const override { return Visualization.CurrentRenderType; }
+	virtual TMap<UObject*, FTransform> GetViewportObjects() const override { return TMap<UObject*, FTransform>(); }
+	virtual void SetCurrentRenderType (TArray<EPVRenderType> InRenderTypes) override;
+	
+	virtual const FPVDebugSettings& GetDebugVisualizationSettings() const override { return Visualization.DebugSettings; }
+	virtual bool IsDebug() const override { return bDebug; }
 	
-	const FPVDebugSettings& GetDebugVisualizationSettings() const { return DebugVisualizationSettings; }
-
 	virtual void PostEditChangeProperty(FPropertyChangedEvent& PropertyChangedEvent) override;
 #endif
 
@@ -46,17 +48,9 @@ public:
 #endif
 
 protected:
-
 #if WITH_EDITORONLY_DATA
 	UPROPERTY(EditAnywhere, Category = DebugSettings, meta=(EditCondition = "bDebug", EditConditionHides))
-	FPVDebugSettings DebugVisualizationSettings;
-#endif
-
-private:
-
-#if WITH_EDITORONLY_DATA
-	UPROPERTY()
-	TArray<EPVRenderType> CurrentRenderType;
+	FPVSettingsVisualization Visualization;
 #endif
 };
 
diff --git a/Engine/Plugins/Experimental/ProceduralVegetationEditor/Source/ProceduralVegetation/Public/Nodes/PVRenderSettings.h b/Engine/Plugins/Experimental/ProceduralVegetationEditor/Source/ProceduralVegetation/Public/Nodes/PVRenderSettings.h
new file mode 100644
--- /dev/null
+++ b/Engine/Plugins/Experimental/ProceduralVegetationEditor/Source/ProceduralVegetation/Public/Nodes/PVRenderSettings.h
@@ -0,0 +1,33 @@
+﻿// Copyright Epic Games, Inc. All Rights Reserved.
+#pragma once
+
+#include "CoreMinimal.h"
+#include "UObject/Interface.h"
+#include "PVRenderSettings.generated.h"
+
+struct FPVDebugSettings;
+enum class EPVRenderType : uint8;
+
+UINTERFACE(MinimalAPI)
+class UPVRenderSettings : public UInterface
+{
+	GENERATED_BODY()
+};
+
+class IPVRenderSettings
+{
+	GENERATED_BODY()
+	
+public:
+
+#if WITH_EDITOR
+	virtual TArray<EPVRenderType> GetDefaultRenderType() const = 0;
+	virtual TArray<EPVRenderType> GetSupportedRenderTypes() const = 0;
+	virtual TArray<EPVRenderType> GetCurrentRenderTypes() const = 0;
+	virtual TMap<UObject*, FTransform> GetViewportObjects() const = 0; 
+	virtual void SetCurrentRenderType(TArray<EPVRenderType> InRenderTypes) = 0;
+	virtual const FPVDebugSettings& GetDebugVisualizationSettings() const = 0;
+
+	virtual bool IsDebug() const = 0;
+#endif
+};
diff --git a/Engine/Plugins/Experimental/ProceduralVegetationEditor/Source/ProceduralVegetation/Public/Nodes/PVSubgraphSettings.h b/Engine/Plugins/Experimental/ProceduralVegetationEditor/Source/ProceduralVegetation/Public/Nodes/PVSubgraphSettings.h
--- a/Engine/Plugins/Experimental/ProceduralVegetationEditor/Source/ProceduralVegetation/Public/Nodes/PVSubgraphSettings.h
+++ b/Engine/Plugins/Experimental/ProceduralVegetationEditor/Source/ProceduralVegetation/Public/Nodes/PVSubgraphSettings.h
@@ -10,8 +10,8 @@
 class UProceduralVegetationInstance;
 class UProceduralVegetation;
 
-UCLASS(BlueprintType, ClassGroup = (Procedural))
-class UPVSubgraphSettings : public UPCGBaseSubgraphSettings
+UCLASS(BlueprintType, HideCategories=(Debug, AssetInfo), ClassGroup = (Procedural))
+class UPVSubgraphSettings : public UPCGBaseSubgraphSettings, public IPVRenderSettings
 {
 	GENERATED_BODY()
 
@@ -26,8 +26,15 @@ public:
 	virtual EPCGSettingsType GetType() const override { return EPCGSettingsType::Subgraph; }
 	virtual UObject* GetJumpTargetForDoubleClick() const override;
 	virtual bool GetPinExtraIcon(const UPCGPin* InPin, FName& OutExtraIcon, FText& OutTooltip) const override { return false; }
-	// TODO: rendering stuff
-	// virtual TArray<EPVRenderType> GetDefaultRenderType() const override { return TArray{ EPVRenderType::PointData }; }
+
+	virtual TArray<EPVRenderType> GetSupportedRenderTypes() const override { return TArray{EPVRenderType::PointData, EPVRenderType::Foliage, EPVRenderType::FoliageGrid, EPVRenderType::FoliageAttachments, EPVRenderType::Mesh, EPVRenderType::Bones}; }
+	virtual TArray<EPVRenderType> GetDefaultRenderType() const override { return TArray{EPVRenderType::PointData}; }
+	virtual TArray<EPVRenderType> GetCurrentRenderTypes() const override { return Visualization.CurrentRenderType; }
+	virtual TMap<UObject*, FTransform> GetViewportObjects() const override { return TMap<UObject*, FTransform>(); }
+	virtual void SetCurrentRenderType (TArray<EPVRenderType> InRenderTypes) override;
+	
+	virtual const FPVDebugSettings& GetDebugVisualizationSettings() const override { return Visualization.DebugSettings; }
+	virtual bool IsDebug() const override { return bDebug; }
 #endif
 	
 	//~UPCGSettings interface implementation
@@ -51,4 +58,18 @@ public:
 	virtual UPCGGraphInterface* GetSubgraphInterface() const override;
 
 	virtual UPCGGraphInterface* GetOriginalSubgraphInterface() const override { return SubgraphInstance->GraphInstance; }
+	virtual bool IsDynamicGraph() const override { return true;}
+
+protected:
+#if WITH_EDITORONLY_DATA
+	UPROPERTY(EditAnywhere, Category = DebugSettings, meta=(EditCondition = "bDebug", EditConditionHides))
+	FPVSettingsVisualization Visualization;
+#endif
+	
+};
+
+class FPVSubgraphElement : public FPCGSubgraphElement
+{
+protected:
+	virtual void PostExecuteInternal(FPCGContext* InContext) const override;
 };
\ No newline at end of file
diff --git a/Engine/Plugins/Experimental/ProceduralVegetationEditor/Source/ProceduralVegetation/Public/ProceduralVegetation.h b/Engine/Plugins/Experimental/ProceduralVegetationEditor/Source/ProceduralVegetation/Public/ProceduralVegetation.h
--- a/Engine/Plugins/Experimental/ProceduralVegetationEditor/Source/ProceduralVegetation/Public/ProceduralVegetation.h
+++ b/Engine/Plugins/Experimental/ProceduralVegetationEditor/Source/ProceduralVegetation/Public/ProceduralVegetation.h
@@ -69,6 +69,7 @@ public:
 	virtual UPCGGraphInterface* GetGraph() override { return Graph; }
 	virtual const UPCGGraphInterface* GetGraph() const override { return Graph; }
 
+	void SetGraph(UProceduralVegetationGraph* InGraph) { Graph = InGraph; };
 	void CreateGraph(const UProceduralVegetationGraph* InGraph = nullptr);
 
 #if WITH_EDITOR
diff --git a/Engine/Plugins/Experimental/ProceduralVegetationEditor/Source/ProceduralVegetationEditor/Private/Helpers/PVGraphHelpers.cpp b/Engine/Plugins/Experimental/ProceduralVegetationEditor/Source/ProceduralVegetationEditor/Private/Helpers/PVGraphHelpers.cpp
--- a/Engine/Plugins/Experimental/ProceduralVegetationEditor/Source/ProceduralVegetationEditor/Private/Helpers/PVGraphHelpers.cpp
+++ b/Engine/Plugins/Experimental/ProceduralVegetationEditor/Source/ProceduralVegetationEditor/Private/Helpers/PVGraphHelpers.cpp
@@ -20,6 +20,19 @@ namespace PV::Graph
 			: nullptr;
 	}
 
+	IPVRenderSettings* GetRenderSettings(const UPCGNode* InNode)
+	{
+		if (InNode && InNode->GetSettings())
+		{
+			if (InNode->GetSettings()->Implements<UPVRenderSettings>())
+			{
+				return Cast<IPVRenderSettings>(InNode->GetSettings());
+			}
+		}
+
+		return nullptr;
+	}
+
 	UPCGPin* GetInPinFromNode(UPCGNode* InNode)
 	{
 		return InNode
diff --git a/Engine/Plugins/Experimental/ProceduralVegetationEditor/Source/ProceduralVegetationEditor/Private/Helpers/PVGraphHelpers.h b/Engine/Plugins/Experimental/ProceduralVegetationEditor/Source/ProceduralVegetationEditor/Private/Helpers/PVGraphHelpers.h
--- a/Engine/Plugins/Experimental/ProceduralVegetationEditor/Source/ProceduralVegetationEditor/Private/Helpers/PVGraphHelpers.h
+++ b/Engine/Plugins/Experimental/ProceduralVegetationEditor/Source/ProceduralVegetationEditor/Private/Helpers/PVGraphHelpers.h
@@ -5,6 +5,7 @@
 #include "CoreMinimal.h"
 #include "PCGPin.h"
 
+class IPVRenderSettings;
 class SCollectionSpreadSheetWidget;
 class UPCGNode;
 class UPVBaseSettings;
@@ -12,7 +13,8 @@ class UPVBaseSettings;
 namespace PV::Graph
 {
 	UPVBaseSettings* GetNodeSettings(const UPCGNode* InNode);
-
+	IPVRenderSettings* GetRenderSettings(const UPCGNode* InNode);
+	
 	UPCGPin* GetInPinFromNode(UPCGNode* InNode);
 	UPCGPin* GetOutPinFromNode(UPCGNode* InNode);
 	UPCGPin* GetUpstreamPin(UPCGPin* InPin);
diff --git a/Engine/Plugins/Experimental/ProceduralVegetationEditor/Source/ProceduralVegetationEditor/Private/PVEditor.cpp b/Engine/Plugins/Experimental/ProceduralVegetationEditor/Source/ProceduralVegetationEditor/Private/PVEditor.cpp
--- a/Engine/Plugins/Experimental/ProceduralVegetationEditor/Source/ProceduralVegetationEditor/Private/PVEditor.cpp
+++ b/Engine/Plugins/Experimental/ProceduralVegetationEditor/Source/ProceduralVegetationEditor/Private/PVEditor.cpp
@@ -93,6 +93,8 @@ void FPVEditor::Initialize(
 	UPCGEngineSubsystem* Subsystem = GEngine->GetEngineSubsystem<UPCGEngineSubsystem>();
 	check(Subsystem);
 
+	Subsystem->GetOnPCGSourceGenerationDone().AddRaw(this, &FPVEditor::OnGraphExecuted);
+	
 	const FPCGDefaultExecutionSourceParams Params
 	{
 		.GraphInterface = ProceduralVegetationGraph
@@ -188,6 +190,11 @@ void FPVEditor::OnClose()
 
 	GetInspectionDataManager().OnInspectedStackChangedDelegate.RemoveAll(this);
 
+	if (UPCGEngineSubsystem* Subsystem = GEngine->GetEngineSubsystem<UPCGEngineSubsystem>())
+	{
+		Subsystem->GetOnPCGSourceGenerationDone().RemoveAll(this);
+	}
+
 	FPCGEditor::OnClose();
 
 	if (ExecutionSource)
@@ -251,7 +258,11 @@ void FPVEditor::OnSelectedNodesChanged(const TSet<UObject*>& InNewSelection)
 	}
 
 	const UPVBaseSettings* InspectionSettings = GetNodeSettings(GetNodeBeingInspected());
-	if (!InspectionSettings || !InspectionSettings->bLockInspection)
+	const IPVRenderSettings* IRenderSettings = GetRenderSettings(GetNodeBeingInspected());
+
+	bool bInspectionLocked = InspectionSettings && InspectionSettings->bLockInspection;
+	
+	if (!IRenderSettings || !bInspectionLocked)
 	{
 		ChangeNodeInspection(GetFirstSelectedEdNode());
 	}
@@ -453,11 +464,14 @@ const UPVData* FPVEditor::GetDataFromPin(const UPCGPin* InPin)
 	return PVData;
 }
 
-void FPVEditor::OnNodeInspected(const FPCGStack& InPCGStack)
+void FPVEditor::OnGraphExecuted(IPCGBaseSubsystem* InSubsystem, IPCGGraphExecutionSource* InExecutionSource, EPCGGenerationStatus InGenerationStatus)
 {
 	UpdateInspectedCollection();
 	UpdateStats();
+}
 
+void FPVEditor::OnNodeInspected(const FPCGStack& InPCGStack)
+{
 	if (NodeBeingInspected)
 	{
 		if (UPVEditorMode* const EditorMode = GetEditorMode())
@@ -479,7 +493,7 @@ void FPVEditor::ChangeNodeInspection(UPCGEditorGraphNodeBase* InNode)
 		return;
 	}
 
-	if (UPVBaseSettings* Settings = GetNodeSettings(InNode->GetPCGNode()))
+	if (IPVRenderSettings* Settings = GetRenderSettings(InNode->GetPCGNode()))
 	{
 		if (EditorViewport.IsValid())
 		{
diff --git a/Engine/Plugins/Experimental/ProceduralVegetationEditor/Source/ProceduralVegetationEditor/Private/PVEditor.h b/Engine/Plugins/Experimental/ProceduralVegetationEditor/Source/ProceduralVegetationEditor/Private/PVEditor.h
--- a/Engine/Plugins/Experimental/ProceduralVegetationEditor/Source/ProceduralVegetationEditor/Private/PVEditor.h
+++ b/Engine/Plugins/Experimental/ProceduralVegetationEditor/Source/ProceduralVegetationEditor/Private/PVEditor.h
@@ -74,6 +74,9 @@ private:
 
 	const UPVData* GetDataFromPin(const UPCGPin* InPin);
 
+	void OnGraphExecuted(IPCGBaseSubsystem* InSubsystem, IPCGGraphExecutionSource* InExecutionSource,
+	EPCGGenerationStatus InGenerationStatus);
+
 	void OnNodeInspected(const FPCGStack& InPCGStack);
 
 	void ChangeNodeInspection(UPCGEditorGraphNodeBase* InNode);
diff --git a/Engine/Plugins/Experimental/ProceduralVegetationEditor/Source/ProceduralVegetationEditor/Private/SPVEditorViewport.cpp b/Engine/Plugins/Experimental/ProceduralVegetationEditor/Source/ProceduralVegetationEditor/Private/SPVEditorViewport.cpp
--- a/Engine/Plugins/Experimental/ProceduralVegetationEditor/Source/ProceduralVegetationEditor/Private/SPVEditorViewport.cpp
+++ b/Engine/Plugins/Experimental/ProceduralVegetationEditor/Source/ProceduralVegetationEditor/Private/SPVEditorViewport.cpp
@@ -897,7 +897,7 @@ void SPVEditorViewport::OnVisualizationModeChanged(EPVVisualizationMode InMode)
 	}
 }
 
-void SPVEditorViewport::OnNodeInspectionChanged(UPVBaseSettings* InSettings)
+void SPVEditorViewport::OnNodeInspectionChanged(IPVRenderSettings* InSettings)
 {
 	SupportedRenderTypes = InSettings->GetSupportedRenderTypes();
 	CurrentVisualizationMode = EPVVisualizationMode::Default;
diff --git a/Engine/Plugins/Experimental/ProceduralVegetationEditor/Source/ProceduralVegetationEditor/Private/SPVEditorViewport.h b/Engine/Plugins/Experimental/ProceduralVegetationEditor/Source/ProceduralVegetationEditor/Private/SPVEditorViewport.h
--- a/Engine/Plugins/Experimental/ProceduralVegetationEditor/Source/ProceduralVegetationEditor/Private/SPVEditorViewport.h
+++ b/Engine/Plugins/Experimental/ProceduralVegetationEditor/Source/ProceduralVegetationEditor/Private/SPVEditorViewport.h
@@ -12,6 +12,8 @@
 
 #include "SPVEditorViewport.generated.h"
 
+class IPVRenderSettings;
+
 UENUM()
 enum class EPVVisualizationMode : uint8
 {
@@ -76,7 +78,7 @@ public:
 	FToolMenuEntry CreateSettingsToolbarMenu();
 	FToolMenuEntry CreateVisualizationModeToolbarMenu();
 
-	void OnNodeInspectionChanged(UPVBaseSettings* InSettings);
+	void OnNodeInspectionChanged(IPVRenderSettings* InSettings);
 	void ClearNodeInspection();
 
 	void SetOverlayText(const FText& CurrentlyLockedNodeName = FText::GetEmpty());
@@ -124,7 +126,7 @@ private:
 
 	TArray<EPVRenderType> SupportedRenderTypes;
 	EPVVisualizationMode CurrentVisualizationMode = EPVVisualizationMode::Default;
-	TObjectPtr<UPVBaseSettings> InspectedNodeSettings = nullptr;
+	IPVRenderSettings* InspectedNodeSettings = nullptr;
 
 	void OnVisualizationModeChanged(EPVVisualizationMode InMode);
 
diff --git a/Engine/Plugins/Experimental/ProceduralVegetationEditor/Source/ProceduralVegetationEditor/Private/Visualizations/PVDataVisualization.cpp b/Engine/Plugins/Experimental/ProceduralVegetationEditor/Source/ProceduralVegetationEditor/Private/Visualizations/PVDataVisualization.cpp
--- a/Engine/Plugins/Experimental/ProceduralVegetationEditor/Source/ProceduralVegetationEditor/Private/Visualizations/PVDataVisualization.cpp
+++ b/Engine/Plugins/Experimental/ProceduralVegetationEditor/Source/ProceduralVegetationEditor/Private/Visualizations/PVDataVisualization.cpp
@@ -76,14 +76,22 @@ void FPVDataVisualization::ExecuteDebugDisplay(FPCGContext* Context, const UPCGS
 {
 	return;
 }
- 
+
 FPCGSetupSceneFunc FPVDataVisualization::GetViewportSetupFunc(const UPCGSettingsInterface* SettingsInterface, const UPCGData* Data) const
 {
-	return [this, WeakData=TWeakObjectPtr<const UPVData>(Cast<UPVData>(Data)), WeakSettings = TWeakObjectPtr<const UPVBaseSettings>(Cast<UPVBaseSettings>(SettingsInterface))](
+	return [this, WeakData=TWeakObjectPtr<const UPVData>(Cast<UPVData>(Data)), WeakSettings = TWeakObjectPtr<const UObject>(Cast<UObject>(SettingsInterface))](
 		FPCGSceneSetupParams& InOutParams)
 	{
+		const IPVRenderSettings* RenderSettings = nullptr;
+
+		if (WeakSettings.IsValid() && WeakSettings->Implements<UPVRenderSettings>())
+		{
+			RenderSettings = Cast<IPVRenderSettings>(WeakSettings);
+		}
+			
 		check(InOutParams.Scene)
-		if (!WeakData.IsValid() || !WeakSettings.IsValid())
+			
+		if (!WeakData.IsValid() || RenderSettings == nullptr)
 		{
 			UE_LOG(LogProceduralVegetationDataVisualization, Error, TEXT("Failed to setup data viewport, the data was lost or invalid."));
 			return;
@@ -111,7 +119,7 @@ FPCGSetupSceneFunc FPVDataVisualization::GetViewportSetupFunc(const UPCGSettings
  
 		FBoxSphereBounds Bounds(ForceInit);
  
-		for (const EPVRenderType& RenderType : WeakSettings->GetCurrentRenderTypes())
+		for (const EPVRenderType& RenderType : RenderSettings->GetCurrentRenderTypes())
 		{
 			if (RenderMap.Contains(RenderType))
 			{
@@ -125,7 +133,7 @@ FPCGSetupSceneFunc FPVDataVisualization::GetViewportSetupFunc(const UPCGSettings
 			}
 		}
 
-		for (const auto& [MeshObject, Transform] : WeakSettings->GetViewportObjects())
+		for (const auto& [MeshObject, Transform] : RenderSettings->GetViewportObjects())
 		{
 			if (UStaticMesh* Mesh = Cast<UStaticMesh>(MeshObject))
 			{
@@ -138,7 +146,7 @@ FPCGSetupSceneFunc FPVDataVisualization::GetViewportSetupFunc(const UPCGSettings
 			}
 		}
 
-		if (PV::Utilities::DebugModeEnabled() && WeakSettings->bDebug)
+		if (PV::Utilities::DebugModeEnabled() && RenderSettings->IsDebug())
 		{
 			FPVDebugVisualizer::DrawDebugVisualizations(WeakData->GetDebugSettings().VisualizationSettings, WeakData->GetCollection(), InOutParams);
 			FPVDebugVisualizer::DrawLoopDebug(WeakData->GetDebugSettings().LoopDebugVisualizationSettings, WeakData->GetDebugSettings().bAutoFocusLoopDebug, InOutParams);
@@ -189,7 +197,7 @@ void FPVDataVisualization::MeshRenderer(const UPVData* Data, const FManagedArray
 		}
 	}
 }
- 
+
 void FPVDataVisualization::SkeletonRenderer(const UPVData* Data, const FManagedArrayCollection& InCollection, FPCGSceneSetupParams& InOutParams,
                                                       FBoxSphereBounds& OutBounds)
 {
