Add a boolean parameter to the static mesh and texture 2D import params to control split vs fused point output when generating growth data from branch hierarchies
Pass the new flag through to the growth data generation functions in both importer implementations

#Synced-CL 51662576

diff --git a/Engine/Plugins/Experimental/ProceduralVegetationEditor/Source/ProceduralVegetation/Private/Implementations/PVImporter_StaticMesh.cpp b/Engine/Plugins/Experimental/ProceduralVegetationEditor/Source/ProceduralVegetation/Private/Implementations/PVImporter_StaticMesh.cpp
--- a/Engine/Plugins/Experimental/ProceduralVegetationEditor/Source/ProceduralVegetation/Private/Implementations/PVImporter_StaticMesh.cpp
+++ b/Engine/Plugins/Experimental/ProceduralVegetationEditor/Source/ProceduralVegetation/Private/Implementations/PVImporter_StaticMesh.cpp
@@ -1153,7 +1153,7 @@ PV::StaticMeshImport::EImportResult PV::StaticMeshImport::ImportGrowthDataFromSt
 		NextBranchPointAttributeDefinition.CopyAttributeToCollection(DynamicMesh, DebugCollection);
 	}
 
-	if (!PV::ImportHelper::GenerateGrowthDataFromBranchHierarchies(Output.GrowthDataCollection, BranchHierarchyDescriptions))
+	if (!PV::ImportHelper::GenerateGrowthDataFromBranchHierarchies(Output.GrowthDataCollection, BranchHierarchyDescriptions, InParams.bUseSplitPoints))
 	{
 		return EImportResult::FailedToGenerateGrowthData;
 	}
diff --git a/Engine/Plugins/Experimental/ProceduralVegetationEditor/Source/ProceduralVegetation/Private/Implementations/PVImporter_Texture2D.cpp b/Engine/Plugins/Experimental/ProceduralVegetationEditor/Source/ProceduralVegetation/Private/Implementations/PVImporter_Texture2D.cpp
--- a/Engine/Plugins/Experimental/ProceduralVegetationEditor/Source/ProceduralVegetation/Private/Implementations/PVImporter_Texture2D.cpp
+++ b/Engine/Plugins/Experimental/ProceduralVegetationEditor/Source/ProceduralVegetation/Private/Implementations/PVImporter_Texture2D.cpp
@@ -1084,7 +1084,7 @@ PV::Texture2DImport::EImportResult PV::Texture2DImport::ImportGrowthDataFromText
 	MapPlantIndicesToBranchHierarchies(PlantMapping, DynamicMeshBranchHierarchies, BranchHierarchiesToPlantIndices);
 
 	// TODO: It would probably be faster to combine all the outputs of the PlantCollections
-	if (!PV::ImportHelper::GenerateGrowthDataFromBranchHierarchies(Output.CombinedPlantCollection, BranchHierarchyDescriptions))
+	if (!PV::ImportHelper::GenerateGrowthDataFromBranchHierarchies(Output.CombinedPlantCollection, BranchHierarchyDescriptions, InParams.bUseSplitPoints))
 	{
 		return EImportResult::FailedToGenerateGrowthData;
 	}
@@ -1098,7 +1098,7 @@ PV::Texture2DImport::EImportResult PV::Texture2DImport::ImportGrowthDataFromText
 			continue;
 		}
 
-		if (!PV::ImportHelper::GenerateGrowthDataFromBranchHierarchy(Output.PlantCollections[PlantIndex], BranchHierarchyDescriptions[i]))
+		if (!PV::ImportHelper::GenerateGrowthDataFromBranchHierarchy(Output.PlantCollections[PlantIndex], BranchHierarchyDescriptions[i], InParams.bUseSplitPoints))
 		{
 			return EImportResult::FailedToGenerateGrowthData;
 		}
diff --git a/Engine/Plugins/Experimental/ProceduralVegetationEditor/Source/ProceduralVegetation/Private/Params/PVImportStaticMeshParams.h b/Engine/Plugins/Experimental/ProceduralVegetationEditor/Source/ProceduralVegetation/Private/Params/PVImportStaticMeshParams.h
--- a/Engine/Plugins/Experimental/ProceduralVegetationEditor/Source/ProceduralVegetation/Private/Params/PVImportStaticMeshParams.h
+++ b/Engine/Plugins/Experimental/ProceduralVegetationEditor/Source/ProceduralVegetation/Private/Params/PVImportStaticMeshParams.h
@@ -39,7 +39,10 @@ struct FPVImportStaticMeshParams
 	UPROPERTY(EditAnywhere, Category = "Visualization")
 	bool bShowFoundBranchCurves = true;
 
-	UPROPERTY(EditAnywhere, Category = "Debug State")
+	UPROPERTY(EditAnywhere, Category = "DebugSettings")
+	bool bUseSplitPoints = false;
+
+	UPROPERTY(EditAnywhere, Category = "DebugSettings")
 	EPVImportStaticMeshDebugState DebugState = EPVImportStaticMeshDebugState::None;
 };
 
diff --git a/Engine/Plugins/Experimental/ProceduralVegetationEditor/Source/ProceduralVegetation/Private/Params/PVImportTexture2DParams.h b/Engine/Plugins/Experimental/ProceduralVegetationEditor/Source/ProceduralVegetation/Private/Params/PVImportTexture2DParams.h
--- a/Engine/Plugins/Experimental/ProceduralVegetationEditor/Source/ProceduralVegetation/Private/Params/PVImportTexture2DParams.h
+++ b/Engine/Plugins/Experimental/ProceduralVegetationEditor/Source/ProceduralVegetation/Private/Params/PVImportTexture2DParams.h
@@ -73,6 +73,9 @@ struct FPVImportTexture2DParams
 	UPROPERTY(EditAnywhere, Category = "Plant Geneneration|Tip detection", meta = (ClampMin=0, ClampMax=0.01, UIMin=0, UIMax=0.01))
 	float MaxTipAngleSearchDist = 0.005;
 
+	UPROPERTY(EditAnywhere, Category = "DebugSettings")
+	bool bUseSplitPoints = false;
+
 	UPROPERTY(EditAnywhere, Category = "DebugSettings")
 	EPVImportTexture2DDebugState DebugState = EPVImportTexture2DDebugState::None;
 };
