Fix carve operation not removing all points when carving trees generated through GenerateGrowthDataFromBranchHierarchy
Replace the shallow two-level branch removal logic with a recursive MarkSubtreeForRemoval helper to correctly remove entire branch subtrees at arbitrary depth
Add a regression test covering deep branch hierarchies to verify the fix

#Synced-CL 51662576

diff --git a/Engine/Plugins/Experimental/ProceduralVegetationEditor/Source/ProceduralVegetation/Private/Implementations/PVCarve.cpp b/Engine/Plugins/Experimental/ProceduralVegetationEditor/Source/ProceduralVegetation/Private/Implementations/PVCarve.cpp
--- a/Engine/Plugins/Experimental/ProceduralVegetationEditor/Source/ProceduralVegetation/Private/Implementations/PVCarve.cpp
+++ b/Engine/Plugins/Experimental/ProceduralVegetationEditor/Source/ProceduralVegetation/Private/Implementations/PVCarve.cpp
@@ -8,6 +8,38 @@
 #include "Facades/PVPointFacade.h"
 #include "Facades/PVTreeFacade.h"
 
+namespace
+{
+	static void MarkSubtreeForRemoval(
+		const int32 BranchIndex,
+		const PV::Facades::FBranchFacade& BranchFacadeSource,
+		const PV::Facades::FFoliageFacade& FoliageFacadeSource,
+		const TMap<int32, int32>& BranchNumbersToBranchIDs,
+		TArray<bool>& BranchesToRemove,
+		TArray<bool>& PointsToRemove,
+		TArray<bool>& FoliageInstancesToRemove)
+	{
+		BranchesToRemove[BranchIndex] = true;
+		for (const int32& P : BranchFacadeSource.GetPoints(BranchIndex))
+		{
+			PointsToRemove[P] = true;
+		}
+		for (const int32& FId : FoliageFacadeSource.GetFoliageEntryIdsForBranch(BranchIndex))
+		{
+			FoliageInstancesToRemove[FId] = true;
+		}
+
+		for (const int32& ChildBranchNumber : BranchFacadeSource.GetChildren(BranchIndex))
+		{
+			const int32* ChildIndexPtr = BranchNumbersToBranchIDs.Find(ChildBranchNumber);
+			if (ChildIndexPtr && !BranchesToRemove[*ChildIndexPtr])
+			{
+				MarkSubtreeForRemoval(*ChildIndexPtr, BranchFacadeSource, FoliageFacadeSource,
+					BranchNumbersToBranchIDs, BranchesToRemove, PointsToRemove, FoliageInstancesToRemove);
+			}
+		}
+	}
+}
 
 void FPVCarve::ApplyCarve(FManagedArrayCollection& OutCollection, const FManagedArrayCollection& SourceCollection,
                           const ECarveBasis CarveBasis, const float Carve)
@@ -278,42 +310,15 @@ void FPVCarve::CarveFromTop(FManagedArrayCollection& OutCollection, const FManag
 				continue;
 			}
 
-			// Remove all branches, child branches and foliage instances beyond last point
+			// Remove all branches in the subtree beyond the last kept point
 			for (const int32& BranchNumber : BranchFacadeSource.GetChildren(BranchIndex))
 			{
 				if (BranchNumbersToLengthFromRoots.Contains(BranchNumber) && BranchNumbersToLengthFromRoots[BranchNumber] > PointFacadeSource.
 					GetLengthFromRoot(BranchPoints[EndIndex]))
 				{
 					const int32 Index = BranchNumbersToBranchIDs[BranchNumber];
-					BranchesToRemove[Index] = true;
-					for (const int32& P : BranchFacadeSource.GetPoints(Index))
-					{
-						PointsToRemove[P] = true;
-					}
-					for (const int32& FId : FoliageFacadeSource.GetFoliageEntryIdsForBranch(Index))
-					{
-						FoliageInstancesToRemove[FId] = true;
-					}
-
-					const TArray<int32>& BranchChildren = BranchFacadeSource.GetChildren(Index);
-					for (int32 k = 0; k < BranchChildren.Num(); ++k)
-					{
-						if (!BranchNumbersToBranchIDs.Contains(BranchChildren[k]))
-						{
-							continue;
-						}
-
-						const int32 ChildIndex = BranchNumbersToBranchIDs[BranchChildren[k]];
-						BranchesToRemove[ChildIndex] = true;
-						for (const int32& P : BranchFacadeSource.GetPoints(ChildIndex))
-						{
-							PointsToRemove[P] = true;
-						}
-						for (const int32& FId : FoliageFacadeSource.GetFoliageEntryIdsForBranch(ChildIndex))
-						{
-							FoliageInstancesToRemove[FId] = true;
-						}
-					}
+					MarkSubtreeForRemoval(Index, BranchFacadeSource, FoliageFacadeSource,
+						BranchNumbersToBranchIDs, BranchesToRemove, PointsToRemove, FoliageInstancesToRemove);
 				}
 			}
 
@@ -429,41 +434,14 @@ void FPVCarve::CarveFromBottom(FManagedArrayCollection& OutCollection, const FMa
 		}
 		EndIndex--;
 
-		// Remove all branches, child branches and foliage instances before current point
+		// Remove all branches in the subtree before current point
 		for (const int32& BranchNumber : BranchFacadeSource.GetChildren(TrunkIndex))
 		{
 			if (BranchNumbersToLengthFromRoots.Contains(BranchNumber) && BranchNumbersToLengthFromRoots[BranchNumber] < CarveLengthFromRoot)
 			{
 				const int32 Index = BranchNumbersToBranchIDs[BranchNumber];
-				BranchesToRemove[Index] = true;
-				for (const int32& P : BranchFacadeSource.GetPoints(Index))
-				{
-					PointsToRemove[P] = true;
-				}
-				for (const int32& FId : FoliageFacadeSource.GetFoliageEntryIdsForBranch(Index))
-				{
-					FoliageInstancesToRemove[FId] = true;
-				}
-
-				const TArray<int32>& BranchChildren = BranchFacadeSource.GetChildren(Index);
-				for (int32 k = 0; k < BranchChildren.Num(); ++k)
-				{
-					if (!BranchNumbersToBranchIDs.Contains(BranchChildren[k]))
-					{
-						continue;
-					}
-
-					const int32 ChildIndex = BranchNumbersToBranchIDs[BranchChildren[k]];
-					BranchesToRemove[ChildIndex] = true;
-					for (const int32& P : BranchFacadeSource.GetPoints(ChildIndex))
-					{
-						PointsToRemove[P] = true;
-					}
-					for (const int32& FId : FoliageFacadeSource.GetFoliageEntryIdsForBranch(ChildIndex))
-					{
-						FoliageInstancesToRemove[FId] = true;
-					}
-				}
+				MarkSubtreeForRemoval(Index, BranchFacadeSource, FoliageFacadeSource,
+					BranchNumbersToBranchIDs, BranchesToRemove, PointsToRemove, FoliageInstancesToRemove);
 			}
 		}
 
diff --git a/Engine/Plugins/Experimental/ProceduralVegetationEditor/Source/ProceduralVegetation/Private/Tests/Elements/PVCarveTest.cpp b/Engine/Plugins/Experimental/ProceduralVegetationEditor/Source/ProceduralVegetation/Private/Tests/Elements/PVCarveTest.cpp
--- a/Engine/Plugins/Experimental/ProceduralVegetationEditor/Source/ProceduralVegetation/Private/Tests/Elements/PVCarveTest.cpp
+++ b/Engine/Plugins/Experimental/ProceduralVegetationEditor/Source/ProceduralVegetation/Private/Tests/Elements/PVCarveTest.cpp
@@ -4,6 +4,10 @@
 
 #include "Tests/PVTestsCommon.h"
 #include "DataTypes/PVGrowthData.h"
+#include "Helpers/PVImportHelpers.h"
+#include "Implementations/PVCarve.h"
+#include "Facades/PVBranchFacade.h"
+#include "Facades/PVPointFacade.h"
 #include "Nodes/PVCarveSettings.h"
 #include "PCGContext.h"
 
@@ -59,4 +63,89 @@ PV_CARVE_TEST(Default)
 	return true;
 }
 
+PV_CARVE_TEST(DeepHierarchyCarve)
+{
+	// Regression test: CarveFromTop must remove the full subtree of a carved branch,
+	// not just direct children and grandchildren (depth 2). Without the fix, points
+	// belonging to branches at depth 3+ from the cut survive as floating artefacts.
+	//
+	// Hierarchy (4 levels):
+	//   B0 (trunk):  P0(z=0) - P1(z=100) - P2(z=200) - P3(z=300) - P4(z=400)
+	//   B1 (depth1): P2(z=200) - P5 - P6
+	//   B2 (depth2): P6 - P7 - P8
+	//   B3 (depth3): P8 - P9 - P10   <-- before fix: P9 and P10 survive as floating points
+	//
+	// Carving at 0.75 cuts the trunk at ~z=100, so B1 and all descendants must be removed.
+
+	FPVBranchHierarchyDescription Hierarchy;
+	Hierarchy.Points = {
+		FVector3f(  0.f, 0.f,   0.f), // P0 - trunk base
+		FVector3f(  0.f, 0.f, 100.f), // P1
+		FVector3f(  0.f, 0.f, 200.f), // P2 - B1 junction
+		FVector3f(  0.f, 0.f, 300.f), // P3
+		FVector3f(  0.f, 0.f, 400.f), // P4 - trunk tip
+		FVector3f( 10.f, 0.f, 200.f), // P5
+		FVector3f( 20.f, 0.f, 200.f), // P6 - B2 junction
+		FVector3f( 20.f, 0.f, 210.f), // P7
+		FVector3f( 20.f, 0.f, 220.f), // P8 - B3 junction
+		FVector3f( 20.f, 0.f, 230.f), // P9
+		FVector3f( 20.f, 0.f, 240.f), // P10
+	};
+	Hierarchy.PointsRadii.Init(1.0f, Hierarchy.Points.Num());
+
+	// B0 - trunk
+	FPVBranchDescription& B0 = Hierarchy.Branches.AddDefaulted_GetRef();
+	B0.BranchIndex = 0;
+	B0.ParentBranchIndex = INDEX_NONE;
+	B0.PointIndices = {0, 1, 2, 3, 4};
+	B0.ChildBranchIndices = {1};
+
+	// B1 - direct child of trunk (depth 1 from cut)
+	FPVBranchDescription& B1 = Hierarchy.Branches.AddDefaulted_GetRef();
+	B1.BranchIndex = 1;
+	B1.ParentBranchIndex = 0;
+	B1.PointIndices = {2, 5, 6};
+	B1.ChildBranchIndices = {2};
+
+	// B2 - grandchild (depth 2 from cut)
+	FPVBranchDescription& B2 = Hierarchy.Branches.AddDefaulted_GetRef();
+	B2.BranchIndex = 2;
+	B2.ParentBranchIndex = 1;
+	B2.PointIndices = {6, 7, 8};
+	B2.ChildBranchIndices = {3};
+
+	// B3 - great-grandchild (depth 3 from cut - the regression case)
+	FPVBranchDescription& B3 = Hierarchy.Branches.AddDefaulted_GetRef();
+	B3.BranchIndex = 3;
+	B3.ParentBranchIndex = 2;
+	B3.PointIndices = {8, 9, 10};
+	B3.ChildBranchIndices = {};
+
+	Hierarchy.RootBranchIndex = 0;
+
+	FManagedArrayCollection SourceCollection;
+	const bool bOk = PV::ImportHelper::GenerateGrowthDataFromBranchHierarchy(SourceCollection, Hierarchy, /*bSplitPoints=*/false);
+	UTEST_TRUE("GenerateGrowthDataFromBranchHierarchy succeeded", bOk);
+
+	// Carve 75% from top: cuts the trunk at ~z=100, so B1 and all descendants should be removed
+	FManagedArrayCollection OutCollection = SourceCollection;
+	FPVCarve::ApplyCarve(OutCollection, SourceCollection, ECarveBasis::LengthFromRoot, 0.75f);
+
+	const PV::Facades::FPointFacade OutPointFacade(OutCollection);
+	const PV::Facades::FBranchFacade OutBranchFacade(OutCollection);
+
+	// All branches except the trunk must be removed
+	UTEST_EQUAL("Only trunk branch remains", OutBranchFacade.GetElementCount(), 1);
+
+	// No points at z >= 200 should survive: B1/B2/B3 all start at z=200 or higher.
+	// Before the fix, B3 points at z=230 and z=240 would survive as floating artefacts.
+	for (int32 i = 0; i < OutPointFacade.GetElementCount(); ++i)
+	{
+		const FVector3f& Pos = OutPointFacade.GetPosition(i);
+		UTEST_TRUE(FString::Printf(TEXT("Floating deep point at z=%.1f was not removed"), Pos.Z), Pos.Z < 200.0f);
+	}
+
+	return true;
+}
+
 #endif //WITH_DEV_AUTOMATION_TESTS
\ No newline at end of file
