Add BranchGeneration skeleton visualization mode that colors tree points by branch hierarchy depth using parent chain length
Fix mesh point rendering to skip mesh-based drawing when in ColorRamp visualization mode
Add early-out guard in skeleton renderer for invalid growth data and export IsValidGrowthData as a module-level API

#Synced-CL 51662576

diff --git a/Engine/Plugins/Experimental/ProceduralVegetationEditor/Source/ProceduralVegetation/Private/Facades/PVTreeFacade.cpp b/Engine/Plugins/Experimental/ProceduralVegetationEditor/Source/ProceduralVegetation/Private/Facades/PVTreeFacade.cpp
--- a/Engine/Plugins/Experimental/ProceduralVegetationEditor/Source/ProceduralVegetation/Private/Facades/PVTreeFacade.cpp
+++ b/Engine/Plugins/Experimental/ProceduralVegetationEditor/Source/ProceduralVegetation/Private/Facades/PVTreeFacade.cpp
@@ -6,6 +6,8 @@
 #include "Facades/PVFoliageFacade.h"
 #include "Facades/PVMetaInfoFacade.h"
 #include "Facades/PVPointFacade.h"
+#include "Helpers/PVAttributesHelper.h"
+#include "Helpers/PVPlantTraversalHelper.h"
 
 #include "GeometryCollection/ManagedArrayAccessor.h"
 
@@ -395,6 +397,30 @@ namespace PV::Facades
 				}
 				break;
 			}
+		case ESkeletonVisualizationModes::BranchGeneration:
+			{
+				const PV::FBranchPointsAttributeConstView       BranchPointsAttribute       = PV::FBranchPointsAttribute::FindAttribute(Collection);
+				const PV::FBranchParentNumberAttributeConstView BranchParentNumberAttribute = PV::FBranchParentNumberAttribute::FindAttribute(Collection);
+				const PV::FBranchNumberAttributeConstView       BranchNumberAttribute       = PV::FBranchNumberAttribute::FindAttribute(Collection);
+				const PV::FBranchParentsAttributeConstView      BranchParentsAttribute      = PV::FBranchParentsAttribute::FindAttribute(Collection);
+				if (PV::AttributesHelper::ValidateAttributeCollection(BranchPointsAttribute, BranchParentNumberAttribute, BranchNumberAttribute, BranchParentsAttribute))
+				{
+					const bool bHasFusedPoints = PV::PlantTraversalHelper::HasFusedPoints(BranchPointsAttribute, BranchParentNumberAttribute, BranchNumberAttribute);
+					for (int32 BranchIndex = 0; BranchIndex < BranchParentsAttribute.Num(); ++BranchIndex)
+					{
+						const int32 GenerationNumber = BranchParentsAttribute[BranchIndex].Num();
+						const bool bIsTrunk = PV::PlantTraversalHelper::IsTrunk(BranchParentNumberAttribute, BranchIndex);
+						const int32 StartIndex = bHasFusedPoints && !bIsTrunk ? 1 : 0;
+
+						const TArray<int32>& BranchPoints = BranchPointsAttribute[BranchIndex];
+						for (int32 BranchPointIndex = StartIndex; BranchPointIndex < BranchPoints.Num(); ++BranchPointIndex)
+						{
+							VisualizationValues[BranchPoints[BranchPointIndex]] = GenerationNumber;
+						}
+					}
+				}
+				break;
+			}
 		default:
 			{
 				VisualizationValues.Empty();
diff --git a/Engine/Plugins/Experimental/ProceduralVegetationEditor/Source/ProceduralVegetation/Public/Facades/PVTreeFacade.h b/Engine/Plugins/Experimental/ProceduralVegetationEditor/Source/ProceduralVegetation/Public/Facades/PVTreeFacade.h
--- a/Engine/Plugins/Experimental/ProceduralVegetationEditor/Source/ProceduralVegetation/Public/Facades/PVTreeFacade.h
+++ b/Engine/Plugins/Experimental/ProceduralVegetationEditor/Source/ProceduralVegetation/Public/Facades/PVTreeFacade.h
@@ -27,7 +27,8 @@ enum class ESkeletonVisualizationModes : uint8
 	Radical UMETA(Category="Auxin"),
 	Ethylene UMETA(Category="Ethlyene"),
 	AuxinRetention UMETA(Category="Ethlyene"),
-	Cytokinin UMETA(Category="Cytokinin")
+	Cytokinin UMETA(Category="Cytokinin"),
+	BranchGeneration UMETA(Category="Debug")
 };
 
 namespace PV::Facades
diff --git a/Engine/Plugins/Experimental/ProceduralVegetationEditor/Source/ProceduralVegetation/Public/Helpers/PVUtilities.h b/Engine/Plugins/Experimental/ProceduralVegetationEditor/Source/ProceduralVegetation/Public/Helpers/PVUtilities.h
--- a/Engine/Plugins/Experimental/ProceduralVegetationEditor/Source/ProceduralVegetation/Public/Helpers/PVUtilities.h
+++ b/Engine/Plugins/Experimental/ProceduralVegetationEditor/Source/ProceduralVegetation/Public/Helpers/PVUtilities.h
@@ -34,7 +34,7 @@ namespace PV::Utilities
 	int32 PROCEDURALVEGETATION_API GetMeshTriangles(const FString InMeshPath);
 	FLinearColor PROCEDURALVEGETATION_API GetRandomHueColor(float Alpha);
 
-	bool IsValidGrowthData(const FManagedArrayCollection& Collection);
+	bool PROCEDURALVEGETATION_API IsValidGrowthData(const FManagedArrayCollection& Collection);
 
 	template <typename T>
 	TArray<T> LerpArrayElements(const TArray<T>& InArray1, const TArray<T>& InArray2, const float Alpha);
diff --git a/Engine/Plugins/Experimental/ProceduralVegetationEditor/Source/ProceduralVegetationEditor/Private/Visualizations/DebugVisualization/PVDebugVisualizationBase.cpp b/Engine/Plugins/Experimental/ProceduralVegetationEditor/Source/ProceduralVegetationEditor/Private/Visualizations/DebugVisualization/PVDebugVisualizationBase.cpp
--- a/Engine/Plugins/Experimental/ProceduralVegetationEditor/Source/ProceduralVegetationEditor/Private/Visualizations/DebugVisualization/PVDebugVisualizationBase.cpp
+++ b/Engine/Plugins/Experimental/ProceduralVegetationEditor/Source/ProceduralVegetationEditor/Private/Visualizations/DebugVisualization/PVDebugVisualizationBase.cpp
@@ -652,7 +652,8 @@ void FPVDebugVisualizationBase::DrawAttributes(FVisualizerDrawContext& InContext
 
 	if (PointsToDraw.Num() > 0)
 	{
-		if (VisualizationSettings.bDrawPointAsMesh)
+		const bool bCanDrawAsMeshPoints = VisualizationSettings.VisualizationMode != EPVDebugValueVisualizationMode::ColorRamp;
+		if (VisualizationSettings.bDrawPointAsMesh && bCanDrawAsMeshPoints)
 		{
 			DrawAsMeshPoints(GetOrCreatePointComponent(InContext.SceneSetupParams), PointsToDraw);
 		}
diff --git a/Engine/Plugins/Experimental/ProceduralVegetationEditor/Source/ProceduralVegetationEditor/Private/Visualizations/PVDataVisualization.cpp b/Engine/Plugins/Experimental/ProceduralVegetationEditor/Source/ProceduralVegetationEditor/Private/Visualizations/PVDataVisualization.cpp
--- a/Engine/Plugins/Experimental/ProceduralVegetationEditor/Source/ProceduralVegetationEditor/Private/Visualizations/PVDataVisualization.cpp
+++ b/Engine/Plugins/Experimental/ProceduralVegetationEditor/Source/ProceduralVegetationEditor/Private/Visualizations/PVDataVisualization.cpp
@@ -202,6 +202,11 @@ void FPVDataVisualization::SkeletonRenderer(const UPVData* Data, const FManagedA
                                                       FBoxSphereBounds& OutBounds)
 {
 	TRACE_CPUPROFILER_EVENT_SCOPE(FPVDataVisualization::SkeletonRenderer);
+
+	if (!PV::Utilities::IsValidGrowthData(InCollection))
+	{
+		return;
+	}
  
 	const TObjectPtr<UPVSkeletonVisualizerComponent> VisualizerComponent = NewObject<UPVSkeletonVisualizerComponent>(
 		GetTransientPackage(), NAME_None, RF_Transient);
