[PVE] Add phyllotaxy settings to both grafter and foliage distributor node parametric distribution settings

- Adds phyllotaxy settings for parametric based distribution for both grafter and foliage distributor nodes
- Fixes issue where ramp basis were not taken into account for axil angle and scale settings
- Uses distance culling to safeguard against near-duplicate attachment points being generated
- Fixes issue where graft instances at tip were getting incorrect point scales computed
- Fixes issue where graft instances were getting incorrect branch parents set

#Synced-CL 51747514

diff --git a/Engine/Plugins/Experimental/ProceduralVegetationEditor/Source/ProceduralVegetation/Private/Helpers/PVHormoneDistributionHelper.cpp b/Engine/Plugins/Experimental/ProceduralVegetationEditor/Source/ProceduralVegetation/Private/Helpers/PVHormoneDistributionHelper.cpp
--- a/Engine/Plugins/Experimental/ProceduralVegetationEditor/Source/ProceduralVegetation/Private/Helpers/PVHormoneDistributionHelper.cpp
+++ b/Engine/Plugins/Experimental/ProceduralVegetationEditor/Source/ProceduralVegetation/Private/Helpers/PVHormoneDistributionHelper.cpp
@@ -21,48 +21,13 @@ void PV::HormoneDistributionHelper::ComputeHormoneBasedAttachmentPoints(
 
 	const int32 NumBranches = BranchPointsAttribute.Num();
 
-	float PhyllotaxyFormationLeaf = 180;
-	float MinBudsLeaf = 1;
-	float MaxBudsLeaf = 1;
-	bool ResetPhyllotaxyLeaf = PhyllotaxySettings.ResetPhyllotaxy;
-	float PhyllotaxyOffsetLeaf = PhyllotaxySettings.PhyllotaxyOffset;
-
-	switch (PhyllotaxySettings.PhyllotaxyType)
-	{
-	case EPhyllotaxyType::Alternate: PhyllotaxyFormationLeaf = 180 + PhyllotaxySettings.PhyllotaxyAdditionalAngle;
-		MinBudsLeaf = MaxBudsLeaf = 1;
-		break;
-	case EPhyllotaxyType::Opposite: PhyllotaxyFormationLeaf = PhyllotaxySettings.PhyllotaxyAdditionalAngle;
-		MinBudsLeaf = MaxBudsLeaf = 2;
-		break;
-	case EPhyllotaxyType::Decussate: PhyllotaxyFormationLeaf = 90 + PhyllotaxySettings.PhyllotaxyAdditionalAngle;
-		MinBudsLeaf = MaxBudsLeaf = 2;
-		break;
-	case EPhyllotaxyType::Whorled: PhyllotaxyFormationLeaf = 90 + PhyllotaxySettings.PhyllotaxyAdditionalAngle;
-		MinBudsLeaf = PhyllotaxySettings.MinimumNodeBuds;
-		MaxBudsLeaf = PhyllotaxySettings.MaximumNodeBuds;
-		break;
-	case EPhyllotaxyType::Spiral:
-		{
-			// These leaf arrangement configurations where Distichous is two leaves, Tristichous is three etc.
-			switch (PhyllotaxySettings.PhyllotaxyFormation)
-			{
-			case EPhyllotaxyFormation::Distichous: PhyllotaxyFormationLeaf = 180 + PhyllotaxySettings.PhyllotaxyAdditionalAngle;
-				break;
-			case EPhyllotaxyFormation::Tristichous: PhyllotaxyFormationLeaf = 120 + PhyllotaxySettings.PhyllotaxyAdditionalAngle;
-				break;
-			case EPhyllotaxyFormation::Pentastichous: PhyllotaxyFormationLeaf = 144 + PhyllotaxySettings.PhyllotaxyAdditionalAngle;
-				break;
-			case EPhyllotaxyFormation::Octastichous: PhyllotaxyFormationLeaf = 135 + PhyllotaxySettings.PhyllotaxyAdditionalAngle;
-				break;
-			case EPhyllotaxyFormation::Parastichous: PhyllotaxyFormationLeaf = PhyllotaxySettings.PhyllotaxyAdditionalAngle;
-				break;
-			}
-
-			MinBudsLeaf = MaxBudsLeaf = 1;
-			break;
-		}
-	}
+	const auto [
+			PhyllotaxyFormationLeaf,
+			MinBudsLeaf,
+			MaxBudsLeaf,
+			bResetPhyllotaxyLeaf,
+			PhyllotaxyOffsetLeaf] =
+		PhyllotaxySettingsHelper::ResolvePhyllotaxy(PhyllotaxySettings);
 
 	for (int32 BranchIndex = 0; BranchIndex < NumBranches; ++BranchIndex)
 	{
@@ -107,7 +72,7 @@ void PV::HormoneDistributionHelper::ComputeHormoneBasedAttachmentPoints(
 		FVector3f InitialAxillaryDirection(BudDirectionAttribute[BranchPoints[0]].Axillary);
 		const FVector3f InitialUpVector(BudDirectionAttribute[BranchPoints[0]].UpVector);
 
-		if (ResetPhyllotaxyLeaf)
+		if (bResetPhyllotaxyLeaf)
 		{
 			InitialAxillaryDirection = FVector3f::CrossProduct(InitialApicalDirection, InitialUpVector);
 		}
diff --git a/Engine/Plugins/Experimental/ProceduralVegetationEditor/Source/ProceduralVegetation/Private/Helpers/PVParametricDistributionHelper.cpp b/Engine/Plugins/Experimental/ProceduralVegetationEditor/Source/ProceduralVegetation/Private/Helpers/PVParametricDistributionHelper.cpp
--- a/Engine/Plugins/Experimental/ProceduralVegetationEditor/Source/ProceduralVegetation/Private/Helpers/PVParametricDistributionHelper.cpp
+++ b/Engine/Plugins/Experimental/ProceduralVegetationEditor/Source/ProceduralVegetation/Private/Helpers/PVParametricDistributionHelper.cpp
@@ -6,13 +6,67 @@
 #include "Facades/PVBranchFacade.h"
 #include "Helpers/PVUtilities.h"
 
+float PV::ParametricDistributionHelper::IndexedRandom01(const int32 Seed)
+{
+	const FRandomStream LocalRandom(Seed);
+	return LocalRandom.GetFraction();
+}
+
+float PV::ParametricDistributionHelper::IndexedRandomRange(const int32 Seed, const float MinValue, const float MaxValue)
+{
+	return FMath::Lerp(MinValue, MaxValue, IndexedRandom01(Seed));
+}
+
+int32 PV::ParametricDistributionHelper::IndexedRandomIntInclusive(const int32 Seed, const int32 MinValue, const int32 MaxValue)
+{
+	if (MinValue >= MaxValue)
+	{
+		return MinValue;
+	}
+
+	const FRandomStream LocalRandom(Seed);
+	return LocalRandom.RandRange(MinValue, MaxValue);
+}
+
+float PV::ParametricDistributionHelper::ComputeAttachmentScale(
+	const int32 BudIndex,
+	const int32 RandomSeed,
+	const float PlantGradient,
+	const float BranchGradient,
+	const float LengthRatio,
+	const FParametricSettings::FScaleSettings& ScaleSettings)
+{
+	const float RampLookup =
+		(ScaleSettings.ScaleRampBasis == FParametricSettings::EDistributionBasis::Branch)
+		? BranchGradient
+		: PlantGradient;
+
+	float ScaleRampValue = ScaleSettings.ScaleRamp->GetRichCurveConst()->Eval(RampLookup);
+
+	if (ScaleSettings.ScaleRampBasis == FParametricSettings::EDistributionBasis::Branch)
+	{
+		ScaleRampValue = FMath::GetMappedRangeValueClamped(
+			FVector2f(0.0f, 1.0f),
+			FVector2f(0.0f, LengthRatio),
+			ScaleRampValue);
+	}
+
+	const float RandomScaleFactor = IndexedRandomRange(
+		BudIndex + RandomSeed + 36,
+		ScaleSettings.RandomizeScaleAngleMinimum,
+		ScaleSettings.RandomizeScaleAngleMaximum);
+
+	return ScaleRampValue * RandomScaleFactor * ScaleSettings.BaseScale;
+}
+
 void PV::ParametricDistributionHelper::ComputeAttachmentPoints(
 	FManagedArrayCollection& OutCollection,
 	const FNormalizedPointCaches& NormalizedPointCaches,
 	const FParametricSettings::FSpacingSettings& SpacingSettings,
+	const FParametricSettings::FPhyllotaxySettings& PhyllotaxySettings,
 	const FParametricSettings::FAngleSettings& AngleSettings,
 	const FParametricSettings::FScaleSettings& ScaleSettings,
-	FRandomStream& RandomStream,
+	const int32 RandomSeed,
 	TArray<DistributionHelper::FAttachmentPoint>& OutAttachmentPoints)
 {
 	auto PointPlantGradientAttribute = FPointPlantGradientAttribute::GetAttribute(OutCollection);
@@ -32,6 +86,14 @@ void PV::ParametricDistributionHelper::ComputeAttachmentPoints(
 		return;
 	}
 
+	const auto [
+			PhyllotaxyFormationLeaf,
+			MinBudsLeaf,
+			MaxBudsLeaf,
+			bResetPhyllotaxyLeaf,
+			PhyllotaxyOffsetLeaf] =
+		PhyllotaxySettingsHelper::ResolvePhyllotaxy(PhyllotaxySettings);
+	
 	// Only need to do this temporarily since legacy skeleton has inconsistent data in children
 	// TODO: Remove this once support for legacy skeletons is dropped
 	PV::Facades::FBranchFacade BranchFacadeOut(OutCollection);
@@ -67,6 +129,8 @@ void PV::ParametricDistributionHelper::ComputeAttachmentPoints(
 
 	AverageBranchLength = AverageBranchLength / BranchIndicesToBranchLengths.Num();
 	AverageBranchLength = FMath::Lerp(AverageBranchLength, MaxBranchLength, 0.33f);
+	
+	int32 GlobalBudIndex = 0;
 
 	for (int32 BranchIndex = 0; BranchIndex < NumOfBranches; BranchIndex++)
 	{
@@ -114,10 +178,36 @@ void PV::ParametricDistributionHelper::ComputeAttachmentPoints(
 		const float LFRofRoot = PointLengthFromRootAttribute[FirstPointIndex];
 		const float LFRofTip = PointLengthFromRootAttribute[TipPointIndex];
 
+		FVector3f PreviousAcceptedPosition = PointPositionAttribute[FirstPointIndex];
+
+		const FVector3f InitialApicalDirection = BudDirectionAttribute[FirstPointIndex].Apical.GetSafeNormal();
+		const FVector3f InitialUpVector = BudDirectionAttribute[FirstPointIndex].UpVector.GetSafeNormal();
+
+		FVector3f InitialAxillaryDirection = BudDirectionAttribute[FirstPointIndex].UpVector.GetSafeNormal();
+		if (bResetPhyllotaxyLeaf)
+		{
+			InitialAxillaryDirection = FVector3f::CrossProduct(InitialApicalDirection, InitialUpVector).GetSafeNormal();
+		}
+
+		if (InitialAxillaryDirection.IsNearlyZero())
+		{
+			FVector Axis1, Axis2;
+			FVector(InitialApicalDirection).FindBestAxisVectors(Axis1, Axis2);
+			InitialAxillaryDirection = FVector3f(Axis2);
+		}
+
+		if (FMath::Abs(PhyllotaxyOffsetLeaf) > KINDA_SMALL_NUMBER)
+		{
+			const FQuat4f OffsetQuat(InitialApicalDirection, FMath::DegreesToRadians(PhyllotaxyOffsetLeaf));
+			InitialAxillaryDirection = OffsetQuat.RotateVector(InitialAxillaryDirection).GetSafeNormal();
+		}
+
+		int32 AxillaryRotationIteration = 0;
+
 		for (int32 j = 0; j < LoopNumber; j++)
 		{
 			float BaseValue = static_cast<float>(j) / FMath::Max(LoopNumber - 1, 1);
-			float PlantGradient;
+			float PlantGradient = 0.0f;
 			float BranchGradient = 0.0f;
 			float AttachmentPointLengthFromRoot = 0.0f;
 			FVector3f Pos(0.0f, 0.0f, 0.0f);
@@ -133,22 +223,38 @@ void PV::ParametricDistributionHelper::ComputeAttachmentPoints(
 					FVector2f(0.0f, 1.0f),
 					FVector2f(SpacingSettings.RelativeStart, SpacingSettings.RelativeEnd),
 					PlantGradient);
+
 				if (PlantGradient < PlantGradientAtRoot)
 				{
 					continue;
 				}
 
-				Pos = SampleVectorByFloat(PlantGradient, SortedPointIndices, PointPlantGradientAttribute, PointPositionAttribute);
-				AttachmentPointLengthFromRoot = SampleFloatByFloat(PlantGradient, SortedPointIndices, PointPlantGradientAttribute,
+				Pos = SampleVectorByFloat(
+					PlantGradient,
+					SortedPointIndices,
+					PointPlantGradientAttribute,
+					PointPositionAttribute);
+
+				AttachmentPointLengthFromRoot = SampleFloatByFloat(
+					PlantGradient,
+					SortedPointIndices,
+					PointPlantGradientAttribute,
 					PointLengthFromRootAttribute);
+
 				BranchGradient = FMath::GetMappedRangeValueClamped(
 					FVector2f(LFRofRoot, LFRofTip),
 					FVector2f(0.0f, 1.0f),
 					AttachmentPointLengthFromRoot);
 
-				FindPointByNormalizedLengthFromRoot(BranchGradient, SortedPointIndices, PointLengthFromRootAttribute, LFRofRoot, LFRofTip,
+				FindPointByNormalizedLengthFromRoot(
+					BranchGradient,
+					SortedPointIndices,
+					PointLengthFromRootAttribute,
+					LFRofRoot,
+					LFRofTip,
 					PointIndexA,
-					PointIndexB, PointAlpha);
+					PointIndexB,
+					PointAlpha);
 			}
 			else // SpacingBasis is Branch
 			{
@@ -164,8 +270,15 @@ void PV::ParametricDistributionHelper::ComputeAttachmentPoints(
 					BaseValue = 1.0f;
 				}
 
-				if (!FindPointByNormalizedLengthFromRoot(BranchGradient, SortedPointIndices, PointLengthFromRootAttribute, LFRofRoot, LFRofTip,
-					PointIndexA, PointIndexB, PointAlpha))
+				if (!FindPointByNormalizedLengthFromRoot(
+					BranchGradient,
+					SortedPointIndices,
+					PointLengthFromRootAttribute,
+					LFRofRoot,
+					LFRofTip,
+					PointIndexA,
+					PointIndexB,
+					PointAlpha))
 				{
 					continue;
 				}
@@ -180,30 +293,37 @@ void PV::ParametricDistributionHelper::ComputeAttachmentPoints(
 					BranchGradient);
 			}
 
+			if (PointIndexA == INDEX_NONE || PointIndexB == INDEX_NONE)
+			{
+				continue;
+			}
+
 			const float PlantGradient0 = 1.0f - PointPlantGradientAttribute[PointIndexA];
 			const float PlantGradient1 = 1.0f - PointPlantGradientAttribute[PointIndexB];
 			PlantGradient = FMath::Lerp(PlantGradient0, PlantGradient1, PointAlpha);
 
-			const float ScaleRampLookUp =
-				(ScaleSettings.ScaleRampBasis == FParametricSettings::EDistributionBasis::Branch)
-				? BranchGradient
-				: PlantGradient;
-
-			float ScaleRampValue = ScaleSettings.ScaleRamp->GetRichCurveConst()->Eval(ScaleRampLookUp);
-
-			if (ScaleSettings.ScaleRampBasis == FParametricSettings::EDistributionBasis::Branch)
+			const float PreviousDistance = FVector3f::Distance(Pos, PreviousAcceptedPosition);
+			if (PreviousDistance < 0.001f)
 			{
-				ScaleRampValue = FMath::GetMappedRangeValueClamped(
-					FVector2f(0.0f, 1.0f),
-					FVector2f(0.0f, LengthRatio),
-					ScaleRampValue);
+				continue;
 			}
+			PreviousAcceptedPosition = Pos;
+
+			int32 NumBuds = IndexedRandomIntInclusive(
+				BranchIndex + 256,
+				MinBudsLeaf,
+				MaxBudsLeaf);
 
-			const float RandomScaleFactor = RandomStream.FRandRange(
-				ScaleSettings.RandomizeScaleAngleMinimum,
-				ScaleSettings.RandomizeScaleAngleMaximum);
+			NumBuds = FMath::Max(NumBuds, 1);
 
-			const float Scale = ScaleRampValue * RandomScaleFactor * ScaleSettings.BaseScale;
+			const bool bIsTipInstance =
+				FMath::IsNearlyEqual(AttachmentPointLengthFromRoot, LFRofTip, 0.001f) &&
+				(BranchGradient >= 0.999f || FMath::IsNearlyEqual(BranchGradient, SpacingSettings.RelativeEnd, 0.001f));
+			
+			if (j == 0 || bIsTipInstance)
+			{
+				NumBuds = 1;
+			}
 
 			const FVector3f BudDirectionApicalA = BudDirectionAttribute[PointIndexA].Apical;
 			const FVector3f BudDirectionApicalB = BudDirectionAttribute[PointIndexB].Apical;
@@ -213,83 +333,114 @@ void PV::ParametricDistributionHelper::ComputeAttachmentPoints(
 			const FVector3f ApicalDirection = FMath::Lerp(BudDirectionApicalA, BudDirectionApicalB, PointAlpha);
 			const FVector3f BudDirectionUp = FMath::Lerp(BudDirectionUpA, BudDirectionUpB, PointAlpha);
 
-			const float AxilRampValue = AngleSettings.AxilAngleRamp->GetRichCurveConst()->Eval(BaseValue);
-			float AxilAngleInDegrees = AngleSettings.AxilAngle * AxilRampValue;
-
-			AxilAngleInDegrees += RandomStream.FRandRange(AngleSettings.RandomizeAxilAngleMinimum, AngleSettings.RandomizeAxilAngleMaximum);
-			AxilAngleInDegrees = FMath::Clamp(AxilAngleInDegrees, -90.0f, 90.0f);
+			float SitePhaseDegrees = 0.0f;
+			if (j > 0)
+			{
+				SitePhaseDegrees = PhyllotaxyFormationLeaf * static_cast<float>(AxillaryRotationIteration);
+			}
 
-			const float Twist = FMath::DegreesToRadians(AngleSettings.Rotation) * static_cast<float>(j);
+			const FQuat4f SitePhaseQuat(InitialApicalDirection, FMath::DegreesToRadians(SitePhaseDegrees));
+			FVector3f SiteAxillaryDirection = SitePhaseQuat.RotateVector(InitialAxillaryDirection).GetSafeNormal();
 
-			FVector3f UpDirection, NormalDirection;
-			ComputeRotationAnglesForAttachment(
-				ApicalDirection,
-				BudDirectionUp,
-				FMath::DegreesToRadians(AxilAngleInDegrees),
-				Twist,
-				UpDirection,
-				NormalDirection);
+			const FQuat4f ApicalCorrection = FQuat4f::FindBetweenNormals(
+				InitialApicalDirection,
+				ApicalDirection.GetSafeNormal());
 
-			const bool bIsTipInstance =
-				FMath::IsNearlyEqual(AttachmentPointLengthFromRoot, LFRofTip, 0.001f) &&
-				(BranchGradient >= 0.999f || FMath::IsNearlyEqual(BranchGradient, SpacingSettings.RelativeEnd, 0.001f));
+			SiteAxillaryDirection = ApicalCorrection.RotateVector(SiteAxillaryDirection).GetSafeNormal();
 
-			if (bIsTipInstance)
+			for (int32 k = 0; k < NumBuds; ++k)
 			{
-				const int32 SecondToLastPointIndex = SortedPointIndices[SortedPointIndices.Num() - 2];
-				const FVector3f TipPosition = PointPositionAttribute[TipPointIndex];
-				const FVector3f SecondToLastPointPosition = PointPositionAttribute[SecondToLastPointIndex];
-				FVector3f TipTangent = (TipPosition - SecondToLastPointPosition).GetSafeNormal();
-				if (TipTangent.IsNearlyZero())
+				const float Scale = ComputeAttachmentScale(
+					GlobalBudIndex,
+					RandomSeed,
+					PlantGradient,
+					BranchGradient,
+					LengthRatio,
+					ScaleSettings);
+
+				FVector3f UpDirection;
+				FVector3f NormalDirection;
+
+				ComputeRotationAnglesForAttachment(
+					ApicalDirection,
+					BudDirectionUp,
+					SiteAxillaryDirection,
+					PlantGradient,
+					BranchGradient,
+					j,
+					k,
+					NumBuds,
+					GlobalBudIndex,
+					RandomSeed,
+					AngleSettings,
+					UpDirection,
+					NormalDirection);
+
+				if (bIsTipInstance)
 				{
-					TipTangent = ApicalDirection.GetSafeNormal();
+					const int32 SecondToLastPointIndex = SortedPointIndices[SortedPointIndices.Num() - 2];
+					const FVector3f TipPosition = PointPositionAttribute[TipPointIndex];
+					const FVector3f SecondToLastPointPosition = PointPositionAttribute[SecondToLastPointIndex];
+					FVector3f TipTangent = (TipPosition - SecondToLastPointPosition).GetSafeNormal();
+
+					if (TipTangent.IsNearlyZero())
+					{
+						TipTangent = ApicalDirection.GetSafeNormal();
+					}
+
+					const FVector3f OldUp = UpDirection.GetSafeNormal();
+					FQuat4f AlignQuat = FQuat4f::FindBetweenNormals(OldUp, TipTangent);
+
+					UpDirection = TipTangent;
+					NormalDirection = AlignQuat.RotateVector(NormalDirection).GetSafeNormal();
+
+					if (FMath::Abs(FVector3f::DotProduct(UpDirection, NormalDirection)) > 0.999f)
+					{
+						FVector Axis1, Axis2;
+						FVector(UpDirection).FindBestAxisVectors(Axis1, Axis2);
+						NormalDirection = FVector3f(Axis1);
+					}
 				}
 
-				const FVector3f OldUp = UpDirection.GetSafeNormal();
-				FQuat4f AlignQuat = FQuat4f::FindBetweenNormals(OldUp, TipTangent);
-
-				UpDirection = TipTangent;
-				NormalDirection = AlignQuat.RotateVector(NormalDirection).GetSafeNormal();
+				ConditionPicker::FPointConditionSample Sample;
+				Sample.UpAlignment = UpAlignmentNormalized(ApicalDirection);
+				Sample.Tip = Tip(bIsTipInstance);
 
-				if (FMath::Abs(FVector3f::DotProduct(UpDirection, NormalDirection)) > 0.999f)
+				if (PointIndexA != INDEX_NONE && PointIndexB != INDEX_NONE)
 				{
-					FVector Axis1, Axis2;
-					FVector(UpDirection).FindBestAxisVectors(Axis1, Axis2);
-					NormalDirection = FVector3f(Axis1);
+					Sample.Light = FMath::Lerp(NormalizedPointCaches.Light[PointIndexA], NormalizedPointCaches.Light[PointIndexB], PointAlpha);
+					Sample.Health = FMath::Lerp(NormalizedPointCaches.Health[PointIndexA], NormalizedPointCaches.Health[PointIndexB], PointAlpha);
+					Sample.Scale = FMath::Lerp(NormalizedPointCaches.PScales[PointIndexA], NormalizedPointCaches.PScales[PointIndexB], PointAlpha);
+				}
+				else
+				{
+					Sample.Light = 0.0f;
+					Sample.Health = 0.0f;
+					Sample.Scale = 0.0f;
 				}
-			}
-
-			ConditionPicker::FPointConditionSample Sample;
-			Sample.UpAlignment = UpAlignmentNormalized(ApicalDirection);
-			Sample.Tip = Tip(bIsTipInstance);
 
-			if (PointIndexA != INDEX_NONE && PointIndexB != INDEX_NONE)
-			{
-				Sample.Light = FMath::Lerp(NormalizedPointCaches.Light[PointIndexA], NormalizedPointCaches.Light[PointIndexB], PointAlpha);
-				Sample.Health = FMath::Lerp(NormalizedPointCaches.Health[PointIndexA], NormalizedPointCaches.Health[PointIndexB], PointAlpha);
-				Sample.Scale = FMath::Lerp(NormalizedPointCaches.PScales[PointIndexA], NormalizedPointCaches.PScales[PointIndexB], PointAlpha);
+				OutAttachmentPoints.Add({
+					.BranchIndex = BranchIndex,
+					.PointIndexA = PointIndexA,
+					.PointIndexB = PointIndexB,
+					.PointAlpha = PointAlpha,
+					.Position = Pos,
+					.UpDirection = UpDirection,
+					.NormalDirection = NormalDirection,
+					.Scale = Scale,
+					.LengthFromRoot = AttachmentPointLengthFromRoot,
+					.bIsTipInstance = bIsTipInstance,
+					.TipPointIndex = TipPointIndex,
+					.ConditionSample = Sample
+				});
+
+				++GlobalBudIndex;
 			}
-			else
+
+			if (j > 0)
 			{
-				Sample.Light = 0.0f;
-				Sample.Health = 0.0f;
-				Sample.Scale = 0.0f;
+				++AxillaryRotationIteration;
 			}
-
-			OutAttachmentPoints.Add({
-				.BranchIndex = BranchIndex,
-				.PointIndexA = PointIndexA,
-				.PointIndexB = PointIndexB,
-				.PointAlpha = PointAlpha,
-				.Position = Pos,
-				.UpDirection = UpDirection,
-				.NormalDirection = NormalDirection,
-				.Scale = Scale,
-				.LengthFromRoot = AttachmentPointLengthFromRoot,
-				.bIsTipInstance = bIsTipInstance,
-				.TipPointIndex = TipPointIndex,
-				.ConditionSample = Sample
-			});
 		}
 	}
 }
@@ -471,9 +622,20 @@ bool PV::ParametricDistributionHelper::FindPointByNormalizedLengthFromRoot(float
 	return true;
 }
 
-void PV::ParametricDistributionHelper::ComputeRotationAnglesForAttachment(const FVector3f& ApicalDirection, const FVector3f& BudDirectionUp,
-                                                                          const float AxilAngle, const float TwistAroundApical,
-                                                                          FVector3f& OutUpDirection, FVector3f& OutNormalDirection)
+void PV::ParametricDistributionHelper::ComputeRotationAnglesForAttachment(
+	const FVector3f& ApicalDirection,
+	const FVector3f& BudDirectionUp,
+	const FVector3f& BaseAxillaryDirection,
+	const float PlantGradient,
+	const float CurveU,
+	const int32 Iteration,
+	const int32 SubIteration,
+	const int32 NumSubIterations,
+	const int32 BudIndex,
+	const int32 RandomSeed,
+	const FParametricSettings::FAngleSettings& AngleSettings,
+	FVector3f& OutUpDirection,
+	FVector3f& OutNormalDirection)
 {
 	const FVector3f ApicalDirectionNormalized = ApicalDirection.GetSafeNormal();
 	if (ApicalDirectionNormalized.IsNearlyZero())
@@ -483,7 +645,12 @@ void PV::ParametricDistributionHelper::ComputeRotationAnglesForAttachment(const
 		return;
 	}
 
-	FVector3f BudDirectionNormalized = BudDirectionUp.GetSafeNormal();
+	FVector3f BudDirectionNormalized = BaseAxillaryDirection.GetSafeNormal();
+	if (BudDirectionNormalized.IsNearlyZero())
+	{
+		BudDirectionNormalized = BudDirectionUp.GetSafeNormal();
+	}
+
 	if (BudDirectionNormalized.IsNearlyZero())
 	{
 		FVector Axis1, Axis2;
@@ -499,10 +666,32 @@ void PV::ParametricDistributionHelper::ComputeRotationAnglesForAttachment(const
 		AxilAxis = FVector3f(Axis2);
 	}
 
-	const FQuat4f AxilQuat(AxilAxis, AxilAngle);
+	float RotationDegrees = AngleSettings.Rotation * static_cast<float>(Iteration);
+	if (NumSubIterations > 0)
+	{
+		RotationDegrees += (360.0f / static_cast<float>(NumSubIterations)) * static_cast<float>(SubIteration + 1);
+	}
+
+	float AxilRampLookup = CurveU;
+	if (AngleSettings.AxilAngleRampBasis == FParametricSettings::EDistributionBasis::Plant)
+	{
+		AxilRampLookup = PlantGradient;
+	}
+
+	const float AxilRampValue = AngleSettings.AxilAngleRamp->GetRichCurveConst()->Eval(AxilRampLookup);
+
+	float AxilAngleInDegrees = AngleSettings.AxilAngle * AxilRampValue;
+	AxilAngleInDegrees += IndexedRandomRange(
+		BudIndex + RandomSeed + 12,
+		AngleSettings.RandomizeAxilAngleMinimum,
+		AngleSettings.RandomizeAxilAngleMaximum);
+
+	AxilAngleInDegrees = FMath::Clamp(AxilAngleInDegrees, -90.0f, 90.0f);
+
+	const FQuat4f AxilQuat(AxilAxis, FMath::DegreesToRadians(AxilAngleInDegrees));
 	const FVector3f AxilAdjusted = AxilQuat.RotateVector(BudDirectionNormalized);
 
-	const FQuat4f TwistQuat(ApicalDirectionNormalized, TwistAroundApical);
+	const FQuat4f TwistQuat(ApicalDirectionNormalized, FMath::DegreesToRadians(RotationDegrees));
 
 	OutUpDirection = TwistQuat.RotateVector(AxilAdjusted).GetSafeNormal();
 	OutNormalDirection = TwistQuat.RotateVector(AxilQuat.RotateVector(ApicalDirectionNormalized)).GetSafeNormal();
diff --git a/Engine/Plugins/Experimental/ProceduralVegetationEditor/Source/ProceduralVegetation/Private/Helpers/PVParametricDistributionHelper.h b/Engine/Plugins/Experimental/ProceduralVegetationEditor/Source/ProceduralVegetation/Private/Helpers/PVParametricDistributionHelper.h
--- a/Engine/Plugins/Experimental/ProceduralVegetationEditor/Source/ProceduralVegetation/Private/Helpers/PVParametricDistributionHelper.h
+++ b/Engine/Plugins/Experimental/ProceduralVegetationEditor/Source/ProceduralVegetation/Private/Helpers/PVParametricDistributionHelper.h
@@ -4,6 +4,7 @@
 
 #include "CoreMinimal.h"
 #include "PVFoliageHelper.h"
+#include "PVPhyllotaxyHelper.h"
 #include "Utils/PVAttributes.h"
 #include "Utils/PVFloatRamp.h"
 
@@ -32,6 +33,17 @@ namespace PV::ParametricDistributionHelper
 			const FPVFloatRamp* SpacingRamp;
 		};
 
+		struct FPhyllotaxySettings
+		{
+			const EPhyllotaxyType PhyllotaxyType;
+			const EPhyllotaxyFormation PhyllotaxyFormation;
+			const int32 MinimumNodeBuds;
+			const int32 MaximumNodeBuds;
+			const float PhyllotaxyAdditionalAngle;
+			const bool ResetPhyllotaxy;
+			const float PhyllotaxyOffset;
+		};
+
 		struct FAngleSettings
 		{
 			const float Rotation;
@@ -50,15 +62,30 @@ namespace PV::ParametricDistributionHelper
 			const float RandomizeScaleAngleMaximum;
 			const FPVFloatRamp* ScaleRamp;
 		};
-	};
+	}
+
+	float IndexedRandom01(const int32 Seed);
+
+	float IndexedRandomRange(const int32 Seed, const float MinValue, const float MaxValue);
+
+	int32 IndexedRandomIntInclusive(const int32 Seed, const int32 MinValue, const int32 MaxValue);
+
+	static float ComputeAttachmentScale(
+		const int32 BudIndex,
+		const int32 RandomSeed,
+		const float PlantGradient,
+		const float BranchGradient,
+		const float LengthRatio,
+		const FParametricSettings::FScaleSettings& ScaleSettings);
 
 	void ComputeAttachmentPoints(
 		FManagedArrayCollection& OutCollection,
 		const FNormalizedPointCaches& NormalizedPointCaches,
 		const FParametricSettings::FSpacingSettings& SpacingSettings,
+		const FParametricSettings::FPhyllotaxySettings& PhyllotaxySettings,
 		const FParametricSettings::FAngleSettings& AngleSettings,
 		const FParametricSettings::FScaleSettings& ScaleSettings,
-		FRandomStream& RandomStream,
+		const int32 RandomSeed,
 		TArray<DistributionHelper::FAttachmentPoint>& OutAttachmentPoints);
 
 	TMap<int32, float> ComputeBranchLengths(
@@ -90,8 +117,15 @@ namespace PV::ParametricDistributionHelper
 	void ComputeRotationAnglesForAttachment(
 		const FVector3f& ApicalDirection,
 		const FVector3f& BudDirectionUp,
-		const float AxilAngle,
-		const float TwistAroundApical,
+		const FVector3f& BaseAxillaryDirection,
+		const float PlantGradient,
+		const float CurveU,
+		const int32 Iteration,
+		const int32 SubIteration,
+		const int32 NumSubIterations,
+		const int32 BudIndex,
+		const int32 RandomSeed,
+		const FParametricSettings::FAngleSettings& AngleSettings,
 		FVector3f& OutUpDirection,
 		FVector3f& OutNormalDirection);
 }
diff --git a/Engine/Plugins/Experimental/ProceduralVegetationEditor/Source/ProceduralVegetation/Private/Helpers/PVPhyllotaxyHelper.h b/Engine/Plugins/Experimental/ProceduralVegetationEditor/Source/ProceduralVegetation/Private/Helpers/PVPhyllotaxyHelper.h
--- a/Engine/Plugins/Experimental/ProceduralVegetationEditor/Source/ProceduralVegetation/Private/Helpers/PVPhyllotaxyHelper.h
+++ b/Engine/Plugins/Experimental/ProceduralVegetationEditor/Source/ProceduralVegetation/Private/Helpers/PVPhyllotaxyHelper.h
@@ -13,7 +13,7 @@ enum class EPhyllotaxyType : uint8
 	Whorled UMETA(DisplayName = "Whorled"),
 	Spiral UMETA(DisplayName = "Spiral")
 };
- 
+
 UENUM(BlueprintType)
 enum class EPhyllotaxyFormation : uint8
 {
@@ -22,4 +22,77 @@ enum class EPhyllotaxyFormation : uint8
 	Pentastichous UMETA(DisplayName = "Pentastichous"),
 	Octastichous UMETA(DisplayName = "Octastichous"),
 	Parastichous UMETA(DisplayName = "Parastichous")
-};
\ No newline at end of file
+};
+
+namespace PV::PhyllotaxySettingsHelper
+{
+	struct FResolvedPhyllotaxy
+	{
+		float FormationDegrees = 180.0f;
+		int32 MinBuds = 1;
+		int32 MaxBuds = 1;
+		bool bResetPhyllotaxy = false;
+		float OffsetDegrees = 0.0f;
+	};
+
+	template <typename TPhyllotaxySettings>
+	FResolvedPhyllotaxy ResolvePhyllotaxy(const TPhyllotaxySettings& Settings)
+	{
+		FResolvedPhyllotaxy Result;
+		Result.bResetPhyllotaxy = Settings.ResetPhyllotaxy;
+		Result.OffsetDegrees = Settings.PhyllotaxyOffset;
+
+		switch (Settings.PhyllotaxyType)
+		{
+		case EPhyllotaxyType::Alternate:
+			Result.FormationDegrees = 180.0f + Settings.PhyllotaxyAdditionalAngle;
+			Result.MinBuds = 1;
+			Result.MaxBuds = 1;
+			break;
+
+		case EPhyllotaxyType::Opposite:
+			Result.FormationDegrees = Settings.PhyllotaxyAdditionalAngle;
+			Result.MinBuds = 2;
+			Result.MaxBuds = 2;
+			break;
+
+		case EPhyllotaxyType::Decussate:
+			Result.FormationDegrees = 90.0f + Settings.PhyllotaxyAdditionalAngle;
+			Result.MinBuds = 2;
+			Result.MaxBuds = 2;
+			break;
+
+		case EPhyllotaxyType::Whorled:
+			Result.FormationDegrees = 90.0f + Settings.PhyllotaxyAdditionalAngle;
+			Result.MinBuds = Settings.MinimumNodeBuds;
+			Result.MaxBuds = Settings.MaximumNodeBuds;
+			break;
+
+		case EPhyllotaxyType::Spiral:
+			switch (Settings.PhyllotaxyFormation)
+			{
+			case EPhyllotaxyFormation::Distichous:
+				Result.FormationDegrees = 180.0f + Settings.PhyllotaxyAdditionalAngle;
+				break;
+			case EPhyllotaxyFormation::Tristichous:
+				Result.FormationDegrees = 120.0f + Settings.PhyllotaxyAdditionalAngle;
+				break;
+			case EPhyllotaxyFormation::Pentastichous:
+				Result.FormationDegrees = 144.0f + Settings.PhyllotaxyAdditionalAngle;
+				break;
+			case EPhyllotaxyFormation::Octastichous:
+				Result.FormationDegrees = 135.0f + Settings.PhyllotaxyAdditionalAngle;
+				break;
+			case EPhyllotaxyFormation::Parastichous:
+				Result.FormationDegrees = Settings.PhyllotaxyAdditionalAngle;
+				break;
+			}
+
+			Result.MinBuds = 1;
+			Result.MaxBuds = 1;
+			break;
+		}
+
+		return Result;
+	}
+}
diff --git a/Engine/Plugins/Experimental/ProceduralVegetationEditor/Source/ProceduralVegetation/Private/Implementations/PVFoliage.cpp b/Engine/Plugins/Experimental/ProceduralVegetationEditor/Source/ProceduralVegetation/Private/Implementations/PVFoliage.cpp
--- a/Engine/Plugins/Experimental/ProceduralVegetationEditor/Source/ProceduralVegetation/Private/Implementations/PVFoliage.cpp
+++ b/Engine/Plugins/Experimental/ProceduralVegetationEditor/Source/ProceduralVegetation/Private/Implementations/PVFoliage.cpp
@@ -152,6 +152,7 @@ void FPVFoliage::DistributeFoliageWithParametricSettings(
 	FManagedArrayCollection& OutCollection,
 	const FManagedArrayCollection& FoliageCollection,
 	const FParametricSettings::FSpacingSettings& SpacingSettings,
+	const FParametricSettings::FPhyllotaxySettings& PhyllotaxySettings,
 	const FParametricSettings::FAngleSettings& AngleSettings,
 	const FParametricSettings::FScaleSettings& ScaleSettings,
 	const int32 RandomSeed,
@@ -194,9 +195,10 @@ void FPVFoliage::DistributeFoliageWithParametricSettings(
 		OutCollection,
 		NormalizedPointCaches,
 		SpacingSettings,
+		PhyllotaxySettings,
 		AngleSettings,
 		ScaleSettings,
-		RandomStream,
+		RandomSeed,
 		AttachmentPoints);
 
 	OutCollection.EmptyGroup(PV::GroupNames::FoliageGroup);
diff --git a/Engine/Plugins/Experimental/ProceduralVegetationEditor/Source/ProceduralVegetation/Private/Implementations/PVFoliage.h b/Engine/Plugins/Experimental/ProceduralVegetationEditor/Source/ProceduralVegetation/Private/Implementations/PVFoliage.h
--- a/Engine/Plugins/Experimental/ProceduralVegetationEditor/Source/ProceduralVegetation/Private/Implementations/PVFoliage.h
+++ b/Engine/Plugins/Experimental/ProceduralVegetationEditor/Source/ProceduralVegetation/Private/Implementations/PVFoliage.h
@@ -28,6 +28,7 @@ struct FPVFoliage
 		FManagedArrayCollection& OutCollection,
 		const FManagedArrayCollection& FoliageCollection,
 		const FParametricSettings::FSpacingSettings& SpacingSettings,
+		const FParametricSettings::FPhyllotaxySettings& PhyllotaxySettings,
 		const FParametricSettings::FAngleSettings& AngleSettings,
 		const FParametricSettings::FScaleSettings& ScaleSettings,
 		const int32 RandomSeed,
diff --git a/Engine/Plugins/Experimental/ProceduralVegetationEditor/Source/ProceduralVegetation/Private/Implementations/PVGrafter.cpp b/Engine/Plugins/Experimental/ProceduralVegetationEditor/Source/ProceduralVegetation/Private/Implementations/PVGrafter.cpp
--- a/Engine/Plugins/Experimental/ProceduralVegetationEditor/Source/ProceduralVegetation/Private/Implementations/PVGrafter.cpp
+++ b/Engine/Plugins/Experimental/ProceduralVegetationEditor/Source/ProceduralVegetation/Private/Implementations/PVGrafter.cpp
@@ -15,7 +15,6 @@
 #include "Helpers/PVPlantTraversalHelper.h"
 #include "Helpers/PVFoliageJSONHelper.h"
 #include "Helpers/PVUtilities.h"
-#include "Helpers/PVAttributesHelper.h"
 
 using namespace PV::DistributionConditionUtils;
 using namespace PV::ParametricDistributionHelper;
@@ -225,6 +224,10 @@ void FPVGrafter::SetGraft(FManagedArrayCollection& OutCollection,
 
 	PointFacadeOut.SetPositionsFromIndex(TransformedPoints, PointFacadeOutStartIndex);
 
+	if (InAttributes.bIsTipInstance)
+	{
+		RescaleBranchForGraftingOnTip(OutCollection, InAttributes.BranchIndex, InAttributes.PointIndex, GraftTrunkLength);
+	}
 
 	const TManagedArray<float>& SourceGraftPointScales = PointFacadeSourceGraft.GetPointScales();
 	TArray<float> TransformedPointScales;
@@ -390,8 +393,6 @@ void FPVGrafter::SetGraft(FManagedArrayCollection& OutCollection,
 				TArray<int32> TargetBranchPoints = BranchFacadeOut.GetPoints(InAttributes.BranchIndex);
 				TargetBranchPoints.Append(TransformedBranchPoints);
 				BranchFacadeOut.SetPoints(InAttributes.BranchIndex, TargetBranchPoints);
-
-				RescaleBranchForGraftingOnTip(OutCollection, InAttributes.BranchIndex, InAttributes.PointIndex, GraftTrunkLength);
 			}
 			else
 			{
@@ -457,7 +458,7 @@ void FPVGrafter::SetGraft(FManagedArrayCollection& OutCollection,
 		{
 			TransformedParentBranchNumbers[TargetIndex] = ParentBranchNumbers[SourceIndex] == 0
 				? BranchNumberBeingAttachedToInTarget
-				: SourceOldBranchNumbersToNewBranchNumbers[BranchNumbers[SourceIndex]];
+				: SourceOldBranchNumbersToNewBranchNumbers[ParentBranchNumbers[SourceIndex]];
 			TargetIndex++;
 		}
 
@@ -486,32 +487,30 @@ void FPVGrafter::SetGraft(FManagedArrayCollection& OutCollection,
 
 	BranchFacadeOut.SetBranchHierarchyNumbersFromIndex(TransformedBranchHierarchyNumbers, BranchFacadeOutStartIndex);
 
-	// Commented out for now since Grower node is currently generating inconsistent source bud numbers 
-	//
-	// const TManagedArray<int32>& BranchSourceBudNumbers = BranchFacadeSourceGraft.GetBranchSourceBudNumbers();
-	// TArray<int32> TransformedBranchSourceBudNumbers;
-	// TransformedBranchSourceBudNumbers.SetNum(NumOfBranchesToAdd);
-	//
-	// SourceIndex = 0;
-	// TargetIndex = 0;
-	// while (SourceIndex < NumOfBranchesInSource)
-	// {
-	// 	if (!(InAttributes.bIsTipInstance && SourceIndex == SourceTrunkBranchIndex))
-	// 	[[likely]]
-	// 	{
-	// 		const int32 ParentPointSourceBudNumber = PointFacadeOut.GetBudNumber(InAttributes.PointIndex);
-	// 		const int32 SourceBudNumber = BranchSourceBudNumbers[SourceIndex];
-	// 		const int32 NewSourceBudNumber = SourceBudNumber == 1
-	// 			? ParentPointSourceBudNumber
-	// 			: SourceGraftOldToNewBudNumbers[SourceBudNumber];
-	// 		TransformedBranchSourceBudNumbers[TargetIndex] = NewSourceBudNumber;
-	// 		TargetIndex++;
-	// 	}
-	//
-	// 	SourceIndex++;
-	// }
-	//
-	// BranchFacadeOut.SetBranchSourceBudNumbersFromIndex(TransformedBranchSourceBudNumbers, BranchFacadeOutStartIndex);
+	const TManagedArray<int32>& BranchSourceBudNumbers = BranchFacadeSourceGraft.GetBranchSourceBudNumbers();
+	TArray<int32> TransformedBranchSourceBudNumbers;
+	TransformedBranchSourceBudNumbers.SetNum(NumOfBranchesToAdd);
+
+	SourceIndex = 0;
+	TargetIndex = 0;
+	while (SourceIndex < NumOfBranchesInSource)
+	{
+		if (!(InAttributes.bIsTipInstance && SourceIndex == SourceGraftTrunkBranchIndex))
+		[[likely]]
+		{
+			const int32 ParentPointSourceBudNumber = PointFacadeOut.GetBudNumber(InAttributes.PointIndex);
+			const int32 SourceBudNumber = BranchSourceBudNumbers[SourceIndex];
+			const int32 NewSourceBudNumber = SourceBudNumber == 1
+				? ParentPointSourceBudNumber
+				: SourceGraftOldToNewBudNumbers[SourceBudNumber];
+			TransformedBranchSourceBudNumbers[TargetIndex] = NewSourceBudNumber;
+			TargetIndex++;
+		}
+
+		SourceIndex++;
+	}
+
+	BranchFacadeOut.SetBranchSourceBudNumbersFromIndex(TransformedBranchSourceBudNumbers, BranchFacadeOutStartIndex);
 
 
 	SourceIndex = 0;
@@ -658,7 +657,7 @@ void FPVGrafter::DistributeGraftWithHormoneBasedSettings(
 	const FPVFoliageDistributionConditions& DistributionConditions)
 {
 	TRACE_CPUPROFILER_EVENT_SCOPE(PVGrafter::DistributeGraftWithHormoneBasedSettings);
-	
+
 	if (!(DistributionSettings.InstanceSpacingRamp && ScaleSettings.ScaleRamp && VectorSettings.AxilAngleRamp))
 	{
 		return;
@@ -804,6 +803,7 @@ void FPVGrafter::DistributeGraftWithHormoneBasedSettings(
 void FPVGrafter::DistributeGraftWithParametricSettings(FManagedArrayCollection& OutCollection,
                                                        TArray<FManagedArrayCollection>& SourceGraftCollections,
                                                        const FParametricSettings::FSpacingSettings& SpacingSettings,
+                                                       const FParametricSettings::FPhyllotaxySettings& PhyllotaxySettings,
                                                        const FParametricSettings::FAngleSettings& AngleSettings,
                                                        const FParametricSettings::FScaleSettings& ScaleSettings,
                                                        const int32 RandomSeed,
@@ -845,12 +845,13 @@ void FPVGrafter::DistributeGraftWithParametricSettings(FManagedArrayCollection&
 
 	TArray<PV::DistributionHelper::FAttachmentPoint> AttachmentPoints;
 	ComputeAttachmentPoints(
-		OutCollection, 
-		NormalizedPointCaches, 
-		SpacingSettings, 
-		AngleSettings, 
-		ScaleSettings, 
-		RandomStream, 
+		OutCollection,
+		NormalizedPointCaches,
+		SpacingSettings,
+		PhyllotaxySettings,
+		AngleSettings,
+		ScaleSettings,
+		RandomSeed,
 		AttachmentPoints);
 
 	auto BudNumberAttribute = PV::FPointBudNumberAttribute::GetAttribute(OutCollection);
diff --git a/Engine/Plugins/Experimental/ProceduralVegetationEditor/Source/ProceduralVegetation/Private/Implementations/PVGrafter.h b/Engine/Plugins/Experimental/ProceduralVegetationEditor/Source/ProceduralVegetation/Private/Implementations/PVGrafter.h
--- a/Engine/Plugins/Experimental/ProceduralVegetationEditor/Source/ProceduralVegetation/Private/Implementations/PVGrafter.h
+++ b/Engine/Plugins/Experimental/ProceduralVegetationEditor/Source/ProceduralVegetation/Private/Implementations/PVGrafter.h
@@ -65,6 +65,7 @@ struct FPVGrafter
 		FManagedArrayCollection& OutCollection,
 		TArray<FManagedArrayCollection>& SourceGraftCollections,
 		const FParametricSettings::FSpacingSettings& SpacingSettings,
+		const FParametricSettings::FPhyllotaxySettings& PhyllotaxySettings,
 		const FParametricSettings::FAngleSettings& AngleSettings,
 		const FParametricSettings::FScaleSettings& ScaleSettings,
 		int32 RandomSeed,
diff --git a/Engine/Plugins/Experimental/ProceduralVegetationEditor/Source/ProceduralVegetation/Private/Nodes/PVFoliageDistributorSettings.cpp b/Engine/Plugins/Experimental/ProceduralVegetationEditor/Source/ProceduralVegetation/Private/Nodes/PVFoliageDistributorSettings.cpp
--- a/Engine/Plugins/Experimental/ProceduralVegetationEditor/Source/ProceduralVegetation/Private/Nodes/PVFoliageDistributorSettings.cpp
+++ b/Engine/Plugins/Experimental/ProceduralVegetationEditor/Source/ProceduralVegetation/Private/Nodes/PVFoliageDistributorSettings.cpp
@@ -198,6 +198,16 @@ bool FPVFoliageDistributorElement::ExecuteInternal(FPCGContext* InContext) const
 					: FParametricSettings::EDistributionBasis::Branch,
 					.SpacingRamp = &Settings->ParametricSettings.SpacingRamp
 				};
+				
+				const FParametricSettings::FPhyllotaxySettings PhyllotaxySettings = {
+					.PhyllotaxyType = Settings->ParametricSettings.PhyllotaxyType,
+					.PhyllotaxyFormation = Settings->ParametricSettings.PhyllotaxyFormation,
+					.MinimumNodeBuds = Settings->ParametricSettings.MinimumNodeBuds,
+					.MaximumNodeBuds = Settings->ParametricSettings.MaximumNodeBuds,
+					.PhyllotaxyAdditionalAngle = Settings->ParametricSettings.PhyllotaxyAdditionalAngle,
+					.ResetPhyllotaxy = Settings->ParametricSettings.ResetPhyllotaxy,
+					.PhyllotaxyOffset = Settings->ParametricSettings.PhyllotaxyOffset
+				};
 
 				const FParametricSettings::FAngleSettings AngleSettings = {
 					.Rotation = Settings->ParametricSettings.Rotation,
@@ -224,6 +234,7 @@ bool FPVFoliageDistributorElement::ExecuteInternal(FPCGContext* InContext) const
 					Collection,
 					FoliageCollection,
 					SpacingSettings,
+					PhyllotaxySettings,
 					AngleSettings,
 					ScaleSettings,
 					Settings->RandomSeed,
diff --git a/Engine/Plugins/Experimental/ProceduralVegetationEditor/Source/ProceduralVegetation/Private/Nodes/PVFoliageDistributorSettings.h b/Engine/Plugins/Experimental/ProceduralVegetationEditor/Source/ProceduralVegetation/Private/Nodes/PVFoliageDistributorSettings.h
--- a/Engine/Plugins/Experimental/ProceduralVegetationEditor/Source/ProceduralVegetation/Private/Nodes/PVFoliageDistributorSettings.h
+++ b/Engine/Plugins/Experimental/ProceduralVegetationEditor/Source/ProceduralVegetation/Private/Nodes/PVFoliageDistributorSettings.h
@@ -138,6 +138,44 @@ struct FPVFoliageDistributionParametricSettings
 	UPROPERTY(EditAnywhere, Category="Spacing Settings",
 		meta=(PCG_Overridable, XAxisMin=0.0f, XAxisMax=1.0f, YAxisMin=0.0f, YAxisMax=1.0f))
 	FPVFloatRamp SpacingRamp;
+	
+	UPROPERTY(EditAnywhere, Category="Phyllotaxy Settings",
+		meta=(PCG_Overridable, Tooltip=
+			"Reset phyllotaxy per branch.\n\nBranch orientation is used to align the initial phyllotaxy. This setting will reset it per branch so that its no longer tied to parent orientation."
+		))
+	bool ResetPhyllotaxy = false;
+
+	UPROPERTY(EditAnywhere, Category="Phyllotaxy Settings",
+		meta=(PCG_Overridable, Tooltip=
+			"Pattern for arranging buds/leaves around the stem.\n\nSelects the phyllotactic pattern used for placement (e.g., Alternate, Opposite, Whorled, Spiral). This defines the sequence and angular spacing of buds/leaves around the stem and along its length."
+		))
+	EPhyllotaxyType PhyllotaxyType = EPhyllotaxyType::Alternate;
+
+	UPROPERTY(EditAnywhere, Category="Phyllotaxy Settings",
+		meta=(PCG_Overridable, EditCondition="PhyllotaxyType == EPhyllotaxyType::Spiral", EditConditionHides))
+	EPhyllotaxyFormation PhyllotaxyFormation = EPhyllotaxyFormation::Distichous;
+
+	UPROPERTY(EditAnywhere, Category="Phyllotaxy Settings",
+		meta=(PCG_Overridable, ClampMin=1, ClampMax=10, UIMin=1, UIMax=10, Tooltip=
+			"Minimum buds/leaves per node.\n\nSets the lower bound on the number of buds/leaves generated at each node. Guarantees at least this many instances (subject to spacing and other constraints). Should be ≤ MaximumNodeBuds. Applied when Whorled phyllotaxy is selected."
+		))
+	int32 MinimumNodeBuds = 2;
+
+	UPROPERTY(EditAnywhere, Category="Phyllotaxy Settings",
+		meta=(PCG_Overridable, ClampMin=1, ClampMax=10, UIMin=1, UIMax=10, Tooltip=
+			"Maximum buds/leaves per node.\n\nCaps the number of buds/leaves spawned. Use with MinimumNodeBuds to bound density and keep arrangements predictable. Applied when Whorled phyllotaxy is selected."
+		))
+	int32 MaximumNodeBuds = 3;
+
+	UPROPERTY(EditAnywhere, Category="Phyllotaxy Settings",
+		meta=(PCG_Overridable, ClampMin=0.0f, ClampMax=360.0f, UIMin=0.0f, UIMax=360.0f, Tooltip=
+			"Extra divergence angle (°) added to the phyllotaxy pattern\n\nAdds an angular offset to the base phyllotactic divergence."))
+	float PhyllotaxyAdditionalAngle = 0.0f;
+	
+	UPROPERTY(EditAnywhere, Category="Phyllotaxy Settings",
+		meta=(PCG_Overridable, ClampMin=0.0f, ClampMax=360.0f, UIMin=0.0f, UIMax=360.0f, Tooltip=
+			"Initial offset value for the phyllotaxy (rotation around parent).\n\nValues are in degrees."))
+	float PhyllotaxyOffset = 0.0f;
 
 	UPROPERTY(EditAnywhere, Category="Angle Settings", meta=(PCG_Overridable, ClampMin=0.0f, ClampMax=360.0f))
 	float Rotation = 0.0f;
diff --git a/Engine/Plugins/Experimental/ProceduralVegetationEditor/Source/ProceduralVegetation/Private/Nodes/PVGrafterSettings.cpp b/Engine/Plugins/Experimental/ProceduralVegetationEditor/Source/ProceduralVegetation/Private/Nodes/PVGrafterSettings.cpp
--- a/Engine/Plugins/Experimental/ProceduralVegetationEditor/Source/ProceduralVegetation/Private/Nodes/PVGrafterSettings.cpp
+++ b/Engine/Plugins/Experimental/ProceduralVegetationEditor/Source/ProceduralVegetation/Private/Nodes/PVGrafterSettings.cpp
@@ -235,6 +235,16 @@ bool FPVGrafterElement::ExecuteInternal(FPCGContext* InContext) const
 			: FParametricSettings::EDistributionBasis::Branch,
 			.SpacingRamp = &Settings->ParametricSettings.SpacingRamp
 		};
+		
+		const FParametricSettings::FPhyllotaxySettings PhyllotaxySettings = {
+			.PhyllotaxyType = Settings->ParametricSettings.PhyllotaxyType,
+			.PhyllotaxyFormation = Settings->ParametricSettings.PhyllotaxyFormation,
+			.MinimumNodeBuds = Settings->ParametricSettings.MinimumNodeBuds,
+			.MaximumNodeBuds = Settings->ParametricSettings.MaximumNodeBuds,
+			.PhyllotaxyAdditionalAngle = Settings->ParametricSettings.PhyllotaxyAdditionalAngle,
+			.ResetPhyllotaxy = Settings->ParametricSettings.ResetPhyllotaxy,
+			.PhyllotaxyOffset = Settings->ParametricSettings.PhyllotaxyOffset
+		}; 
 
 		const FParametricSettings::FAngleSettings AngleSettings = {
 			.Rotation = Settings->ParametricSettings.Rotation,
@@ -261,6 +271,7 @@ bool FPVGrafterElement::ExecuteInternal(FPCGContext* InContext) const
 			OutCollection,
 			SourceGraftCollections,
 			SpacingSettings,
+			PhyllotaxySettings,
 			AngleSettings,
 			ScaleSettings,
 			Settings->RandomSeed,
diff --git a/Engine/Plugins/Experimental/ProceduralVegetationEditor/Source/ProceduralVegetation/Private/Nodes/PVGrafterSettings.h b/Engine/Plugins/Experimental/ProceduralVegetationEditor/Source/ProceduralVegetation/Private/Nodes/PVGrafterSettings.h
--- a/Engine/Plugins/Experimental/ProceduralVegetationEditor/Source/ProceduralVegetation/Private/Nodes/PVGrafterSettings.h
+++ b/Engine/Plugins/Experimental/ProceduralVegetationEditor/Source/ProceduralVegetation/Private/Nodes/PVGrafterSettings.h
@@ -63,6 +63,44 @@ struct FPVGrafterDistributionParametricSettings
 	UPROPERTY(EditAnywhere, Category="Spacing Settings",
 		meta=(PCG_Overridable, XAxisMin=0.0f, XAxisMax=1.0f, YAxisMin=0.0f, YAxisMax=1.0f))
 	FPVFloatRamp SpacingRamp;
+	
+	UPROPERTY(EditAnywhere, Category="Phyllotaxy Settings",
+		meta=(PCG_Overridable, Tooltip=
+			"Reset phyllotaxy per branch.\n\nBranch orientation is used to align the initial phyllotaxy. This setting will reset it per branch so that its no longer tied to parent orientation."
+		))
+	bool ResetPhyllotaxy = false;
+
+	UPROPERTY(EditAnywhere, Category="Phyllotaxy Settings",
+		meta=(PCG_Overridable, Tooltip=
+			"Pattern for arranging buds/leaves around the stem.\n\nSelects the phyllotactic pattern used for placement (e.g., Alternate, Opposite, Whorled, Spiral). This defines the sequence and angular spacing of buds/leaves around the stem and along its length."
+		))
+	EPhyllotaxyType PhyllotaxyType = EPhyllotaxyType::Alternate;
+
+	UPROPERTY(EditAnywhere, Category="Phyllotaxy Settings",
+		meta=(PCG_Overridable, EditCondition="PhyllotaxyType == EPhyllotaxyType::Spiral", EditConditionHides))
+	EPhyllotaxyFormation PhyllotaxyFormation = EPhyllotaxyFormation::Distichous;
+
+	UPROPERTY(EditAnywhere, Category="Phyllotaxy Settings",
+		meta=(PCG_Overridable, ClampMin=1, ClampMax=10, UIMin=1, UIMax=10, Tooltip=
+			"Minimum buds/leaves per node.\n\nSets the lower bound on the number of buds/leaves generated at each node. Guarantees at least this many instances (subject to spacing and other constraints). Should be ≤ MaximumNodeBuds. Applied when Whorled phyllotaxy is selected."
+		))
+	int32 MinimumNodeBuds = 2;
+
+	UPROPERTY(EditAnywhere, Category="Phyllotaxy Settings",
+		meta=(PCG_Overridable, ClampMin=1, ClampMax=10, UIMin=1, UIMax=10, Tooltip=
+			"Maximum buds/leaves per node.\n\nCaps the number of buds/leaves spawned. Use with MinimumNodeBuds to bound density and keep arrangements predictable. Applied when Whorled phyllotaxy is selected."
+		))
+	int32 MaximumNodeBuds = 3;
+
+	UPROPERTY(EditAnywhere, Category="Phyllotaxy Settings",
+		meta=(PCG_Overridable, ClampMin=0.0f, ClampMax=360.0f, UIMin=0.0f, UIMax=360.0f, Tooltip=
+			"Extra divergence angle (°) added to the phyllotaxy pattern\n\nAdds an angular offset to the base phyllotactic divergence."))
+	float PhyllotaxyAdditionalAngle = 0.0f;
+	
+	UPROPERTY(EditAnywhere, Category="Phyllotaxy Settings",
+		meta=(PCG_Overridable, ClampMin=0.0f, ClampMax=360.0f, UIMin=0.0f, UIMax=360.0f, Tooltip=
+			"Initial offset value for the phyllotaxy (rotation around parent).\n\nValues are in degrees."))
+	float PhyllotaxyOffset = 0.0f;
 
 	UPROPERTY(EditAnywhere, Category="Angle Settings", meta=(PCG_Overridable, ClampMin=0.0f, ClampMax=360.0f))
 	float Rotation = 0.0f;
