[PVE] Fix issue where scale node is not working for fused points based growth data

#Synced-CL 51749465

diff --git a/Engine/Plugins/Experimental/ProceduralVegetationEditor/Source/ProceduralVegetation/Private/Implementations/PVScale.cpp b/Engine/Plugins/Experimental/ProceduralVegetationEditor/Source/ProceduralVegetation/Private/Implementations/PVScale.cpp
--- a/Engine/Plugins/Experimental/ProceduralVegetationEditor/Source/ProceduralVegetation/Private/Implementations/PVScale.cpp
+++ b/Engine/Plugins/Experimental/ProceduralVegetationEditor/Source/ProceduralVegetation/Private/Implementations/PVScale.cpp
@@ -42,11 +42,14 @@ void FPVScale::ApplyScale(const float InManualScale, FManagedArrayCollection& Ou
 		const int32 RootPointIndex = TrunkPoints[0];
 		const FVector3f& PivotPos = PointFacade.GetPosition(RootPointIndex);
 
+		TSet<int32> VisitedPointIndices;
 		for (const int32 BranchIndex : BranchIndices)
 		{
 			for (const TArray<int32>& BranchPoints = BranchFacade.GetPoints(BranchIndex);
 			     const int32& PointIndex : BranchPoints)
 			{
+				if (VisitedPointIndices.Contains(PointIndex)) continue;
+				
 				FVector3f PointPos = PointFacade.GetPosition(PointIndex);
 				float LFS = PointFacade.GetLengthFromSeed(PointIndex);
 				float LFR = PointFacade.GetLengthFromRoot(PointIndex);
@@ -62,6 +65,8 @@ void FPVScale::ApplyScale(const float InManualScale, FManagedArrayCollection& Ou
 				PointFacade.ModifyLengthFromSeeds()[PointIndex] = LFS;
 				PointFacade.ModifyLengthFromRoots()[PointIndex] = LFR;
 				PointFacade.ModifyPointScales()[PointIndex] *= InManualScale;
+				
+				VisitedPointIndices.Add(PointIndex);
 			}
 
 			const TArray<int32>& FoliageEntryIds = FoliageFacade.GetFoliageEntryIdsForBranch(BranchIndex);
