[PVE] Update scale node to skip scaling on duplicate bud points

Issue:
The scale node was applying position, length, and point scale to all points
on a branch unconditionally, including points that share the same bud number
as the source bud. This caused double-scaling on shared bud points when
multiple branches originate from the same bud.

Fix:
Added a guard condition so scaling is only applied when the branch is first
generation (trunk) or the point's bud number differs from the branch's source
bud number.

Also, Removed Foliage instancing adjustment because it is not longer coming from the growth data.

#Synced-CL 51705062
#Synced-CL 51705062

diff --git a/Engine/Plugins/Experimental/ProceduralVegetationEditor/Source/ProceduralVegetation/Private/Implementations/PVScale.cpp b/Engine/Plugins/Experimental/ProceduralVegetationEditor/Source/ProceduralVegetation/Private/Implementations/PVScale.cpp
--- a/Engine/Plugins/Experimental/ProceduralVegetationEditor/Source/ProceduralVegetation/Private/Implementations/PVScale.cpp
+++ b/Engine/Plugins/Experimental/ProceduralVegetationEditor/Source/ProceduralVegetation/Private/Implementations/PVScale.cpp
@@ -13,7 +13,6 @@ void FPVScale::ApplyScale(const float InManualScale, FManagedArrayCollection& Ou
 	PV::Facades::FBranchFacade BranchFacade(OutCollection);
 	PV::Facades::FPointFacade PointFacade(OutCollection);
 	PV::Facades::FPlantFacade PlantFacade(OutCollection);
-	PV::Facades::FFoliageFacade FoliageFacade(OutCollection);
 
 	if (!PointFacade.IsValid() || !BranchFacade.IsValid() || !PlantFacade.IsValid())
 	{
@@ -42,47 +41,34 @@ void FPVScale::ApplyScale(const float InManualScale, FManagedArrayCollection& Ou
 		const int32 RootPointIndex = TrunkPoints[0];
 		const FVector3f& PivotPos = PointFacade.GetPosition(RootPointIndex);
 
-		TSet<int32> VisitedPointIndices;
 		for (const int32 BranchIndex : BranchIndices)
 		{
+			const bool IsFirstGeneration = BranchFacade.IsTrunk(BranchIndex);
+			const int SourceBudNumber = BranchFacade.GetBranchSourceBudNumber(BranchIndex);
+			
 			for (const TArray<int32>& BranchPoints = BranchFacade.GetPoints(BranchIndex);
-			     const int32& PointIndex : BranchPoints)
+				 const int32& PointIndex : BranchPoints)
 			{
-				if (VisitedPointIndices.Contains(PointIndex)) continue;
+				const int BudNumber = PointFacade.GetBudNumber(PointIndex);
 				
-				FVector3f PointPos = PointFacade.GetPosition(PointIndex);
-				float LFS = PointFacade.GetLengthFromSeed(PointIndex);
-				float LFR = PointFacade.GetLengthFromRoot(PointIndex);
-
-				PointPos -= PivotPos;
-				PointPos *= InManualScale;
-				PointPos += PivotPos;
-
-				LFS *= InManualScale;
-				LFR *= InManualScale;
-				
-				PointFacade.ModifyPositions()[PointIndex] = PointPos;
-				PointFacade.ModifyLengthFromSeeds()[PointIndex] = LFS;
-				PointFacade.ModifyLengthFromRoots()[PointIndex] = LFR;
-				PointFacade.ModifyPointScales()[PointIndex] *= InManualScale;
+				if (IsFirstGeneration || BudNumber != SourceBudNumber)
+				{
+					FVector3f PointPos = PointFacade.GetPosition(PointIndex);
+					float LFS = PointFacade.GetLengthFromSeed(PointIndex);
+					float LFR = PointFacade.GetLengthFromRoot(PointIndex);
+
+					PointPos -= PivotPos;
+					PointPos *= InManualScale;
+					PointPos += PivotPos;
+
+					LFS *= InManualScale;
+					LFR *= InManualScale;
 				
-				VisitedPointIndices.Add(PointIndex);
-			}
-
-			const TArray<int32>& FoliageEntryIds = FoliageFacade.GetFoliageEntryIdsForBranch(BranchIndex);
-
-			for (const int32 FoliageIndex : FoliageEntryIds)
-			{
-				PV::Facades::FFoliageEntryData Data = FoliageFacade.GetFoliageEntry(FoliageIndex);
-
-				Data.PivotPoint -= PivotPos;
-				Data.PivotPoint *= InManualScale;
-				Data.PivotPoint += PivotPos;
-
-				Data.Scale *= InManualScale;
-				Data.LengthFromRoot = Data.LengthFromRoot * InManualScale;
-
-				FoliageFacade.SetFoliageEntry(FoliageIndex, Data);
+					PointFacade.ModifyPositions()[PointIndex] = PointPos;
+					PointFacade.ModifyLengthFromSeeds()[PointIndex] = LFS;
+					PointFacade.ModifyLengthFromRoots()[PointIndex] = LFR;
+					PointFacade.ModifyPointScales()[PointIndex] *= InManualScale;
+				}
 			}
 		}
 	}
