Fix poor selection of parent branch during PVE 2D/3D plant import

Replace the single dot-product branch continuation heuristic in the procedural vegetation import pipeline with a three-metric scoring system combining directional continuity, subtree length-to-tip, and radius similarity. Normalize each metric to [0,1] and sum with equal weight to produce more botanically correct branch hierarchies

#Synced-CL 51755814

diff --git a/Engine/Plugins/Experimental/ProceduralVegetationEditor/Source/ProceduralVegetation/Private/Helpers/PVImportHelpers.cpp b/Engine/Plugins/Experimental/ProceduralVegetationEditor/Source/ProceduralVegetation/Private/Helpers/PVImportHelpers.cpp
--- a/Engine/Plugins/Experimental/ProceduralVegetationEditor/Source/ProceduralVegetation/Private/Helpers/PVImportHelpers.cpp
+++ b/Engine/Plugins/Experimental/ProceduralVegetationEditor/Source/ProceduralVegetation/Private/Helpers/PVImportHelpers.cpp
@@ -175,6 +175,8 @@ namespace PV::ImportHelper
 		const FBooleanVertexAttribute& RootPointAttribute,
 		const FBooleanVertexAttribute& EndPointAttribute,
 		const FIntVertexAttribute& NextBranchPointAttribute,
+		const FFloatVertexAttribute& VertexScaleAttribute,
+		const FFloatVertexAttribute& LengthFromRootAttribute,
 		TArray<FMeshBranchHierarchy>& OutMeshBranchHierarchies
 	)
 	{
@@ -185,35 +187,116 @@ namespace PV::ImportHelper
 		TArray<bool> bProcessedVertices;
 		bProcessedVertices.SetNumZeroed(DynamicMesh.MaxVertexID());
 
-		const auto FindBestParentPoint = [&](int32 PrevVertexID, int32 VertexID, const TArray<int32>& NextPoints)
+		TArray<float> MaxSubtreeLengthFromRoot;
+		{
+			MaxSubtreeLengthFromRoot.SetNumZeroed(DynamicMesh.MaxVertexID());
+
+			const TArray<int32> EndPoints = EndPointAttribute.FindAllNonZero();
+			for (int32 EndVertexID : EndPoints)
+			{
+				const float EndpointLength = LengthFromRootAttribute.GetValue(EndVertexID);
+				if (EndpointLength <= 0.f)
+				{
+					continue; // endpoint not reachable from root, skip
+				}
+				int32 Current = EndVertexID;
+				while (Current != INDEX_NONE)
+				{
+					if (EndpointLength <= MaxSubtreeLengthFromRoot[Current])
+					{
+						break;
+					}
+					MaxSubtreeLengthFromRoot[Current] = EndpointLength;
+					Current = NextBranchPointAttribute.GetValue(Current);
+				}
+			}
+		}
+		
+		const auto FindBestParentPoint = [&](int32 PrevVertexID, int32 VertexID, const TArray<int32>& NextPoints) -> int32
 		{
 			if (NextPoints.Num() == 1)
 			{
 				return NextPoints[0];
 			}
-
 			check(NextPoints.Num() > 0);
 
-			const FVector3d& VertexPos = DynamicMesh.GetVertex(VertexID);
-			const FVector3d& PrevVertexPos = DynamicMesh.GetVertex(PrevVertexID);
+			const FVector3d VertexPos = DynamicMesh.GetVertex(VertexID);
+			const FVector3d PrevVertexPos = DynamicMesh.GetVertex(PrevVertexID);
 			const FVector3d ApicalDirection = (VertexPos - PrevVertexPos).GetSafeNormal();
+			const float CurrentRadius = VertexScaleAttribute.GetValue(VertexID);
+
+			TArray<float, TInlineAllocator<4>> Dots;
+			TArray<float, TInlineAllocator<4>> Lengths;
+			TArray<float, TInlineAllocator<4>> Radii;
+			Dots.SetNumUninitialized(NextPoints.Num());
+			Lengths.SetNumUninitialized(NextPoints.Num());
+			Radii.SetNumUninitialized(NextPoints.Num());
+
+			for (int32 i = 0; i < NextPoints.Num(); ++i)
+			{
+				const int32 NextVertexID = NextPoints[i];
+				const FVector3d NextPos = DynamicMesh.GetVertex(NextVertexID);
+				const FVector3d NextDir = (NextPos - VertexPos).GetSafeNormal();
+				Dots[i] = FVector3d::DotProduct(ApicalDirection, NextDir);
+				Lengths[i] = MaxSubtreeLengthFromRoot[NextVertexID] - LengthFromRootAttribute.GetValue(NextVertexID);
+				Radii[i] = FMath::Abs(VertexScaleAttribute.GetValue(NextVertexID) - CurrentRadius);
+			}
 
-			float BestDot = -MAX_flt;
-			int32 BestDotIndex = INDEX_NONE;
-			for (int32 NextVertexID : NextPoints)
+			// Normalize a raw float array to range [0,1].
+			// bInvert: output [1,0] instead of [0,1]
+			const auto NormalizeArray = [](TArray<float, TInlineAllocator<4>>& InOutArray, bool bInvert)
 			{
-				const FVector3d& NextVertexPos = DynamicMesh.GetVertex(NextVertexID);
-				const FVector3d NextApicalDirection = (NextVertexPos - VertexPos).GetSafeNormal();
-				const float Dot = FVector3d::DotProduct(ApicalDirection, NextApicalDirection);
-				if (Dot > BestDot)
+				float MinVal = InOutArray[0];
+				float MaxVal = InOutArray[0];
+				for (int32 i = 1; i < InOutArray.Num(); ++i)
 				{
-					BestDot = Dot;
-					BestDotIndex = NextVertexID;
+					MinVal = FMath::Min(MinVal, InOutArray[i]);
+					MaxVal = FMath::Max(MaxVal, InOutArray[i]);
+				}
+
+				const float Range = MaxVal - MinVal;
+				if (Range < UE_SMALL_NUMBER)
+				{
+					for (int32 i = 0; i < InOutArray.Num(); ++i)
+					{ 
+						InOutArray[i] = 0.5f;
+					}
+				}
+				else
+				{
+					for (int32 i = 0; i < InOutArray.Num(); ++i)
+					{
+						InOutArray[i] = (InOutArray[i] - MinVal) / Range;
+					}
+
+					if (bInvert)
+					{
+						for (int32 i = 0; i < InOutArray.Num(); ++i)
+						{
+							InOutArray[i] = 1.f - InOutArray[i];
+						}
+					}
+				}
+			};
+
+			NormalizeArray(Dots, false);
+			NormalizeArray(Lengths,false);
+			NormalizeArray(Radii, true);
+
+			float BestScore = -MAX_flt;
+			int32 BestIndex = INDEX_NONE;
+			for (int32 i = 0; i < NextPoints.Num(); ++i)
+			{
+				const float Score = Dots[i] + Lengths[i] + Radii[i];
+				if (Score > BestScore)
+				{
+					BestScore = Score;
+					BestIndex = NextPoints[i];
 				}
 			}
 
-			check(BestDotIndex != INDEX_NONE);
-			return BestDotIndex;
+			check(BestIndex != INDEX_NONE);
+			return BestIndex;
 		};
 
 		const auto FindNextBranchPoints = [&](int32 VertexID)
@@ -542,6 +625,8 @@ void PV::ImportHelper::ComputeDynamicMeshBranchHierarchies(
 		RootPointsAttribute,
 		EndPointsAttribute,
 		NextBranchPointAttribute,
+		VertexScaleAttribute,
+		LengthFromRootAttribute,
 		OutDynamicMeshBranchHierarchies
 	);
 }
