Fix crash in tree visualization by guarding attribute access with validity checks before use
Refactor visualization value computation to use typed direct attribute access instead of facade wrappers with raw index access
Introduce structured view types for abscission senescence and lateral elongation attributes to replace magic-index float arrays
Fix issues when visualizing a fused point skeleton

#Synced-CL 51755814

diff --git a/Engine/Plugins/Experimental/ProceduralVegetationEditor/Source/ProceduralVegetation/Private/Facades/PVTreeFacade.cpp b/Engine/Plugins/Experimental/ProceduralVegetationEditor/Source/ProceduralVegetation/Private/Facades/PVTreeFacade.cpp
--- a/Engine/Plugins/Experimental/ProceduralVegetationEditor/Source/ProceduralVegetation/Private/Facades/PVTreeFacade.cpp
+++ b/Engine/Plugins/Experimental/ProceduralVegetationEditor/Source/ProceduralVegetation/Private/Facades/PVTreeFacade.cpp
@@ -155,256 +155,372 @@ namespace PV::Facades
 
 	TArray<float> FTreeFacade::GetVisualizationValues(const FManagedArrayCollection& Collection, ESkeletonVisualizationModes VisualizationMode)
 	{
-		const FBranchFacade BranchFacade(Collection);
-		const FPointFacade PointFacade(Collection);
-		const FMetaInfoFacade DetailsFacade(Collection);
-
 		TArray<float> VisualizationValues;
-		VisualizationValues.SetNum(PointFacade.GetElementCount());
 
 		switch (VisualizationMode)
 		{
 		case ESkeletonVisualizationModes::None:
 			{
-				VisualizationValues.Empty();
 				break;
 			}
 		case ESkeletonVisualizationModes::GravitationalStress:
 			{
-				const TArray<float>& AbscissionSenescenseArray = DetailsFacade.GetAbscissionSenescense();
-				const float WeightAbscissionThreshold = 1 - AbscissionSenescenseArray[0];
-
-				for (int BranchIndex = 0; BranchIndex < BranchFacade.GetElementCount(); ++BranchIndex)
+				const PV::FBranchPointsAttributeConstView               BranchPointsAttribute         = PV::FBranchPointsAttribute::FindAttribute(Collection);
+				const PV::FBudDirectionAttributeConstView               BudDirectionAttribute         = PV::FBudDirectionAttribute::FindAttribute(Collection);
+				const PV::FBudLateralMeristemAttributeConstView         BudLateralMeristemAttribute   = PV::FBudLateralMeristemAttribute::FindAttribute(Collection);
+				const PV::FDetailAbscissionSenescenseAttributeConstView AbscissionSenescenseAttribute = PV::FDetailAbscissionSenescenseAttribute::FindAttribute(Collection);
+				const PV::FBranchParentNumberAttributeConstView         BranchParentNumberAttribute   = PV::FBranchParentNumberAttribute::FindAttribute(Collection);
+				const PV::FBranchNumberAttributeConstView               BranchNumberAttribute         = PV::FBranchNumberAttribute::FindAttribute(Collection);
+				if (PV::ValidateAttributeCollection(BranchPointsAttribute, BudDirectionAttribute, BudLateralMeristemAttribute, AbscissionSenescenseAttribute, BranchParentNumberAttribute, BranchNumberAttribute))
 				{
-					const TArray<int32>& BranchPoints = BranchFacade.GetPoints(BranchIndex);
-					const FVector3f FirstBudDirection = PointFacade.GetBudDirection(BranchPoints[0])[0];
-					const FVector3f SecondBudDirection = PointFacade.GetBudDirection(BranchPoints[1])[0];
-					const float BirthParentChildDot = PointFacade.GetBudLateralMeristem(BranchPoints[1])[4];
-
-					const float GravityParentChildDot = FVector3f::DotProduct(SecondBudDirection, FirstBudDirection);
-					const float GravitationalPressure = FMath::GetMappedRangeValueUnclamped(
-						FVector2f{0.0f, WeightAbscissionThreshold},
-						FVector2f{0.0f, 1.0f},
-						BirthParentChildDot - GravityParentChildDot
-					);
-					for (const int32 BranchPointIndex : BranchPoints)
+					VisualizationValues.SetNum(BudDirectionAttribute.Num());
+					const float WeightAbscissionThreshold = 1 - AbscissionSenescenseAttribute[0].GravitationalAbscissionFactor;
+					const bool bHasFusedPoints = PV::PlantTraversalHelper::HasFusedPoints(BranchPointsAttribute, BranchParentNumberAttribute, BranchNumberAttribute);
+
+					for (int32 BranchIndex = 0; BranchIndex < BranchPointsAttribute.Num(); ++BranchIndex)
 					{
-						VisualizationValues[BranchPointIndex] = GravitationalPressure;
+						const TArray<int32>& BranchPoints  = BranchPointsAttribute[BranchIndex];
+						if (BranchPoints.Num() < 2)
+						{
+							continue;
+						}
+
+						const FVector3f FirstBudDirection  = BudDirectionAttribute[BranchPoints[0]].Apical;
+						const FVector3f SecondBudDirection = BudDirectionAttribute[BranchPoints[1]].Apical;
+						const float BirthParentChildDot    = BudLateralMeristemAttribute[BranchPoints[1]].ParentDot;
+
+						const float GravityParentChildDot = FVector3f::DotProduct(SecondBudDirection, FirstBudDirection);
+						const float GravitationalPressure = FMath::GetMappedRangeValueUnclamped(
+							FVector2f{0.0f, WeightAbscissionThreshold},
+							FVector2f{0.0f, 1.0f},
+							BirthParentChildDot - GravityParentChildDot
+						);
+						const bool bIsTrunk = PV::PlantTraversalHelper::IsTrunk(BranchParentNumberAttribute, BranchIndex);
+						const int32 StartIndex = bHasFusedPoints && !bIsTrunk ? 1 : 0;
+						for (int32 BranchPointIndex = StartIndex; BranchPointIndex < BranchPoints.Num(); ++BranchPointIndex)
+						{
+							VisualizationValues[BranchPoints[BranchPointIndex]] = GravitationalPressure;
+						}
 					}
 				}
 				break;
 			}
 		case ESkeletonVisualizationModes::CellDensity:
 			{
-				const TArray<float>& LateralElongationArray = DetailsFacade.GetLateralElongation();
-
-				const float CellDensity = (1 - LateralElongationArray[1]) * (LateralElongationArray[0] * 100.0f);
-
-				for (int BranchIndex = 0; BranchIndex < BranchFacade.GetElementCount(); ++BranchIndex)
+				const PV::FBranchPointsAttributeConstView            BranchPointsAttribute       = PV::FBranchPointsAttribute::FindAttribute(Collection);
+				const PV::FBudLateralMeristemAttributeConstView      BudLateralMeristemAttribute = PV::FBudLateralMeristemAttribute::FindAttribute(Collection);
+				const PV::FDetailLateralElongationAttributeConstView LateralElongationAttribute  = PV::FDetailLateralElongationAttribute::FindAttribute(Collection);
+				if (PV::ValidateAttributeCollection(BranchPointsAttribute, BudLateralMeristemAttribute, LateralElongationAttribute))
 				{
-					const TArray<int32>& BranchPoints = BranchFacade.GetPoints(BranchIndex);
-					for (const int32 BranchPointIndex : BranchPoints)
+					VisualizationValues.SetNum(BudLateralMeristemAttribute.Num());
+					const float CellDensity = (1 - LateralElongationAttribute[0].SeedScaleEffect) * (LateralElongationAttribute[0].Elongation * 100.0f);
+
+					for (int32 BranchIndex = 0; BranchIndex < BranchPointsAttribute.Num(); ++BranchIndex)
 					{
-						const float CurrentLateralMeristem = PointFacade.GetBudLateralMeristem(BranchPointIndex)[0];
-						VisualizationValues[BranchPointIndex] = FMath::GetMappedRangeValueClamped(
-							FVector2f{0.0f, CellDensity},
-							FVector2f{0.0f, 1.0f},
-							CurrentLateralMeristem
-						);
+						const TArray<int32>& BranchPoints = BranchPointsAttribute[BranchIndex];
+						for (const int32 BranchPointIndex : BranchPoints)
+						{
+							const float CurrentLateralMeristem = BudLateralMeristemAttribute[BranchPointIndex].LateralMeristem;
+							VisualizationValues[BranchPointIndex] = FMath::GetMappedRangeValueClamped(
+								FVector2f{0.0f, CellDensity},
+								FVector2f{0.0f, 1.0f},
+								CurrentLateralMeristem
+							);
+						}
 					}
 				}
 				break;
 			}
 		case ESkeletonVisualizationModes::CellWeight:
 			{
-				const TArray<float>& LateralElongationArray = DetailsFacade.GetLateralElongation();
-
-				const float CellWeight = LateralElongationArray[2];
-
-				for (int BranchIndex = 0; BranchIndex < BranchFacade.GetElementCount(); ++BranchIndex)
+				const PV::FBranchPointsAttributeConstView            BranchPointsAttribute       = PV::FBranchPointsAttribute::FindAttribute(Collection);
+				const PV::FBudLateralMeristemAttributeConstView      BudLateralMeristemAttribute = PV::FBudLateralMeristemAttribute::FindAttribute(Collection);
+				const PV::FDetailLateralElongationAttributeConstView LateralElongationAttribute  = PV::FDetailLateralElongationAttribute::FindAttribute(Collection);
+				if (PV::ValidateAttributeCollection(BranchPointsAttribute, BudLateralMeristemAttribute, LateralElongationAttribute))
 				{
-					const TArray<int32>& BranchPoints = BranchFacade.GetPoints(BranchIndex);
-					for (const int32 BranchPointIndex : BranchPoints)
+					VisualizationValues.SetNum(BudLateralMeristemAttribute.Num());
+					const float CellWeight = LateralElongationAttribute[0].WhorledImpact;
+
+					for (int32 BranchIndex = 0; BranchIndex < BranchPointsAttribute.Num(); ++BranchIndex)
 					{
-						const float CurrentLateralWeight = PointFacade.GetBudLateralMeristem(BranchPointIndex)[3] * CellWeight;
-						VisualizationValues[BranchPointIndex] = FMath::Clamp(CurrentLateralWeight, 0.0f, 1.0f);
+						const TArray<int32>& BranchPoints = BranchPointsAttribute[BranchIndex];
+						for (const int32 BranchPointIndex : BranchPoints)
+						{
+							const float CurrentLateralWeight = BudLateralMeristemAttribute[BranchPointIndex].Davinci * CellWeight;
+							VisualizationValues[BranchPointIndex] = FMath::Clamp(CurrentLateralWeight, 0.0f, 1.0f);
+						}
 					}
 				}
 				break;
 			}
 		case ESkeletonVisualizationModes::LightSenescence:
 			{
-				const TArray<float>& AbscissionSenescenseArray = DetailsFacade.GetAbscissionSenescense();
-				const float LightAbscissionMin = AbscissionSenescenseArray[6];
-				const float LightAbscissionMax = AbscissionSenescenseArray[7];
-				for (int BranchIndex = 0; BranchIndex < BranchFacade.GetElementCount(); ++BranchIndex)
+				const PV::FBranchPointsAttributeConstView               BranchPointsAttribute         = PV::FBranchPointsAttribute::FindAttribute(Collection);
+				const PV::FBranchNumberAttributeConstView               BranchNumberAttribute         = PV::FBranchNumberAttribute::FindAttribute(Collection);
+				const PV::FBudDevelopmentAttributeConstView             BudDevelopmentAttribute       = PV::FBudDevelopmentAttribute::FindAttribute(Collection);
+				const PV::FDetailAbscissionSenescenseAttributeConstView AbscissionSenescenseAttribute = PV::FDetailAbscissionSenescenseAttribute::FindAttribute(Collection);
+				const PV::FBranchParentNumberAttributeConstView         BranchParentNumberAttribute   = PV::FBranchParentNumberAttribute::FindAttribute(Collection);
+				if (PV::ValidateAttributeCollection(BranchPointsAttribute, BranchNumberAttribute, BudDevelopmentAttribute, AbscissionSenescenseAttribute, BranchParentNumberAttribute))
 				{
-					const TArray<int32>& BranchPoints = BranchFacade.GetPoints(BranchIndex);
-					const int32 BranchNumber = BranchFacade.GetBranchNumber(BranchIndex);
-					FRandomStream RandomStream(BranchNumber);
-
-					const int32 LightSenescence = PointFacade.GetBudDevelopment(BranchPoints[1])[4];
-					const float LightAbscissionThreshold = FMath::RoundToInt32(RandomStream.FRandRange(LightAbscissionMin, LightAbscissionMax));
-					const float LightSenescenceValue = FMath::GetMappedRangeValueUnclamped(
-						FVector2f{0.0f, LightAbscissionThreshold},
-						FVector2f{0.0f, 1.0f},
-						LightSenescence
-					);
-					for (const int32 BranchPointIndex : BranchPoints)
+					VisualizationValues.SetNum(BudDevelopmentAttribute.Num());
+					const float LightAbscissionMin = AbscissionSenescenseAttribute[0].LightAbscissionMin;
+					const float LightAbscissionMax = AbscissionSenescenseAttribute[0].LightAbscissionMax;
+					const bool bHasFusedPoints = PV::PlantTraversalHelper::HasFusedPoints(BranchPointsAttribute, BranchParentNumberAttribute, BranchNumberAttribute);
+
+					for (int32 BranchIndex = 0; BranchIndex < BranchPointsAttribute.Num(); ++BranchIndex)
 					{
-						VisualizationValues[BranchPointIndex] = LightSenescenceValue;
+						const TArray<int32>& BranchPoints = BranchPointsAttribute[BranchIndex];
+						if (BranchPoints.Num() < 2)
+						{
+							continue;
+						}
+
+						const int32 BranchNumber = BranchNumberAttribute[BranchIndex];
+						FRandomStream RandomStream(BranchNumber);
+
+						const int32 LightSenescence = BudDevelopmentAttribute[BranchPoints[1]].LightSenescense;
+						const float LightAbscissionThreshold = FMath::RoundToInt32(RandomStream.FRandRange(LightAbscissionMin, LightAbscissionMax));
+						const float LightSenescenceValue = FMath::GetMappedRangeValueUnclamped(
+							FVector2f{0.0f, LightAbscissionThreshold},
+							FVector2f{0.0f, 1.0f},
+							LightSenescence
+						);
+						const bool bIsTrunk = PV::PlantTraversalHelper::IsTrunk(BranchParentNumberAttribute, BranchIndex);
+						const int32 StartIndex = bHasFusedPoints && !bIsTrunk ? 1 : 0;
+						for (int32 BranchPointIndex = StartIndex; BranchPointIndex < BranchPoints.Num(); ++BranchPointIndex)
+						{
+							VisualizationValues[BranchPoints[BranchPointIndex]] = LightSenescenceValue;
+						}
 					}
 				}
 				break;
 			}
 		case ESkeletonVisualizationModes::BranchLight:
 			{
-				for (int BranchIndex = 0; BranchIndex < BranchFacade.GetElementCount(); ++BranchIndex)
+				const PV::FBranchPointsAttributeConstView       BranchPointsAttribute       = PV::FBranchPointsAttribute::FindAttribute(Collection);
+				const PV::FBudLightDetectedAttributeConstView   BudLightDetectedAttribute   = PV::FBudLightDetectedAttribute::FindAttribute(Collection);
+				const PV::FBranchParentNumberAttributeConstView BranchParentNumberAttribute = PV::FBranchParentNumberAttribute::FindAttribute(Collection);
+				const PV::FBranchNumberAttributeConstView       BranchNumberAttribute       = PV::FBranchNumberAttribute::FindAttribute(Collection);
+				if (PV::ValidateAttributeCollection(BranchPointsAttribute, BudLightDetectedAttribute, BranchParentNumberAttribute, BranchNumberAttribute))
 				{
-					const TArray<int32>& BranchPoints = BranchFacade.GetPoints(BranchIndex);
-					const float LightDetected = PointFacade.GetBudLightDetected(BranchPoints[1])[2];
-					for (const int32 BranchPointIndex : BranchPoints)
+					VisualizationValues.SetNum(BudLightDetectedAttribute.Num());
+					const bool bHasFusedPoints = PV::PlantTraversalHelper::HasFusedPoints(BranchPointsAttribute, BranchParentNumberAttribute, BranchNumberAttribute);
+					for (int32 BranchIndex = 0; BranchIndex < BranchPointsAttribute.Num(); ++BranchIndex)
 					{
-						VisualizationValues[BranchPointIndex] = LightDetected;
+						const TArray<int32>& BranchPoints = BranchPointsAttribute[BranchIndex];
+						if (BranchPoints.Num() < 2)
+						{
+							continue;
+						}
+
+						const float LightDetected = BudLightDetectedAttribute[BranchPoints[1]].Branch;
+						const bool bIsTrunk = PV::PlantTraversalHelper::IsTrunk(BranchParentNumberAttribute, BranchIndex);
+						const int32 StartIndex = bHasFusedPoints && !bIsTrunk ? 1 : 0;
+						for (int32 BranchPointIndex = StartIndex; BranchPointIndex < BranchPoints.Num(); ++BranchPointIndex)
+						{
+							VisualizationValues[BranchPoints[BranchPointIndex]] = LightDetected;
+						}
 					}
 				}
 				break;
 			}
 		case ESkeletonVisualizationModes::LightAbscissionRetention:
 			{
-				const TArray<float>& AbscissionSenescenseArray = DetailsFacade.GetAbscissionSenescense();
-				const TArray<float>& LateralElongationArray = DetailsFacade.GetLateralElongation();
-
-				const float LightRetention = AbscissionSenescenseArray[8];
-				const float LateralElongation = LateralElongationArray[0];
-				for (int BranchIndex = 0; BranchIndex < BranchFacade.GetElementCount(); ++BranchIndex)
+				const PV::FBranchPointsAttributeConstView               BranchPointsAttribute         = PV::FBranchPointsAttribute::FindAttribute(Collection);
+				const PV::FBudLateralMeristemAttributeConstView         BudLateralMeristemAttribute   = PV::FBudLateralMeristemAttribute::FindAttribute(Collection);
+				const PV::FDetailAbscissionSenescenseAttributeConstView AbscissionSenescenseAttribute = PV::FDetailAbscissionSenescenseAttribute::FindAttribute(Collection);
+				const PV::FDetailLateralElongationAttributeConstView    LateralElongationAttribute    = PV::FDetailLateralElongationAttribute::FindAttribute(Collection);
+				if (PV::ValidateAttributeCollection(BranchPointsAttribute, BudLateralMeristemAttribute, AbscissionSenescenseAttribute, LateralElongationAttribute))
 				{
-					const TArray<int32>& BranchPoints = BranchFacade.GetPoints(BranchIndex);
-					for (const int32 BranchPointIndex : BranchPoints)
+					VisualizationValues.SetNum(BudLateralMeristemAttribute.Num());
+					const float LightRetention    = AbscissionSenescenseAttribute[0].LightScaleRetention;
+					const float LateralElongation = LateralElongationAttribute[0].Elongation;
+
+					for (int32 BranchIndex = 0; BranchIndex < BranchPointsAttribute.Num(); ++BranchIndex)
 					{
-						const float LateralMeristem = PointFacade.GetBudLateralMeristem(BranchPointIndex)[0];
-						const float LightRetentionThreshold = (LateralElongation * 100.0f) * (1 - LightRetention);
-						const float LightRetentionValue = FMath::GetMappedRangeValueUnclamped(
-							FVector2f{0.0f, LightRetentionThreshold},
-							FVector2f{0.0f, 1.0f},
-							LateralMeristem
-						);
-						VisualizationValues[BranchPointIndex] = LightRetentionValue;
+						const TArray<int32>& BranchPoints = BranchPointsAttribute[BranchIndex];
+						for (const int32 BranchPointIndex : BranchPoints)
+						{
+							const float LateralMeristem = BudLateralMeristemAttribute[BranchPointIndex].LateralMeristem;
+							const float LightRetentionThreshold = (LateralElongation * 100.0f) * (1 - LightRetention);
+							const float LightRetentionValue = FMath::GetMappedRangeValueUnclamped(
+								FVector2f{0.0f, LightRetentionThreshold},
+								FVector2f{0.0f, 1.0f},
+								LateralMeristem
+							);
+							VisualizationValues[BranchPointIndex] = LightRetentionValue;
+						}
 					}
 				}
 				break;
 			}
 		case ESkeletonVisualizationModes::AgeSenescence:
 			{
-				const TArray<float>& AbscissionSenescenseArray = DetailsFacade.GetAbscissionSenescense();
-
-				const float AgeAbscissionMin = AbscissionSenescenseArray[2];
-				const float AgeAbscissionMax = AbscissionSenescenseArray[3];
-				for (int BranchIndex = 0; BranchIndex < BranchFacade.GetElementCount(); ++BranchIndex)
+				const PV::FBranchPointsAttributeConstView               BranchPointsAttribute         = PV::FBranchPointsAttribute::FindAttribute(Collection);
+				const PV::FBranchNumberAttributeConstView               BranchNumberAttribute         = PV::FBranchNumberAttribute::FindAttribute(Collection);
+				const PV::FBudDevelopmentAttributeConstView             BudDevelopmentAttribute       = PV::FBudDevelopmentAttribute::FindAttribute(Collection);
+				const PV::FDetailAbscissionSenescenseAttributeConstView AbscissionSenescenseAttribute = PV::FDetailAbscissionSenescenseAttribute::FindAttribute(Collection);
+				const PV::FBranchParentNumberAttributeConstView         BranchParentNumberAttribute   = PV::FBranchParentNumberAttribute::FindAttribute(Collection);
+				if (PV::ValidateAttributeCollection(BranchPointsAttribute, BranchNumberAttribute, BudDevelopmentAttribute, AbscissionSenescenseAttribute, BranchParentNumberAttribute))
 				{
-					const TArray<int32>& BranchPoints = BranchFacade.GetPoints(BranchIndex);
-					const int32 BranchNumber = BranchFacade.GetBranchNumber(BranchIndex);
-					FRandomStream RandomStream(BranchNumber);
-
-					const int32 AgeSenescence = PointFacade.GetBudDevelopment(BranchPoints[1])[3];
-					const float AgeAbscissionThreshold = FMath::RoundToInt32(RandomStream.FRandRange(AgeAbscissionMin, AgeAbscissionMax));
-					const float AgeSenescenceValue = FMath::GetMappedRangeValueUnclamped(
-						FVector2f{0.0f, AgeAbscissionThreshold},
-						FVector2f{0.0f, 1.0f},
-						AgeSenescence
-					);
-					for (const int32 BranchPointIndex : BranchPoints)
+					VisualizationValues.SetNum(BudDevelopmentAttribute.Num());
+					const float AgeAbscissionMin = AbscissionSenescenseAttribute[0].AgeAbscissionMin;
+					const float AgeAbscissionMax = AbscissionSenescenseAttribute[0].AgeAbscissionMax;
+					const bool bHasFusedPoints = PV::PlantTraversalHelper::HasFusedPoints(BranchPointsAttribute, BranchParentNumberAttribute, BranchNumberAttribute);
+
+					for (int32 BranchIndex = 0; BranchIndex < BranchPointsAttribute.Num(); ++BranchIndex)
 					{
-						VisualizationValues[BranchPointIndex] = AgeSenescenceValue;
+						const TArray<int32>& BranchPoints = BranchPointsAttribute[BranchIndex];
+						if (BranchPoints.Num() < 2)
+						{
+							continue;
+						}
+
+						const int32 BranchNumber = BranchNumberAttribute[BranchIndex];
+						FRandomStream RandomStream(BranchNumber);
+
+						const int32 AgeSenescence = BudDevelopmentAttribute[BranchPoints[1]].AgeSenescense;
+						const float AgeAbscissionThreshold = FMath::RoundToInt32(RandomStream.FRandRange(AgeAbscissionMin, AgeAbscissionMax));
+						const float AgeSenescenceValue = FMath::GetMappedRangeValueUnclamped(
+							FVector2f{0.0f, AgeAbscissionThreshold},
+							FVector2f{0.0f, 1.0f},
+							AgeSenescence
+						);
+						const bool bIsTrunk = PV::PlantTraversalHelper::IsTrunk(BranchParentNumberAttribute, BranchIndex);
+						const int32 StartIndex = bHasFusedPoints && !bIsTrunk ? 1 : 0;
+						for (int32 BranchPointIndex = StartIndex; BranchPointIndex < BranchPoints.Num(); ++BranchPointIndex)
+						{
+							VisualizationValues[BranchPoints[BranchPointIndex]] = AgeSenescenceValue;
+						}
 					}
 				}
 				break;
 			}
 		case ESkeletonVisualizationModes::AgeAbscissionRetention:
 			{
-				const TArray<float>& AbscissionSenescenseArray = DetailsFacade.GetAbscissionSenescense();
-				const TArray<float>& LateralElongationArray = DetailsFacade.GetLateralElongation();
-
-				const float AgeRetention = AbscissionSenescenseArray[4];
-				const float LateralElongation = LateralElongationArray[0];
-				for (int BranchIndex = 0; BranchIndex < BranchFacade.GetElementCount(); ++BranchIndex)
+				const PV::FBranchPointsAttributeConstView               BranchPointsAttribute         = PV::FBranchPointsAttribute::FindAttribute(Collection);
+				const PV::FBudLateralMeristemAttributeConstView         BudLateralMeristemAttribute   = PV::FBudLateralMeristemAttribute::FindAttribute(Collection);
+				const PV::FDetailAbscissionSenescenseAttributeConstView AbscissionSenescenseAttribute = PV::FDetailAbscissionSenescenseAttribute::FindAttribute(Collection);
+				const PV::FDetailLateralElongationAttributeConstView    LateralElongationAttribute    = PV::FDetailLateralElongationAttribute::FindAttribute(Collection);
+				if (PV::ValidateAttributeCollection(BranchPointsAttribute, BudLateralMeristemAttribute, AbscissionSenescenseAttribute, LateralElongationAttribute))
 				{
-					const TArray<int32>& BranchPoints = BranchFacade.GetPoints(BranchIndex);
-					for (const int32 BranchPointIndex : BranchPoints)
+					VisualizationValues.SetNum(BudLateralMeristemAttribute.Num());
+					const float AgeRetention    = AbscissionSenescenseAttribute[0].AgeScaleRetention;
+					const float LateralElongation = LateralElongationAttribute[0].Elongation;
+
+					for (int32 BranchIndex = 0; BranchIndex < BranchPointsAttribute.Num(); ++BranchIndex)
 					{
-						const float LateralMeristem = PointFacade.GetBudLateralMeristem(BranchPointIndex)[0];
-						const float AgeRetentionThreshold = (LateralElongation * 100.0f) * (1 - AgeRetention);
-						const float AgeRetentionValue = FMath::GetMappedRangeValueUnclamped(
-							FVector2f{0.0f, AgeRetentionThreshold},
-							FVector2f{0.0f, 1.0f},
-							LateralMeristem
-						);
-						VisualizationValues[BranchPointIndex] = AgeRetentionValue;
+						const TArray<int32>& BranchPoints = BranchPointsAttribute[BranchIndex];
+						for (const int32 BranchPointIndex : BranchPoints)
+						{
+							const float LateralMeristem = BudLateralMeristemAttribute[BranchPointIndex].LateralMeristem;
+							const float AgeRetentionThreshold = (LateralElongation * 100.0f) * (1 - AgeRetention);
+							const float AgeRetentionValue = FMath::GetMappedRangeValueUnclamped(
+								FVector2f{0.0f, AgeRetentionThreshold},
+								FVector2f{0.0f, 1.0f},
+								LateralMeristem
+							);
+							VisualizationValues[BranchPointIndex] = AgeRetentionValue;
+						}
 					}
 				}
 				break;
 			}
 		case ESkeletonVisualizationModes::Apical:
 			{
-				for (int32 Index = 0; Index < PointFacade.GetElementCount(); ++Index)
+				const PV::FBudHormoneLevelsAttributeConstView BudHormoneLevelsAttribute = PV::FBudHormoneLevelsAttribute::FindAttribute(Collection);
+				if (PV::ValidateAttributeCollection(BudHormoneLevelsAttribute))
 				{
-					VisualizationValues[Index] = PointFacade.GetBudHormoneLevels(Index)[0];
+					VisualizationValues.SetNum(BudHormoneLevelsAttribute.Num());
+					for (int32 Index = 0; Index < BudHormoneLevelsAttribute.Num(); ++Index)
+					{
+						VisualizationValues[Index] = BudHormoneLevelsAttribute[Index].Apical;
+					}
 				}
 				break;
 			}
 		case ESkeletonVisualizationModes::Auxin:
 			{
-				for (int32 Index = 0; Index < PointFacade.GetElementCount(); ++Index)
+				const PV::FBudHormoneLevelsAttributeConstView BudHormoneLevelsAttribute = PV::FBudHormoneLevelsAttribute::FindAttribute(Collection);
+				if (PV::ValidateAttributeCollection(BudHormoneLevelsAttribute))
 				{
-					VisualizationValues[Index] = PointFacade.GetBudHormoneLevels(Index)[1];
+					VisualizationValues.SetNum(BudHormoneLevelsAttribute.Num());
+					for (int32 Index = 0; Index < BudHormoneLevelsAttribute.Num(); ++Index)
+					{
+						VisualizationValues[Index] = BudHormoneLevelsAttribute[Index].Axillary;
+					}
 				}
 				break;
 			}
 		case ESkeletonVisualizationModes::Radical:
 			{
-				for (int32 Index = 0; Index < PointFacade.GetElementCount(); ++Index)
+				const PV::FBudHormoneLevelsAttributeConstView BudHormoneLevelsAttribute = PV::FBudHormoneLevelsAttribute::FindAttribute(Collection);
+				if (PV::ValidateAttributeCollection(BudHormoneLevelsAttribute))
 				{
-					VisualizationValues[Index] = PointFacade.GetBudHormoneLevels(Index)[3];
+					VisualizationValues.SetNum(BudHormoneLevelsAttribute.Num());
+					for (int32 Index = 0; Index < BudHormoneLevelsAttribute.Num(); ++Index)
+					{
+						VisualizationValues[Index] = BudHormoneLevelsAttribute[Index].Radical;
+					}
 				}
 				break;
 			}
 		case ESkeletonVisualizationModes::Ethylene:
 			{
-				for (int32 Index = 0; Index < PointFacade.GetElementCount(); ++Index)
+				const PV::FBudHormoneLevelsAttributeConstView BudHormoneLevelsAttribute = PV::FBudHormoneLevelsAttribute::FindAttribute(Collection);
+				if (PV::ValidateAttributeCollection(BudHormoneLevelsAttribute))
 				{
-					VisualizationValues[Index] = PointFacade.GetBudHormoneLevels(Index)[4];
+					VisualizationValues.SetNum(BudHormoneLevelsAttribute.Num());
+					for (int32 Index = 0; Index < BudHormoneLevelsAttribute.Num(); ++Index)
+					{
+						VisualizationValues[Index] = BudHormoneLevelsAttribute[Index].Ethylene;
+					}
 				}
 				break;
 			}
 		case ESkeletonVisualizationModes::AuxinRetention:
 			{
-				const TArray<float>& LeafGrowthArray = DetailsFacade.GetLeafGrowth();
-				const float AuxinRetention = LeafGrowthArray[5];
-
-				for (int32 Index = 0; Index < PointFacade.GetElementCount(); ++Index)
+				const PV::FBudHormoneLevelsAttributeConstView BudHormoneLevelsAttribute = PV::FBudHormoneLevelsAttribute::FindAttribute(Collection);
+				const PV::FDetailLeafGrowthAttributeConstView LeafGrowthAttribute       = PV::FDetailLeafGrowthAttribute::FindAttribute(Collection);
+				if (PV::ValidateAttributeCollection(BudHormoneLevelsAttribute, LeafGrowthAttribute))
 				{
-					const float EthyleneLevel = FMath::Clamp(PointFacade.GetBudHormoneLevels(Index)[4], 0.0f, 1.0f);
-					const float ApicalAuxin = 1.0f - FMath::Lerp(EthyleneLevel, 0, PointFacade.GetBudHormoneLevels(Index)[0]);
-					VisualizationValues[Index] = FMath::Lerp(0, ApicalAuxin, AuxinRetention);
+					VisualizationValues.SetNum(BudHormoneLevelsAttribute.Num());
+					const float AuxinRetention = LeafGrowthAttribute[0].EthyleneThreshold;
+
+					for (int32 Index = 0; Index < BudHormoneLevelsAttribute.Num(); ++Index)
+					{
+						const PV::FBudHormoneLevelsConstView HormoneLevels = BudHormoneLevelsAttribute[Index];
+						const float EthyleneLevel = FMath::Clamp(HormoneLevels.Ethylene, 0.0f, 1.0f);
+						const float ApicalAuxin = 1.0f - FMath::Lerp(EthyleneLevel, 0, HormoneLevels.Apical);
+						VisualizationValues[Index] = FMath::Lerp(0, ApicalAuxin, AuxinRetention);
+					}
 				}
 				break;
 			}
 		case ESkeletonVisualizationModes::Cytokinin:
 			{
-				for (int32 Index = 0; Index < PointFacade.GetElementCount(); ++Index)
+				const PV::FBudHormoneLevelsAttributeConstView BudHormoneLevelsAttribute = PV::FBudHormoneLevelsAttribute::FindAttribute(Collection);
+				if (PV::ValidateAttributeCollection(BudHormoneLevelsAttribute))
 				{
-					VisualizationValues[Index] = PointFacade.GetBudHormoneLevels(Index)[5];
+					VisualizationValues.SetNum(BudHormoneLevelsAttribute.Num());
+					for (int32 Index = 0; Index < BudHormoneLevelsAttribute.Num(); ++Index)
+					{
+						VisualizationValues[Index] = BudHormoneLevelsAttribute[Index].Cytokinin;
+					}
 				}
 				break;
 			}
 		case ESkeletonVisualizationModes::BranchGeneration:
 			{
+				const PV::FPointPositionAttributeConstView      PointPositionAttribute      = PV::FPointPositionAttribute::FindAttribute(Collection);
 				const PV::FBranchPointsAttributeConstView       BranchPointsAttribute       = PV::FBranchPointsAttribute::FindAttribute(Collection);
 				const PV::FBranchParentNumberAttributeConstView BranchParentNumberAttribute = PV::FBranchParentNumberAttribute::FindAttribute(Collection);
 				const PV::FBranchNumberAttributeConstView       BranchNumberAttribute       = PV::FBranchNumberAttribute::FindAttribute(Collection);
 				const PV::FBranchParentsAttributeConstView      BranchParentsAttribute      = PV::FBranchParentsAttribute::FindAttribute(Collection);
-				if (PV::AttributesHelper::ValidateAttributeCollection(BranchPointsAttribute, BranchParentNumberAttribute, BranchNumberAttribute, BranchParentsAttribute))
+				if (PV::ValidateAttributeCollection(PointPositionAttribute, BranchPointsAttribute, BranchParentNumberAttribute, BranchNumberAttribute, BranchParentsAttribute))
 				{
+					VisualizationValues.SetNum(PointPositionAttribute.Num());
+
 					const bool bHasFusedPoints = PV::PlantTraversalHelper::HasFusedPoints(BranchPointsAttribute, BranchParentNumberAttribute, BranchNumberAttribute);
 					for (int32 BranchIndex = 0; BranchIndex < BranchParentsAttribute.Num(); ++BranchIndex)
 					{
@@ -423,7 +539,7 @@ namespace PV::Facades
 			}
 		default:
 			{
-				VisualizationValues.Empty();
+				break;
 			}
 		}
 		
diff --git a/Engine/Plugins/Experimental/ProceduralVegetationEditor/Source/ProceduralVegetation/Private/Helpers/PVAttributesHelper.h b/Engine/Plugins/Experimental/ProceduralVegetationEditor/Source/ProceduralVegetation/Private/Helpers/PVAttributesHelper.h
--- a/Engine/Plugins/Experimental/ProceduralVegetationEditor/Source/ProceduralVegetation/Private/Helpers/PVAttributesHelper.h
+++ b/Engine/Plugins/Experimental/ProceduralVegetationEditor/Source/ProceduralVegetation/Private/Helpers/PVAttributesHelper.h
@@ -5,38 +5,6 @@
 
 namespace PV::AttributesHelper
 {
-	template<typename... TAttributeViews>
-	bool ValidateAttributeCollection(TAttributeViews... Attributes)
-	{
-		TMap<FName, int32> GroupSizes;
-		GroupSizes.Reserve(sizeof...(TAttributeViews));
-
-		bool bIsValid = true;
-		auto AddEntry = [&GroupSizes, &bIsValid](auto Attr)
-		{
-			if (!Attr.IsValid())
-			{
-				bIsValid = false;
-				return;
-			}
-
-			const FName GroupName = std::remove_reference_t<decltype(Attr)>::GroupName;
-			const int32 Size = Attr.Num();
-			if (const int32* ExistingSize = GroupSizes.Find(GroupName))
-			{
-				bIsValid &= (*ExistingSize == Size);
-			}
-			else
-			{
-				GroupSizes.Add(GroupName, Size);
-			}
-		};
-
-		(AddEntry(Attributes), ...);
-
-		return bIsValid;
-	}
-
 	struct FComputeBudNumbersAttributes
 	{
 		FPointBudNumberAttributeView PointBudNumber;
diff --git a/Engine/Plugins/Experimental/ProceduralVegetationEditor/Source/ProceduralVegetation/Private/Tests/PVAttributesTests.cpp b/Engine/Plugins/Experimental/ProceduralVegetationEditor/Source/ProceduralVegetation/Private/Tests/PVAttributesTests.cpp
--- a/Engine/Plugins/Experimental/ProceduralVegetationEditor/Source/ProceduralVegetation/Private/Tests/PVAttributesTests.cpp
+++ b/Engine/Plugins/Experimental/ProceduralVegetationEditor/Source/ProceduralVegetation/Private/Tests/PVAttributesTests.cpp
@@ -824,10 +824,8 @@ PV_ATTRIBUTES_TEST(ComputeSmoothedPointScales)
 
 PV_ATTRIBUTES_TEST(ValidateAttributeCollection)
 {
-	using namespace PV::AttributesHelper;
-
 	// No attributes → always valid
-	UTEST_TRUE("NoAttributes", ValidateAttributeCollection());
+	UTEST_TRUE("NoAttributes", PV::ValidateAttributeCollection());
 
 	// Single valid attribute → valid
 	{
@@ -837,13 +835,13 @@ PV_ATTRIBUTES_TEST(ValidateAttributeCollection)
 		PV::FPointLengthFromRootAttribute::AddAttribute(Collection);
 
 		const PV::FPointLengthFromRootAttributeConstView View = PV::FPointLengthFromRootAttribute::FindAttribute(Collection);
-		UTEST_TRUE("SingleValid", ValidateAttributeCollection(View));
+		UTEST_TRUE("SingleValid", PV::ValidateAttributeCollection(View));
 	}
 
 	// Single invalid (null) attribute → invalid
 	{
 		const PV::FPointLengthFromRootAttributeConstView InvalidView;
-		UTEST_FALSE("SingleInvalid", ValidateAttributeCollection(InvalidView));
+		UTEST_FALSE("SingleInvalid", PV::ValidateAttributeCollection(InvalidView));
 	}
 
 	// Two valid attributes in the same group with equal sizes → valid
@@ -856,7 +854,7 @@ PV_ATTRIBUTES_TEST(ValidateAttributeCollection)
 
 		const PV::FPointLengthFromRootAttributeConstView LengthView = PV::FPointLengthFromRootAttribute::FindAttribute(Collection);
 		const PV::FPointScaleAttributeConstView ScaleView = PV::FPointScaleAttribute::FindAttribute(Collection);
-		UTEST_TRUE("SameGroupSameSize", ValidateAttributeCollection(LengthView, ScaleView));
+		UTEST_TRUE("SameGroupSameSize", PV::ValidateAttributeCollection(LengthView, ScaleView));
 	}
 
 	// Two valid attributes in the same group but from collections with different sizes → invalid
@@ -873,7 +871,7 @@ PV_ATTRIBUTES_TEST(ValidateAttributeCollection)
 
 		const PV::FPointLengthFromRootAttributeConstView LengthView = PV::FPointLengthFromRootAttribute::FindAttribute(Collection5);
 		const PV::FPointScaleAttributeConstView ScaleView = PV::FPointScaleAttribute::FindAttribute(Collection10);
-		UTEST_FALSE("SameGroupDifferentSizes", ValidateAttributeCollection(LengthView, ScaleView));
+		UTEST_FALSE("SameGroupDifferentSizes", PV::ValidateAttributeCollection(LengthView, ScaleView));
 	}
 
 	// Valid attributes in different groups → sizes between groups are independent, so this is valid
@@ -888,7 +886,7 @@ PV_ATTRIBUTES_TEST(ValidateAttributeCollection)
 
 		const PV::FPointLengthFromRootAttributeConstView LengthView = PV::FPointLengthFromRootAttribute::FindAttribute(Collection);
 		const PV::FBranchNumberAttributeConstView BranchNumberView = PV::FBranchNumberAttribute::FindAttribute(Collection);
-		UTEST_TRUE("DifferentGroupsDifferentSizes", ValidateAttributeCollection(LengthView, BranchNumberView));
+		UTEST_TRUE("DifferentGroupsDifferentSizes", PV::ValidateAttributeCollection(LengthView, BranchNumberView));
 	}
 
 	// One invalid attribute mixed with valid ones → invalid
@@ -902,7 +900,7 @@ PV_ATTRIBUTES_TEST(ValidateAttributeCollection)
 		const PV::FPointLengthFromRootAttributeConstView LengthView = PV::FPointLengthFromRootAttribute::FindAttribute(Collection);
 		const PV::FPointScaleAttributeConstView ScaleView = PV::FPointScaleAttribute::FindAttribute(Collection);
 		const PV::FPointPositionAttributeConstView InvalidPositionView;
-		UTEST_FALSE("OneInvalidAmongValid", ValidateAttributeCollection(LengthView, InvalidPositionView, ScaleView));
+		UTEST_FALSE("OneInvalidAmongValid", PV::ValidateAttributeCollection(LengthView, InvalidPositionView, ScaleView));
 	}
 
 	return true;
diff --git a/Engine/Plugins/Experimental/ProceduralVegetationEditor/Source/ProceduralVegetation/Private/Utils/PVAttributes.h b/Engine/Plugins/Experimental/ProceduralVegetationEditor/Source/ProceduralVegetation/Private/Utils/PVAttributes.h
--- a/Engine/Plugins/Experimental/ProceduralVegetationEditor/Source/ProceduralVegetation/Private/Utils/PVAttributes.h
+++ b/Engine/Plugins/Experimental/ProceduralVegetationEditor/Source/ProceduralVegetation/Private/Utils/PVAttributes.h
@@ -353,6 +353,45 @@ struct TFixedSizeArrayHelper
 	}
 };
 
+// Returns true if all attributes are valid, non-empty, and attributes sharing a group name have the same element count.
+template<typename... TAttributeViews>
+inline bool ValidateAttributeCollection(TAttributeViews... Attributes)
+{
+	TMap<FName, int32> GroupSizes;
+	GroupSizes.Reserve(sizeof...(TAttributeViews));
+
+	bool bIsValid = true;
+	auto AddEntry = [&GroupSizes, &bIsValid](auto Attr)
+	{
+		if (!Attr.IsValid())
+		{
+			bIsValid = false;
+			return;
+		}
+
+		if (Attr.Num() == 0)
+		{
+			bIsValid = false;
+			return;
+		}
+
+		const FName GroupName = std::remove_reference_t<decltype(Attr)>::GroupName;
+		const int32 Size = Attr.Num();
+		if (const int32* ExistingSize = GroupSizes.Find(GroupName))
+		{
+			bIsValid &= (*ExistingSize == Size);
+		}
+		else
+		{
+			GroupSizes.Add(GroupName, Size);
+		}
+	};
+
+	(AddEntry(Attributes), ...);
+
+	return bIsValid;
+}
+
 #define PV_DEFINE_STRUCTURED_ATTRIBUTE(Name, Type, ElementViewType, ElementConstViewType, AttributeNameStr, AttributeGroupStr)          \
 	using F##Name##Attribute = TAttributeDefinition<Type, ElementViewType, ElementConstViewType, AttributeNameStr, AttributeGroupStr>;  \
 	using F##Name##AttributeView = F##Name##Attribute::AttributeView;                                                                   \
@@ -601,9 +640,113 @@ PV_DEFINE_ATTRIBUTE(FoliageName, FString, AttributeNames::FoliageNameStr, GroupN
 PV_DEFINE_ATTRIBUTE(DetailLeafPhyllotaxy, TArray<float>, AttributeNames::LeafPhyllotaxyStr, GroupNames::DetailsGroupStr);
 PV_DEFINE_ATTRIBUTE(DetailFoliagePath, FString, AttributeNames::FoliagePathStr, GroupNames::DetailsGroupStr);
 PV_DEFINE_ATTRIBUTE(DetailGuid, FGuid, AttributeNames::GuidStr, GroupNames::DetailsGroupStr);
-PV_DEFINE_ATTRIBUTE(DetailLeafGrowth, TArray<float>, AttributeNames::LeafGrowthStr, GroupNames::DetailsGroupStr);
-PV_DEFINE_ATTRIBUTE(DetailAbscissionSenescense, TArray<float>, AttributeNames::AbscissionSenescenseStr, GroupNames::DetailsGroupStr);
-PV_DEFINE_ATTRIBUTE(DetailLateralElongation, TArray<float>, AttributeNames::LateralElongationStr, GroupNames::DetailsGroupStr);
+
+template<typename TFloatType>
+struct TDetailLeafGrowthView
+{
+	using ArrayType = TFixedSizeArrayHelper<TFloatType>::ArrayType;
+
+	inline const static TArray<float> DefaultData = { 0, 0, 0, 0, 0, 0, 0 };
+
+	ArrayType& Array;
+	TFloatType& Ethylene;
+	TFloatType& StartScale;
+	TFloatType& EndScale;
+	TFloatType& Density;
+	TFloatType& AbscisicAcid;
+	TFloatType& EthyleneThreshold;
+	TFloatType& LeafAuxinRetention;
+
+	TDetailLeafGrowthView() = delete;
+	TDetailLeafGrowthView(ArrayType& InData)
+		: Array(TFixedSizeArrayHelper<TFloatType>::GetFixedArray(InData, DefaultData))
+		, Ethylene(Array[0])
+		, StartScale(Array[1])
+		, EndScale(Array[2])
+		, Density(Array[3])
+		, AbscisicAcid(Array[4])
+		, EthyleneThreshold(Array[5])
+		, LeafAuxinRetention(Array[6])
+	{
+	}
+};
+using FDetailLeafGrowthView = TDetailLeafGrowthView<float>;
+using FDetailLeafGrowthConstView = TDetailLeafGrowthView<const float>;
+PV_DEFINE_STRUCTURED_ATTRIBUTE(DetailLeafGrowth, TArray<float>, FDetailLeafGrowthView, FDetailLeafGrowthConstView, AttributeNames::LeafGrowthStr, GroupNames::DetailsGroupStr);
+
+template<typename TFloatType>
+struct TDetailAbscissionSenescenseView
+{
+	using ArrayType = TFixedSizeArrayHelper<TFloatType>::ArrayType;
+
+	inline const static TArray<float> DefaultData = { 0, 0, 0, 0, 0, 0, 0, 0, 0 };
+
+	ArrayType& Array;
+	TFloatType& GravitationalAbscissionFactor;
+	TFloatType& AgeAbscissionThreshold;
+	TFloatType& AgeAbscissionMin;
+	TFloatType& AgeAbscissionMax;
+	TFloatType& AgeScaleRetention;
+	TFloatType& LightAbscissionThreshold;
+	TFloatType& LightAbscissionMin;
+	TFloatType& LightAbscissionMax;
+	TFloatType& LightScaleRetention;
+
+	TDetailAbscissionSenescenseView() = delete;
+	TDetailAbscissionSenescenseView(ArrayType& InData)
+		: Array(TFixedSizeArrayHelper<TFloatType>::GetFixedArray(InData, DefaultData))
+		, GravitationalAbscissionFactor(Array[0])
+		, AgeAbscissionThreshold(Array[1])
+		, AgeAbscissionMin(Array[2])
+		, AgeAbscissionMax(Array[3])
+		, AgeScaleRetention(Array[4])
+		, LightAbscissionThreshold(Array[5])
+		, LightAbscissionMin(Array[6])
+		, LightAbscissionMax(Array[7])
+		, LightScaleRetention(Array[8])
+	{
+	}
+};
+using FDetailAbscissionSenescenseView = TDetailAbscissionSenescenseView<float>;
+using FDetailAbscissionSenescenseConstView = TDetailAbscissionSenescenseView<const float>;
+PV_DEFINE_STRUCTURED_ATTRIBUTE(DetailAbscissionSenescense, TArray<float>, FDetailAbscissionSenescenseView, FDetailAbscissionSenescenseConstView, AttributeNames::AbscissionSenescenseStr, GroupNames::DetailsGroupStr);
+
+template<typename TFloatType>
+struct TDetailLateralElongationView
+{
+	using ArrayType = TFixedSizeArrayHelper<TFloatType>::ArrayType;
+
+	inline const static TArray<float> DefaultData = { 0, 0, 0, 0, 0, 0, 0, 0, 0 };
+
+	ArrayType& Array;
+	TFloatType& Elongation;
+	TFloatType& SeedScaleEffect;
+	TFloatType& WhorledImpact;
+	TFloatType& CellDensityMin;
+	TFloatType& CellDensityMax;
+	TFloatType& CellDensityCore;
+	TFloatType& CellDevelopmentTime;
+	TFloatType& LeafWeight;
+	TFloatType& CellDensityInput;
+
+	TDetailLateralElongationView() = delete;
+	TDetailLateralElongationView(ArrayType& InData)
+		: Array(TFixedSizeArrayHelper<TFloatType>::GetFixedArray(InData, DefaultData))
+		, Elongation(Array[0])
+		, SeedScaleEffect(Array[1])
+		, WhorledImpact(Array[2])
+		, CellDensityMin(Array[3])
+		, CellDensityMax(Array[4])
+		, CellDensityCore(Array[5])
+		, CellDevelopmentTime(Array[6])
+		, LeafWeight(Array[7])
+		, CellDensityInput(Array[8])
+	{
+	}
+};
+using FDetailLateralElongationView = TDetailLateralElongationView<float>;
+using FDetailLateralElongationConstView = TDetailLateralElongationView<const float>;
+PV_DEFINE_STRUCTURED_ATTRIBUTE(DetailLateralElongation, TArray<float>, FDetailLateralElongationView, FDetailLateralElongationConstView, AttributeNames::LateralElongationStr, GroupNames::DetailsGroupStr);
 PV_DEFINE_ATTRIBUTE(DetailTrunkURange, TArray<FVector2f>, AttributeNames::TrunkURangeStr, GroupNames::DetailsGroupStr);
 PV_DEFINE_ATTRIBUTE(DetailTrunkMaterialPath, FString, AttributeNames::TrunkMaterialPathStr, GroupNames::DetailsGroupStr);
 
