Fix crash that can occur when changing the DebugState on the Texture2D/StaticMesh importer nodes if  they are connected to another node

- Move debug-only parameters (bUseSplitPoints and DebugState) out of the main import params structs and into separate DebugParams structs for both StaticMesh and Texture2D importers
- Update importer functions to accept the new debug params as a separate argument, and expose the debug params in node settings only when the node is in debug mode
- Don't pass the DebugCollection unless PVE DebugMode is enabled
- Remove "Procedural Vegetation" prefix from node titles
- Fix visualization of root point not showing up correctly

#Synced-CL 51755814

diff --git a/Engine/Plugins/Experimental/ProceduralVegetationEditor/Source/ProceduralVegetation/Private/Helpers/PVPlantTraversalHelper.h b/Engine/Plugins/Experimental/ProceduralVegetationEditor/Source/ProceduralVegetation/Private/Helpers/PVPlantTraversalHelper.h
--- a/Engine/Plugins/Experimental/ProceduralVegetationEditor/Source/ProceduralVegetation/Private/Helpers/PVPlantTraversalHelper.h
+++ b/Engine/Plugins/Experimental/ProceduralVegetationEditor/Source/ProceduralVegetation/Private/Helpers/PVPlantTraversalHelper.h
@@ -102,6 +102,11 @@ namespace PV::PlantTraversalHelper
 		const TFunction<EForEachResult(int32, StateType&)>& Body
 	)
 	{
+		if (!ValidateAttributeCollection(BranchParentNumberAttribute, BranchChildrenAttribute, BranchNumberAttribute))
+		{
+			return;
+		}
+
 		const TFunction<EForEachResult(int32, const StateType&)> WalkBranches_Recursive = [&](int32 BranchIndex, const StateType& State) -> EForEachResult
 		{
 			StateType BranchState = State;
@@ -204,6 +209,16 @@ namespace PV::PlantTraversalHelper
 		int32 BranchIndex,
 		int32 BranchPointIndex)
 	{
+		if (!ValidateAttributeCollection(BranchPointsAttribute, PointLengthFromRootAttribute, BranchFoliageIDsAttribute, FoliageLengthFromRootAttribute))
+		{
+			return TArray<int32>();
+		}
+
+		if (!BranchPointsAttribute.IsValidIndex(BranchIndex) || !BranchFoliageIDsAttribute.IsValidIndex(BranchIndex))
+		{
+			return TArray<int32>();
+		}
+
 		const TArray<int32>& BranchPoints = BranchPointsAttribute[BranchIndex];
 		if (!BranchPoints.IsValidIndex(BranchPointIndex))
 		{
@@ -250,6 +265,11 @@ namespace PV::PlantTraversalHelper
 		int32 BranchPointIndex
 	)
 	{
+		if (!ValidateAttributeCollection(BranchChildrenAttribute, BranchParentNumberAttribute, BranchNumberAttribute, BranchPointsAttribute, PointLengthFromRootAttribute))
+		{
+			return TArray<int32>();
+		}
+
 		if (!BranchPointsAttribute.IsValidIndex(BranchIndex))
 		{
 			return TArray<int32>();
@@ -320,6 +340,11 @@ namespace PV::PlantTraversalHelper
 		int32 BranchPointIndex
 	)
 	{
+		if (!ValidateAttributeCollection(BranchChildrenAttribute, BranchParentNumberAttribute, BranchNumberAttribute, BranchPointsAttribute))
+		{
+			return TArray<int32>();
+		}
+
 		if (!BranchPointsAttribute.IsValidIndex(BranchIndex))
 		{
 			return TArray<int32>();
@@ -424,6 +449,11 @@ namespace PV::PlantTraversalHelper
 		const TFunction<EForEachResult(const FRecursiveWalkBranchPointsParams& Params, StateType& PointState)>& Body
 	)
 	{
+		if (!ValidateAttributeCollection(BranchChildrenAttribute, BranchParentNumberAttribute, BranchNumberAttribute, BranchPointsAttribute, PointLengthFromRootAttribute))
+		{
+			return EForEachResult::Break;
+		}
+
 		const TFunction<EForEachResult(int32, int32, int32, const StateType&)> WalkPoints = [&](int32 BranchIndex, int32 BranchPointIndex, int32 ParentPointIndex, const StateType& ParentState) -> EForEachResult
 		{
 			StateType PointState = ParentState;
@@ -512,6 +542,11 @@ namespace PV::PlantTraversalHelper
 		const TFunction<EForEachResult(const FRecursiveWalkBranchPointsParams& Params, StateType& PointState)>& Body
 	)
 	{
+		if (!ValidateAttributeCollection(BranchChildrenAttribute, BranchParentNumberAttribute, BranchNumberAttribute, BranchPointsAttribute))
+		{
+			return EForEachResult::Break;
+		}
+
 		const TFunction<EForEachResult(int32, int32, int32, const StateType&)> WalkPoints = [&](int32 BranchIndex, int32 BranchPointIndex, int32 ParentPointIndex, const StateType& ParentState) -> EForEachResult
 		{
 			StateType PointState = ParentState;
@@ -754,6 +789,11 @@ namespace PV::PlantTraversalHelper
 		const TFunction<EForEachResult(int32 BranchIndex, int32 BranchPointIndex, int32 PointIndex)>& Body
 	)
 	{
+		if (!ValidateAttributeCollection(BranchPointsAttribute, BranchParentNumberAttribute, BranchNumberAttribute))
+		{
+			return;
+		}
+
 		const bool bHasFusedPoints = HasFusedPoints(BranchPointsAttribute, BranchParentNumberAttribute, BranchNumberAttribute);
 		for (int32 BranchIndex = 0; BranchIndex < BranchPointsAttribute.Num(); ++BranchIndex)
 		{
@@ -799,18 +839,41 @@ namespace PV::PlantTraversalHelper
 	inline void ForEachPlantPoint(
 		PV::FBranchPlantNumberAttributeConstView BranchPlantNumberAttribute,
 		PV::FBranchPointsAttributeConstView BranchPointsAttribute,
+		PV::FBranchParentNumberAttributeConstView BranchParentNumberAttribute,
+		PV::FBranchNumberAttributeConstView BranchNumberAttribute,
 		int32 PlantNumber,
 		const TFunction<EForEachResult(int32 BranchIndex, int32 PointIndex)>& Body
 	)
 	{
+		if (!ValidateAttributeCollection(BranchPlantNumberAttribute, BranchPointsAttribute, BranchParentNumberAttribute, BranchNumberAttribute))
+		{
+			return;
+		}
+
+		const bool bFusedPoints = HasFusedPoints(BranchPointsAttribute, BranchParentNumberAttribute, BranchNumberAttribute);
+
 		const TArray<int32> BranchIndices = GetPlantBranchIndices(BranchPlantNumberAttribute, PlantNumber);
 		for (int32 BranchIndex : BranchIndices)
 		{
-			for (int32 PointIndex : BranchPointsAttribute[BranchIndex])
+			if (bFusedPoints)
 			{
-				if (Body(BranchIndex, PointIndex) == EForEachResult::Break)
+				const bool bIsTrunk = IsTrunk(BranchParentNumberAttribute, BranchIndex);
+				for (int32 i = bIsTrunk ? 0 : 1; i < BranchPointsAttribute[BranchIndex].Num(); ++i)
 				{
-					return;
+					if (Body(BranchIndex, BranchPointsAttribute[BranchIndex][i]) == EForEachResult::Break)
+					{
+						return;
+					}
+				}
+			}
+			else
+			{
+				for (int32 PointIndex : BranchPointsAttribute[BranchIndex])
+				{
+					if (Body(BranchIndex, PointIndex) == EForEachResult::Break)
+					{
+						return;
+					}
 				}
 			}
 		}
@@ -824,10 +887,8 @@ namespace PV::PlantTraversalHelper
 	{
 		const PV::FBranchPlantNumberAttributeConstView BranchPlantNumberAttribute = FBranchPlantNumberAttribute::FindAttribute(Collection);
 		const PV::FBranchPointsAttributeConstView BranchPointsAttribute = FBranchPointsAttribute::FindAttribute(Collection);
-
-		if (BranchPlantNumberAttribute && BranchPointsAttribute)
-		{
-			ForEachPlantPoint(BranchPlantNumberAttribute, BranchPointsAttribute, PlantNumber, Body);
-		}
+		const PV::FBranchParentNumberAttributeConstView BranchParentNumberAttribute = FBranchParentNumberAttribute::FindAttribute(Collection);
+		const PV::FBranchNumberAttributeConstView BranchNumberAttribute = FBranchNumberAttribute::FindAttribute(Collection);
+		ForEachPlantPoint(BranchPlantNumberAttribute, BranchPointsAttribute, BranchParentNumberAttribute, BranchNumberAttribute, PlantNumber, Body);
 	}
 }
\ No newline at end of file
diff --git a/Engine/Plugins/Experimental/ProceduralVegetationEditor/Source/ProceduralVegetation/Private/Implementations/PVImporter_StaticMesh.cpp b/Engine/Plugins/Experimental/ProceduralVegetationEditor/Source/ProceduralVegetation/Private/Implementations/PVImporter_StaticMesh.cpp
--- a/Engine/Plugins/Experimental/ProceduralVegetationEditor/Source/ProceduralVegetation/Private/Implementations/PVImporter_StaticMesh.cpp
+++ b/Engine/Plugins/Experimental/ProceduralVegetationEditor/Source/ProceduralVegetation/Private/Implementations/PVImporter_StaticMesh.cpp
@@ -1001,7 +1001,7 @@ static void FindRootPoints(
 
 }
 
-PV::StaticMeshImport::EImportResult PV::StaticMeshImport::ImportGrowthDataFromStaticMesh(const FPVImportStaticMeshParams& InParams, FPVImportStaticMeshOutput& Output)
+PV::StaticMeshImport::EImportResult PV::StaticMeshImport::ImportGrowthDataFromStaticMesh(const FPVImportStaticMeshParams& InParams, const FPVImportStaticMeshDebugParams& InDebugParams, FPVImportStaticMeshOutput& Output)
 {
 	using namespace PV::ImportHelper;
 	
@@ -1034,7 +1034,7 @@ PV::StaticMeshImport::EImportResult PV::StaticMeshImport::ImportGrowthDataFromSt
 		}
 	}
 
-	if (InParams.DebugState == EPVImportStaticMeshDebugState::InitialMesh)
+	if (InDebugParams.DebugState == EPVImportStaticMeshDebugState::InitialMesh)
 	{
 		FManagedArrayCollection& DebugCollection = Output.DebugCollection.Emplace();
 		AddDynamicMeshToCollection(DynamicMesh, DebugCollection);
@@ -1049,7 +1049,7 @@ PV::StaticMeshImport::EImportResult PV::StaticMeshImport::ImportGrowthDataFromSt
 	AddDynamicMeshToCollection(DynamicMesh, Output.VisualizationCollection, FTransform(VisualizationOffset));
 
 	ComputeNormals(DynamicMesh);
-	if (InParams.DebugState == EPVImportStaticMeshDebugState::ComputeNormals)
+	if (InDebugParams.DebugState == EPVImportStaticMeshDebugState::ComputeNormals)
 	{
 		FManagedArrayCollection& DebugCollection = Output.DebugCollection.Emplace();
 		AddDynamicMeshToCollection(DynamicMesh, DebugCollection);
@@ -1059,7 +1059,7 @@ PV::StaticMeshImport::EImportResult PV::StaticMeshImport::ImportGrowthDataFromSt
 		DynamicMesh,
 		VertexScaleAttributeDefinition.GetOrAttachAttribute(DynamicMesh)
 	);
-	if (InParams.DebugState == EPVImportStaticMeshDebugState::ComputeVerticesCenterAndRadius)
+	if (InDebugParams.DebugState == EPVImportStaticMeshDebugState::ComputeVerticesCenterAndRadius)
 	{
 		FManagedArrayCollection& DebugCollection = Output.DebugCollection.Emplace();
 		AddDynamicMeshToCollection(DynamicMesh, DebugCollection);
@@ -1072,7 +1072,7 @@ PV::StaticMeshImport::EImportResult PV::StaticMeshImport::ImportGrowthDataFromSt
 		FuseToAttributeDefinition.GetOrAttachAttribute(DynamicMesh),
 		FillerVertexAttributeDefinition.GetOrAttachAttribute(DynamicMesh)
 	);
-	if (InParams.DebugState == EPVImportStaticMeshDebugState::CollapseBranchesToLine)
+	if (InDebugParams.DebugState == EPVImportStaticMeshDebugState::CollapseBranchesToLine)
 	{
 		FManagedArrayCollection& DebugCollection = Output.DebugCollection.Emplace();
 		AddDynamicMeshToCollection(DynamicMesh, DebugCollection);
@@ -1085,7 +1085,7 @@ PV::StaticMeshImport::EImportResult PV::StaticMeshImport::ImportGrowthDataFromSt
 		EndPointAttributeDefinition.GetOrAttachAttribute(DynamicMesh),
 		FillerVertexAttributeDefinition.GetOrAttachAttribute(DynamicMesh)
 	);
-	if (InParams.DebugState == EPVImportStaticMeshDebugState::FindEndPoints)
+	if (InDebugParams.DebugState == EPVImportStaticMeshDebugState::FindEndPoints)
 	{
 		FManagedArrayCollection& DebugCollection = Output.DebugCollection.Emplace();
 		AddDynamicMeshToCollection(DynamicMesh, DebugCollection);
@@ -1097,7 +1097,7 @@ PV::StaticMeshImport::EImportResult PV::StaticMeshImport::ImportGrowthDataFromSt
 		RootPointAttributeDefinition.GetOrAttachAttribute(DynamicMesh),
 		EndPointAttributeDefinition.GetOrAttachAttribute(DynamicMesh)
 	);
-	if (InParams.DebugState == EPVImportStaticMeshDebugState::FindRootPoints)
+	if (InDebugParams.DebugState == EPVImportStaticMeshDebugState::FindRootPoints)
 	{
 		FManagedArrayCollection& DebugCollection = Output.DebugCollection.Emplace();
 		AddDynamicMeshToCollection(DynamicMesh, DebugCollection);
@@ -1112,7 +1112,7 @@ PV::StaticMeshImport::EImportResult PV::StaticMeshImport::ImportGrowthDataFromSt
 		EndPointAttributeDefinition.GetOrAttachAttribute(DynamicMesh),
 		FillerVertexAttributeDefinition.GetOrAttachAttribute(DynamicMesh)
 	);
-	if (InParams.DebugState == EPVImportStaticMeshDebugState::ConnectBranchesToParent)
+	if (InDebugParams.DebugState == EPVImportStaticMeshDebugState::ConnectBranchesToParent)
 	{
 		FManagedArrayCollection& DebugCollection = Output.DebugCollection.Emplace();
 		AddDynamicMeshToCollection(DynamicMesh, DebugCollection);
@@ -1141,7 +1141,7 @@ PV::StaticMeshImport::EImportResult PV::StaticMeshImport::ImportGrowthDataFromSt
 		BranchHierarchyDescriptions
 	);
 
-	if (InParams.DebugState == EPVImportStaticMeshDebugState::ComputeBranchHierarchies)
+	if (InDebugParams.DebugState == EPVImportStaticMeshDebugState::ComputeBranchHierarchies)
 	{
 		FManagedArrayCollection& DebugCollection = Output.DebugCollection.Emplace();
 		AddDynamicMeshToCollection(DynamicMesh, DebugCollection);
@@ -1153,7 +1153,7 @@ PV::StaticMeshImport::EImportResult PV::StaticMeshImport::ImportGrowthDataFromSt
 		NextBranchPointAttributeDefinition.CopyAttributeToCollection(DynamicMesh, DebugCollection);
 	}
 
-	if (!PV::ImportHelper::GenerateGrowthDataFromBranchHierarchies(Output.GrowthDataCollection, BranchHierarchyDescriptions, InParams.bUseSplitPoints))
+	if (!PV::ImportHelper::GenerateGrowthDataFromBranchHierarchies(Output.GrowthDataCollection, BranchHierarchyDescriptions, InDebugParams.bUseSplitPoints))
 	{
 		return EImportResult::FailedToGenerateGrowthData;
 	}
diff --git a/Engine/Plugins/Experimental/ProceduralVegetationEditor/Source/ProceduralVegetation/Private/Implementations/PVImporter_StaticMesh.h b/Engine/Plugins/Experimental/ProceduralVegetationEditor/Source/ProceduralVegetation/Private/Implementations/PVImporter_StaticMesh.h
--- a/Engine/Plugins/Experimental/ProceduralVegetationEditor/Source/ProceduralVegetation/Private/Implementations/PVImporter_StaticMesh.h
+++ b/Engine/Plugins/Experimental/ProceduralVegetationEditor/Source/ProceduralVegetation/Private/Implementations/PVImporter_StaticMesh.h
@@ -21,5 +21,5 @@ namespace PV::StaticMeshImport
 		FailedToGenerateGrowthData
 	};
 
-	EImportResult ImportGrowthDataFromStaticMesh(const FPVImportStaticMeshParams& InParams, FPVImportStaticMeshOutput& Output);
+	EImportResult ImportGrowthDataFromStaticMesh(const FPVImportStaticMeshParams& InParams, const FPVImportStaticMeshDebugParams& InDebugParams, FPVImportStaticMeshOutput& Output);
 };
diff --git a/Engine/Plugins/Experimental/ProceduralVegetationEditor/Source/ProceduralVegetation/Private/Implementations/PVImporter_Texture2D.cpp b/Engine/Plugins/Experimental/ProceduralVegetationEditor/Source/ProceduralVegetation/Private/Implementations/PVImporter_Texture2D.cpp
--- a/Engine/Plugins/Experimental/ProceduralVegetationEditor/Source/ProceduralVegetation/Private/Implementations/PVImporter_Texture2D.cpp
+++ b/Engine/Plugins/Experimental/ProceduralVegetationEditor/Source/ProceduralVegetation/Private/Implementations/PVImporter_Texture2D.cpp
@@ -937,7 +937,7 @@ namespace Internal
 	}
 };
 
-PV::Texture2DImport::EImportResult PV::Texture2DImport::ImportGrowthDataFromTexture2D(const FPVImportTexture2DParams& InParams, const UMaterialInterface* TextureVisualizationMaterial, FPVImportTexture2DOutput& Output)
+PV::Texture2DImport::EImportResult PV::Texture2DImport::ImportGrowthDataFromTexture2D(const FPVImportTexture2DParams& InParams, const FPVImportTexture2DDebugParams& InDebugParams, const UMaterialInterface* TextureVisualizationMaterial, FPVImportTexture2DOutput& Output)
 {
 	using namespace PV::ImportHelper;
 	using namespace Internal;
@@ -970,7 +970,7 @@ PV::Texture2DImport::EImportResult PV::Texture2DImport::ImportGrowthDataFromText
 		return EImportResult::InvalidTextureSize;
 	}
 	
-	if (InParams.DebugState == EPVImportTexture2DDebugState::SampleImage)
+	if (InDebugParams.DebugState == EPVImportTexture2DDebugState::SampleImage)
 	{
 		FManagedArrayCollection& DebugCollection = Output.DebugCollection.Emplace();
 		AddDebugPixelValuesToCollection(DebugCollection, BlackAndWhiteImageData, SampleSize, InParams.Scale);
@@ -986,7 +986,7 @@ PV::Texture2DImport::EImportResult PV::Texture2DImport::ImportGrowthDataFromText
 	UE::Geometry::FDynamicMesh3 DynamicMesh;
 	TriangulateCurvesAndSplitEdges(PerimeterCurves, DynamicMesh, InParams.Scale);
 
-	if (InParams.DebugState == EPVImportTexture2DDebugState::TriangulateCurves)
+	if (InDebugParams.DebugState == EPVImportTexture2DDebugState::TriangulateCurves)
 	{
 		FManagedArrayCollection& DebugCollection = Output.DebugCollection.Emplace();
 		AddDynamicMeshToCollection(DynamicMesh, DebugCollection);
@@ -999,7 +999,7 @@ PV::Texture2DImport::EImportResult PV::Texture2DImport::ImportGrowthDataFromText
 		VertexScaleAttributeDefinition.GetOrAttachAttribute(DynamicMesh)
 	);
 
-	if (InParams.DebugState == EPVImportTexture2DDebugState::ComputeVertexScale)
+	if (InDebugParams.DebugState == EPVImportTexture2DDebugState::ComputeVertexScale)
 	{
 		FManagedArrayCollection& DebugCollection = Output.DebugCollection.Emplace();
 		AddDynamicMeshToCollection(DynamicMesh, DebugCollection);
@@ -1068,7 +1068,7 @@ PV::Texture2DImport::EImportResult PV::Texture2DImport::ImportGrowthDataFromText
 
 	check(DynamicMeshBranchHierarchies.Num() == BranchHierarchyDescriptions.Num());
 
-	if (InParams.DebugState == EPVImportTexture2DDebugState::ComputeBranchHierarchies)
+	if (InDebugParams.DebugState == EPVImportTexture2DDebugState::ComputeBranchHierarchies)
 	{
 		FManagedArrayCollection& DebugCollection = Output.DebugCollection.Emplace();
 		AddDynamicMeshToCollection(DynamicMesh, DebugCollection);
@@ -1084,7 +1084,7 @@ PV::Texture2DImport::EImportResult PV::Texture2DImport::ImportGrowthDataFromText
 	MapPlantIndicesToBranchHierarchies(PlantMapping, DynamicMeshBranchHierarchies, BranchHierarchiesToPlantIndices);
 
 	// TODO: It would probably be faster to combine all the outputs of the PlantCollections
-	if (!PV::ImportHelper::GenerateGrowthDataFromBranchHierarchies(Output.CombinedPlantCollection, BranchHierarchyDescriptions, InParams.bUseSplitPoints))
+	if (!PV::ImportHelper::GenerateGrowthDataFromBranchHierarchies(Output.CombinedPlantCollection, BranchHierarchyDescriptions, InDebugParams.bUseSplitPoints))
 	{
 		return EImportResult::FailedToGenerateGrowthData;
 	}
@@ -1098,7 +1098,7 @@ PV::Texture2DImport::EImportResult PV::Texture2DImport::ImportGrowthDataFromText
 			continue;
 		}
 
-		if (!PV::ImportHelper::GenerateGrowthDataFromBranchHierarchy(Output.PlantCollections[PlantIndex], BranchHierarchyDescriptions[i], InParams.bUseSplitPoints))
+		if (!PV::ImportHelper::GenerateGrowthDataFromBranchHierarchy(Output.PlantCollections[PlantIndex], BranchHierarchyDescriptions[i], InDebugParams.bUseSplitPoints))
 		{
 			return EImportResult::FailedToGenerateGrowthData;
 		}
diff --git a/Engine/Plugins/Experimental/ProceduralVegetationEditor/Source/ProceduralVegetation/Private/Implementations/PVImporter_Texture2D.h b/Engine/Plugins/Experimental/ProceduralVegetationEditor/Source/ProceduralVegetation/Private/Implementations/PVImporter_Texture2D.h
--- a/Engine/Plugins/Experimental/ProceduralVegetationEditor/Source/ProceduralVegetation/Private/Implementations/PVImporter_Texture2D.h
+++ b/Engine/Plugins/Experimental/ProceduralVegetationEditor/Source/ProceduralVegetation/Private/Implementations/PVImporter_Texture2D.h
@@ -19,5 +19,5 @@ namespace PV::Texture2DImport
 		FailedToGenerateGrowthData
 	};
 
-	EImportResult ImportGrowthDataFromTexture2D(const FPVImportTexture2DParams& InParams, const UMaterialInterface* TextureVisualizationMaterial, FPVImportTexture2DOutput& Output);
+	EImportResult ImportGrowthDataFromTexture2D(const FPVImportTexture2DParams& InParams, const FPVImportTexture2DDebugParams& InDebugParams, const UMaterialInterface* TextureVisualizationMaterial, FPVImportTexture2DOutput& Output);
 };
diff --git a/Engine/Plugins/Experimental/ProceduralVegetationEditor/Source/ProceduralVegetation/Private/Nodes/PVStaticMeshImporterSettings.cpp b/Engine/Plugins/Experimental/ProceduralVegetationEditor/Source/ProceduralVegetation/Private/Nodes/PVStaticMeshImporterSettings.cpp
--- a/Engine/Plugins/Experimental/ProceduralVegetationEditor/Source/ProceduralVegetation/Private/Nodes/PVStaticMeshImporterSettings.cpp
+++ b/Engine/Plugins/Experimental/ProceduralVegetationEditor/Source/ProceduralVegetation/Private/Nodes/PVStaticMeshImporterSettings.cpp
@@ -16,7 +16,7 @@
 #if WITH_EDITOR
 FText UPVStaticMeshImporterSettings::GetDefaultNodeTitle() const
 { 
-	return LOCTEXT("NodeTitle", "Procedural Vegetation Static Mesh Importer"); 
+	return LOCTEXT("NodeTitle", "Static Mesh Importer"); 
 }
 
 FText UPVStaticMeshImporterSettings::GetNodeTooltipText() const
@@ -68,7 +68,7 @@ bool FPVStaticMeshImporterElement::ExecuteInternal(FPCGContext* InContext) const
 	}
 	
 	FPVImportStaticMeshOutput Output;
-	const auto Result = PV::StaticMeshImport::ImportGrowthDataFromStaticMesh(Settings->Params, Output);
+	const auto Result = PV::StaticMeshImport::ImportGrowthDataFromStaticMesh(Settings->Params, Settings->DebugParams, Output);
 	if (Result != PV::StaticMeshImport::EImportResult::Success)
 	{
 		return Result == PV::StaticMeshImport::EImportResult::MeshNotReady ? false : true;
@@ -76,7 +76,7 @@ bool FPVStaticMeshImporterElement::ExecuteInternal(FPCGContext* InContext) const
 
 	UPVGrowthData* OutVariantData = FPCGContext::NewObject_AnyThread<UPVGrowthData>(InContext);
 
-	if (Output.DebugCollection.IsSet())
+	if (Output.DebugCollection.IsSet() && PV::Utilities::DebugModeEnabled())
 	{
 		OutVariantData->Initialize(MoveTemp(Output.DebugCollection.GetValue()));
 	}
diff --git a/Engine/Plugins/Experimental/ProceduralVegetationEditor/Source/ProceduralVegetation/Private/Nodes/PVStaticMeshImporterSettings.h b/Engine/Plugins/Experimental/ProceduralVegetationEditor/Source/ProceduralVegetation/Private/Nodes/PVStaticMeshImporterSettings.h
--- a/Engine/Plugins/Experimental/ProceduralVegetationEditor/Source/ProceduralVegetation/Private/Nodes/PVStaticMeshImporterSettings.h
+++ b/Engine/Plugins/Experimental/ProceduralVegetationEditor/Source/ProceduralVegetation/Private/Nodes/PVStaticMeshImporterSettings.h
@@ -32,6 +32,8 @@ public:
 	UPROPERTY(EditAnywhere, Category="Importer Settings", meta = (ShowOnlyInnerProperties))
 	FPVImportStaticMeshParams Params;
 
+	UPROPERTY(EditAnywhere, Transient, Category = DebugSettings, meta = (EditCondition = "bDebug", EditConditionHides))
+	FPVImportStaticMeshDebugParams DebugParams;
 };
 
 class FPVStaticMeshImporterElement : public FPVBaseElement
diff --git a/Engine/Plugins/Experimental/ProceduralVegetationEditor/Source/ProceduralVegetation/Private/Nodes/PVTexture2DImporterSettings.cpp b/Engine/Plugins/Experimental/ProceduralVegetationEditor/Source/ProceduralVegetation/Private/Nodes/PVTexture2DImporterSettings.cpp
--- a/Engine/Plugins/Experimental/ProceduralVegetationEditor/Source/ProceduralVegetation/Private/Nodes/PVTexture2DImporterSettings.cpp
+++ b/Engine/Plugins/Experimental/ProceduralVegetationEditor/Source/ProceduralVegetation/Private/Nodes/PVTexture2DImporterSettings.cpp
@@ -56,7 +56,7 @@ UPVTexture2DImporterSettings::UPVTexture2DImporterSettings()
 #if WITH_EDITOR
 FText UPVTexture2DImporterSettings::GetDefaultNodeTitle() const
 { 
-	return LOCTEXT("NodeTitle", "Procedural Vegetation Texture 2D Importer"); 
+	return LOCTEXT("NodeTitle", "Texture 2D Importer"); 
 }
 
 FText UPVTexture2DImporterSettings::GetNodeTooltipText() const
@@ -142,8 +142,8 @@ bool FPVTexture2DImporterElement::ExecuteInternal(FPCGContext* InContext) const
 		UserRootPositionsVisualization.DebugType = EPVDebugType::Custom;
 		UserRootPositionsVisualization.bShowAnchorPoints = false;
 		UserRootPositionsVisualization.bDrawPointAsMesh = false;
-		UserRootPositionsVisualization.AttributeToFilter = PVImportNames::UserRootPositionsGroup;
-		UserRootPositionsVisualization.CustomGroupToFilter = PVImportNames::UserRootPositionsAttribute;
+		UserRootPositionsVisualization.AttributeToFilter = PVImportNames::UserRootPositionsAttribute;
+		UserRootPositionsVisualization.CustomGroupToFilter = PVImportNames::UserRootPositionsGroup;
 		UserRootPositionsVisualization.VisualizationMode = EPVDebugValueVisualizationMode::Point;
 		UserRootPositionsVisualization.Color = FColor::Blue;
 		UserRootPositionsVisualization.GizmoScale = 10;
@@ -183,8 +183,9 @@ bool FPVTexture2DImporterElement::ExecuteInternal(FPCGContext* InContext) const
 
 	FPVImportTexture2DOutput ImportTexture2DOutput;
 	const auto ImportResult = PV::Texture2DImport::ImportGrowthDataFromTexture2D(
-		Settings->Params, 
-		TextureVisualizationMaterialInstance, 
+		Settings->Params,
+		Settings->DebugParams,
+		TextureVisualizationMaterialInstance,
 		ImportTexture2DOutput
 	);
 	if (ImportResult != PV::Texture2DImport::EImportResult::Success)
@@ -201,7 +202,7 @@ bool FPVTexture2DImporterElement::ExecuteInternal(FPCGContext* InContext) const
 		OutVariantData->AddVisualizationCollection(ImportTexture2DOutput.VisualizationCollection, VisualizationSettings);
 #endif
 
-		if (ImportTexture2DOutput.DebugCollection.IsSet())
+		if (ImportTexture2DOutput.DebugCollection.IsSet() && PV::Utilities::DebugModeEnabled())
 		{
 			OutVariantData->Initialize(MoveTemp(ImportTexture2DOutput.DebugCollection.GetValue()));
 		}
diff --git a/Engine/Plugins/Experimental/ProceduralVegetationEditor/Source/ProceduralVegetation/Private/Nodes/PVTexture2DImporterSettings.h b/Engine/Plugins/Experimental/ProceduralVegetationEditor/Source/ProceduralVegetation/Private/Nodes/PVTexture2DImporterSettings.h
--- a/Engine/Plugins/Experimental/ProceduralVegetationEditor/Source/ProceduralVegetation/Private/Nodes/PVTexture2DImporterSettings.h
+++ b/Engine/Plugins/Experimental/ProceduralVegetationEditor/Source/ProceduralVegetation/Private/Nodes/PVTexture2DImporterSettings.h
@@ -37,6 +37,9 @@ public:
 	UPROPERTY(EditAnywhere, Category="Importer Settings", meta = (ShowOnlyInnerProperties))
 	FPVImportTexture2DParams Params;
 
+	UPROPERTY(EditAnywhere, Transient, Category = DebugSettings, meta = (EditCondition = "bDebug", EditConditionHides))
+	FPVImportTexture2DDebugParams DebugParams;
+
 #if WITH_EDITORONLY_DATA
 	UPROPERTY(EditAnywhere, Category = DebugSettings, meta = (EditCondition = "bDebug", EditConditionHides))
 	TObjectPtr<UMaterialInterface> TextureVisualizationMaterial;
diff --git a/Engine/Plugins/Experimental/ProceduralVegetationEditor/Source/ProceduralVegetation/Private/Params/PVImportStaticMeshParams.h b/Engine/Plugins/Experimental/ProceduralVegetationEditor/Source/ProceduralVegetation/Private/Params/PVImportStaticMeshParams.h
--- a/Engine/Plugins/Experimental/ProceduralVegetationEditor/Source/ProceduralVegetation/Private/Params/PVImportStaticMeshParams.h
+++ b/Engine/Plugins/Experimental/ProceduralVegetationEditor/Source/ProceduralVegetation/Private/Params/PVImportStaticMeshParams.h
@@ -38,6 +38,12 @@ struct FPVImportStaticMeshParams
 
 	UPROPERTY(EditAnywhere, Category = "Visualization")
 	bool bShowFoundBranchCurves = true;
+};
+
+USTRUCT(BlueprintType)
+struct FPVImportStaticMeshDebugParams
+{
+	GENERATED_BODY()
 
 	UPROPERTY(EditAnywhere, Category = "DebugSettings")
 	bool bUseSplitPoints = false;
diff --git a/Engine/Plugins/Experimental/ProceduralVegetationEditor/Source/ProceduralVegetation/Private/Params/PVImportTexture2DParams.h b/Engine/Plugins/Experimental/ProceduralVegetationEditor/Source/ProceduralVegetation/Private/Params/PVImportTexture2DParams.h
--- a/Engine/Plugins/Experimental/ProceduralVegetationEditor/Source/ProceduralVegetation/Private/Params/PVImportTexture2DParams.h
+++ b/Engine/Plugins/Experimental/ProceduralVegetationEditor/Source/ProceduralVegetation/Private/Params/PVImportTexture2DParams.h
@@ -72,6 +72,12 @@ struct FPVImportTexture2DParams
 
 	UPROPERTY(EditAnywhere, Category = "Plant Geneneration|Tip detection", meta = (ClampMin=0, ClampMax=0.01, UIMin=0, UIMax=0.01))
 	float MaxTipAngleSearchDist = 0.005;
+};
+
+USTRUCT(BlueprintType)
+struct FPVImportTexture2DDebugParams
+{
+	GENERATED_BODY()
 
 	UPROPERTY(EditAnywhere, Category = "DebugSettings")
 	bool bUseSplitPoints = false;
