[PVE] Grafter bug fixes and general improvements
- Fixes issue where parents were not being computed correctly during grafting
- Fixes issue where points added to graft were not being stored in a sorted order on branches
- Fixes issue where bud lateral meristem attribute was not being set for grafts
- UX improvement where grafter pallet now at least has one input pin upon initialization

#Synced-CL 51754256

diff --git a/Engine/Plugins/Experimental/ProceduralVegetationEditor/Source/ProceduralVegetation/Private/Facades/PVPointFacade.cpp b/Engine/Plugins/Experimental/ProceduralVegetationEditor/Source/ProceduralVegetation/Private/Facades/PVPointFacade.cpp
--- a/Engine/Plugins/Experimental/ProceduralVegetationEditor/Source/ProceduralVegetation/Private/Facades/PVPointFacade.cpp
+++ b/Engine/Plugins/Experimental/ProceduralVegetationEditor/Source/ProceduralVegetation/Private/Facades/PVPointFacade.cpp
@@ -823,6 +823,11 @@ namespace PV::Facades
 		return BudLightDetected.Get();
 	}
 
+	const TManagedArray<TArray<float>>& FPointFacade::GetBudLateralMeristemArrays() const
+	{
+		return BudLateralMeristem.Get();
+	}
+
 	void FPointFacade::SetBudLightDetectedArraysFromIndex(const TArray<TArray<float>>& InputArray, const int32 StartIndex)
 	{
 		FILL_COLLECTION_ATTRIBUTE_FROM_INDEX(BudLightDetected, InputArray, StartIndex);
diff --git a/Engine/Plugins/Experimental/ProceduralVegetationEditor/Source/ProceduralVegetation/Private/Implementations/PVGrafter.cpp b/Engine/Plugins/Experimental/ProceduralVegetationEditor/Source/ProceduralVegetation/Private/Implementations/PVGrafter.cpp
--- a/Engine/Plugins/Experimental/ProceduralVegetationEditor/Source/ProceduralVegetation/Private/Implementations/PVGrafter.cpp
+++ b/Engine/Plugins/Experimental/ProceduralVegetationEditor/Source/ProceduralVegetation/Private/Implementations/PVGrafter.cpp
@@ -251,6 +251,14 @@ void FPVGrafter::SetGraft(FManagedArrayCollection& OutCollection,
 			});
 
 	PointFacadeOut.SetPointScalesFromIndex(TransformedPointScales, PointFacadeOutStartIndex);
+	
+	const TManagedArray<TArray<float>>& BudLateralMeristemArrays = PointFacadeSourceGraft.GetBudLateralMeristemArrays();
+
+	ForEachGraftPointExcludingRoot(NumOfPointsInSourceGraft, SourceGraftTrunkRootPointIndex,
+		[&](const int32 SourcePointIndex, const int32 TargetPointIndex)
+			{
+				PointFacadeOut.SetBudLateralMeristem(TargetPointIndex + PointFacadeOutStartIndex, BudLateralMeristemArrays[SourcePointIndex]);
+			});
 
 	const TManagedArray<float>& LFRs = PointFacadeSourceGraft.GetLengthFromRootsArray();
 	TArray<float> TransformedLFRs;
@@ -541,6 +549,8 @@ void FPVGrafter::SetGraft(FManagedArrayCollection& OutCollection,
 	}
 
 
+	TArray<int32> HostAncestorPrefix = BranchFacadeOut.GetParents(InAttributes.BranchIndex);
+	HostAncestorPrefix.Add(BranchFacadeOut.GetBranchNumber(InAttributes.BranchIndex));
 	SourceIndex = 0;
 	TargetIndex = NumOfBranchesInTarget;
 	while (SourceIndex < NumOfBranchesInSource)
@@ -550,14 +560,14 @@ void FPVGrafter::SetGraft(FManagedArrayCollection& OutCollection,
 		{
 			const TArray<int32>& SourceBranchParents = BranchFacadeSourceGraft.GetParents(SourceIndex);
 			const int32 NumOfParentsInSource = SourceBranchParents.Num();
-			TArray<int32> TransformedParents;
-			TransformedParents.SetNum(NumOfParentsInSource);
+			TArray<int32> TransformedParents = HostAncestorPrefix;
 
 			for (int32 j = 0; j < NumOfParentsInSource; ++j)
 			{
-				if (SourceOldBranchNumbersToNewBranchNumbers.Contains(SourceBranchParents[j]))
+				if (SourceBranchParents[j] != 0 
+					&& SourceOldBranchNumbersToNewBranchNumbers.Contains(SourceBranchParents[j]))
 				{
-					TransformedParents[j] = SourceOldBranchNumbersToNewBranchNumbers[SourceBranchParents[j]];
+					TransformedParents.Add(SourceOldBranchNumbersToNewBranchNumbers[SourceBranchParents[j]]);
 				}
 			}
 
@@ -759,6 +769,7 @@ void FPVGrafter::DistributeGraftWithHormoneBasedSettings(
 			}
 		};
 
+	auto PointLengthFromRootAttribute = PV::FPointLengthFromRootAttribute::GetAttribute(OutCollection);
 	for (const auto& AttachmentPoint : AttachmentPoints)
 	{
 		const int32 PickedCandidateIndex = ConditionPicker::PickCandidateIndex(
@@ -781,7 +792,15 @@ void FPVGrafter::DistributeGraftWithHormoneBasedSettings(
 				AttachmentPoint.PointAlpha,
 				CurrentBudNumber,
 				CurrentBudNumber);
+			
 			CurrentBudNumber++;
+			
+			auto& Points = BranchPointsAttribute[AttachmentPoint.BranchIndex];
+			Points.Add(AttachmentPointIndex);
+			Points.Sort([PointLengthFromRootAttribute](const int32& A, const int32& B)
+					{
+						return PointLengthFromRootAttribute[A] < PointLengthFromRootAttribute[B];
+					});
 		}
 
 		FPVGraftAttributes Attributes{
@@ -888,6 +907,8 @@ void FPVGrafter::DistributeGraftWithParametricSettings(FManagedArrayCollection&
 			}
 		};
 
+	auto BranchPointsAttribute = PV::FBranchPointsAttribute::GetAttribute(OutCollection);
+	auto PointLengthFromRootAttribute = PV::FPointLengthFromRootAttribute::GetAttribute(OutCollection);
 	for (const auto& AttachmentPoint : AttachmentPoints)
 	{
 		const int32 PickedCandidateIndex = ConditionPicker::PickCandidateIndex(
@@ -918,8 +939,15 @@ void FPVGrafter::DistributeGraftWithParametricSettings(FManagedArrayCollection&
 				AttachmentPoint.PointAlpha,
 				CurrentBudNumber,
 				CurrentBudNumber);
-
+			
 			CurrentBudNumber++;
+			
+			auto& Points = BranchPointsAttribute[AttachmentPoint.BranchIndex];
+			Points.Add(AttachmentPointIndex);
+			Points.Sort([PointLengthFromRootAttribute](const int32& A, const int32& B)
+					{
+						return PointLengthFromRootAttribute[A] < PointLengthFromRootAttribute[B];
+					});
 		}
 
 		FPVGraftAttributes Attributes
diff --git a/Engine/Plugins/Experimental/ProceduralVegetationEditor/Source/ProceduralVegetation/Private/Nodes/PVGrafterPaletteSettings.cpp b/Engine/Plugins/Experimental/ProceduralVegetationEditor/Source/ProceduralVegetation/Private/Nodes/PVGrafterPaletteSettings.cpp
--- a/Engine/Plugins/Experimental/ProceduralVegetationEditor/Source/ProceduralVegetation/Private/Nodes/PVGrafterPaletteSettings.cpp
+++ b/Engine/Plugins/Experimental/ProceduralVegetationEditor/Source/ProceduralVegetation/Private/Nodes/PVGrafterPaletteSettings.cpp
@@ -9,12 +9,21 @@
 #include "DataTypes/PVGrowthData.h"
 #include "Facades/PVMetaInfoFacade.h"
 #include "GeometryCollection/ManagedArrayCollection.h"
+#include "Helpers/PCGHelpers.h"
 #include "Helpers/PCGSettingsHelpers.h"
 #include "Helpers/PVFoliageHelper.h"
 
 #define LOCTEXT_NAMESPACE "PVGrafterPaletteSettings"
 
 #if WITH_EDITOR
+UPVGrafterPaletteSettings::UPVGrafterPaletteSettings()
+{
+	if (PCGHelpers::IsNewObjectAndNotDefault(this))
+	{
+		GrafterInfos = { FPVFoliageAttributes{} };
+	}
+}
+
 FText UPVGrafterPaletteSettings::GetDefaultNodeTitle() const
 {
 	return LOCTEXT("NodeTitle", "Grafter Palette");
diff --git a/Engine/Plugins/Experimental/ProceduralVegetationEditor/Source/ProceduralVegetation/Private/Nodes/PVGrafterPaletteSettings.h b/Engine/Plugins/Experimental/ProceduralVegetationEditor/Source/ProceduralVegetation/Private/Nodes/PVGrafterPaletteSettings.h
--- a/Engine/Plugins/Experimental/ProceduralVegetationEditor/Source/ProceduralVegetation/Private/Nodes/PVGrafterPaletteSettings.h
+++ b/Engine/Plugins/Experimental/ProceduralVegetationEditor/Source/ProceduralVegetation/Private/Nodes/PVGrafterPaletteSettings.h
@@ -13,7 +13,9 @@ class UPVGrafterPaletteSettings : public UPVBaseSettings
 	GENERATED_BODY()
 
 public:
-#if WITH_EDITOR
+	UPVGrafterPaletteSettings();
+
+	#if WITH_EDITOR
 	virtual FName GetDefaultNodeName() const override { return FName(TEXT("Grafter Palette")); }
 	virtual FText GetDefaultNodeTitle() const override;
 	virtual FText GetNodeTooltipText() const override;
diff --git a/Engine/Plugins/Experimental/ProceduralVegetationEditor/Source/ProceduralVegetation/Private/Nodes/PVGrafterSettings.cpp b/Engine/Plugins/Experimental/ProceduralVegetationEditor/Source/ProceduralVegetation/Private/Nodes/PVGrafterSettings.cpp
--- a/Engine/Plugins/Experimental/ProceduralVegetationEditor/Source/ProceduralVegetation/Private/Nodes/PVGrafterSettings.cpp
+++ b/Engine/Plugins/Experimental/ProceduralVegetationEditor/Source/ProceduralVegetation/Private/Nodes/PVGrafterSettings.cpp
@@ -262,8 +262,8 @@ bool FPVGrafterElement::ExecuteInternal(FPCGContext* InContext) const
 			? FParametricSettings::EDistributionBasis::Plant
 			: FParametricSettings::EDistributionBasis::Branch,
 			.BaseScale = Settings->ParametricSettings.BaseScale,
-			.RandomizeScaleAngleMinimum = Settings->ParametricSettings.RandomizeScaleAngleMinimum,
-			.RandomizeScaleAngleMaximum = Settings->ParametricSettings.RandomizeScaleAngleMaximum,
+			.RandomizeScaleAngleMinimum = Settings->ParametricSettings.RandomizeScaleMinimum,
+			.RandomizeScaleAngleMaximum = Settings->ParametricSettings.RandomizeScaleMaximum,
 			.ScaleRamp = &Settings->ParametricSettings.ScaleRamp
 		};
 
diff --git a/Engine/Plugins/Experimental/ProceduralVegetationEditor/Source/ProceduralVegetation/Private/Nodes/PVGrafterSettings.h b/Engine/Plugins/Experimental/ProceduralVegetationEditor/Source/ProceduralVegetation/Private/Nodes/PVGrafterSettings.h
--- a/Engine/Plugins/Experimental/ProceduralVegetationEditor/Source/ProceduralVegetation/Private/Nodes/PVGrafterSettings.h
+++ b/Engine/Plugins/Experimental/ProceduralVegetationEditor/Source/ProceduralVegetation/Private/Nodes/PVGrafterSettings.h
@@ -128,10 +128,10 @@ struct FPVGrafterDistributionParametricSettings
 	float BaseScale = 0.2f;
 
 	UPROPERTY(EditAnywhere, Category="Parametric Scale Settings", meta=(PCG_Overridable, ClampMin=0.0f, ClampMax=1.0f))
-	float RandomizeScaleAngleMinimum = 1.0f;
+	float RandomizeScaleMinimum = 1.0f;
 
 	UPROPERTY(EditAnywhere, Category="Parametric Scale Settings", meta=(PCG_Overridable, ClampMin=1.0f, ClampMax=2.0f))
-	float RandomizeScaleAngleMaximum = 1.0f;
+	float RandomizeScaleMaximum = 1.0f;
 
 	UPROPERTY(EditAnywhere, Category="Parametric Scale Settings",
 		meta=(PCG_Overridable, XAxisMin=0.0f, XAxisMax=1.0f, YAxisMin=0.0f, YAxisMax=1.0f))
diff --git a/Engine/Plugins/Experimental/ProceduralVegetationEditor/Source/ProceduralVegetation/Public/Facades/PVPointFacade.h b/Engine/Plugins/Experimental/ProceduralVegetationEditor/Source/ProceduralVegetation/Public/Facades/PVPointFacade.h
--- a/Engine/Plugins/Experimental/ProceduralVegetationEditor/Source/ProceduralVegetation/Public/Facades/PVPointFacade.h
+++ b/Engine/Plugins/Experimental/ProceduralVegetationEditor/Source/ProceduralVegetation/Public/Facades/PVPointFacade.h
@@ -236,6 +236,8 @@ Accessor.ModifyAt(i, InputArray[i]); \
 		
 		const TManagedArray<TArray<float>>& GetBudLightDetectedArrays() const;
 		
+		const TManagedArray<TArray<float>>& GetBudLateralMeristemArrays() const;
+		
 		void SetBudLightDetectedArraysFromIndex(const TArray<TArray<float>>& InputArray, const int32 StartIndex);
 
 		void GetBudDevelopment(int32 Index, TArray<int>& OutValue) const;
