[PVE] Fixes for grafter node
- Fix issue where bud numbers are not kept unique
- Fix issue where children attribute is not computed correctly

#Synced-CL 51772764

diff --git a/Engine/Plugins/Experimental/ProceduralVegetationEditor/Source/ProceduralVegetation/Private/Implementations/PVGrafter.cpp b/Engine/Plugins/Experimental/ProceduralVegetationEditor/Source/ProceduralVegetation/Private/Implementations/PVGrafter.cpp
--- a/Engine/Plugins/Experimental/ProceduralVegetationEditor/Source/ProceduralVegetation/Private/Implementations/PVGrafter.cpp
+++ b/Engine/Plugins/Experimental/ProceduralVegetationEditor/Source/ProceduralVegetation/Private/Implementations/PVGrafter.cpp
@@ -178,7 +178,7 @@ void FPVGrafter::SetGraft(FManagedArrayCollection& OutCollection,
 
 	const TManagedArray<int32>& SourceSkeletonBudNumbers = PointFacadeOut.GetBudNumbersAttribute();
 	check(SourceSkeletonBudNumbers.Num() > 0);
-	int32 MaxBudNumber = *Algo::MaxElement(SourceSkeletonBudNumbers.GetConstArray());
+	int32 MaxBudNumber = InAttributes.CurrentBudNumber;
 	const TManagedArray<int32>& SourceGraftBudNumbers = PointFacadeSourceGraft.GetBudNumbersAttribute();
 	TArray<int32> TransformedBudNumbers;
 	const int32 NumOfPointsInSourceGraft = PointFacadeSourceGraft.GetElementCount();
@@ -531,13 +531,13 @@ void FPVGrafter::SetGraft(FManagedArrayCollection& OutCollection,
 			const TArray<int32>& SourceBranchChildren = BranchFacadeSourceGraft.GetChildren(SourceIndex);
 			const int32 NumOfChildrenInSource = SourceBranchChildren.Num();
 			TArray<int32> TransformedChildren;
-			TransformedChildren.SetNum(NumOfChildrenInSource);
+			TransformedChildren.Reserve(NumOfChildrenInSource);
 
 			for (int32 j = 0; j < NumOfChildrenInSource; ++j)
 			{
-				if (SourceOldBranchNumbersToNewBranchNumbers.Contains(SourceBranchChildren[j]))
+				if (const int32* NewNumber = SourceOldBranchNumbersToNewBranchNumbers.Find(SourceBranchChildren[j]))
 				{
-					TransformedChildren[j] = SourceOldBranchNumbersToNewBranchNumbers[SourceBranchChildren[j]];
+					TransformedChildren.Add(*NewNumber);
 				}
 			}
 
@@ -548,7 +548,31 @@ void FPVGrafter::SetGraft(FManagedArrayCollection& OutCollection,
 		SourceIndex++;
 	}
 
+	TArray<int32> HostBranchChildren = BranchFacadeOut.GetChildren(InAttributes.BranchIndex);
 
+	if (InAttributes.bIsTipInstance)
+	{
+		// The graft trunk merged into the host branch, so the trunk's source children
+		// become new children of the host branch.
+		const TArray<int32>& SourceTrunkChildren = BranchFacadeSourceGraft.GetChildren(SourceGraftTrunkBranchIndex);
+		for (const int32 SourceChildBranchNumber : SourceTrunkChildren)
+		{
+			if (const int32* NewNumber = SourceOldBranchNumbersToNewBranchNumbers.Find(SourceChildBranchNumber))
+			{
+				HostBranchChildren.Add(*NewNumber);
+			}
+		}
+	}
+	else
+	{
+		// The graft trunk is a new branch; add it as a direct child of the host branch.
+		const int32 GraftTrunkNewBranchNumber = SourceOldBranchNumbersToNewBranchNumbers[BranchNumbers[SourceGraftTrunkBranchIndex]];
+		HostBranchChildren.Add(GraftTrunkNewBranchNumber);
+	}
+
+	BranchFacadeOut.SetChildren(InAttributes.BranchIndex, MoveTemp(HostBranchChildren));
+
+	
 	TArray<int32> HostAncestorPrefix = BranchFacadeOut.GetParents(InAttributes.BranchIndex);
 	HostAncestorPrefix.Add(BranchFacadeOut.GetBranchNumber(InAttributes.BranchIndex));
 	SourceIndex = 0;
@@ -737,8 +761,6 @@ void FPVGrafter::DistributeGraftWithHormoneBasedSettings(
 
 	auto BudNumberAttribute = PV::FPointBudNumberAttribute::GetAttribute(OutCollection);
 	check(BudNumberAttribute.Num() > 0);
-	const int32 MaxBudNumber = *Algo::MaxElement(BudNumberAttribute);
-	int32 CurrentBudNumber = MaxBudNumber + 1;
 
 	auto IsValidCandidate = [&](int32 CandidateIndex) -> bool
 		{
@@ -772,6 +794,8 @@ void FPVGrafter::DistributeGraftWithHormoneBasedSettings(
 	auto PointLengthFromRootAttribute = PV::FPointLengthFromRootAttribute::GetAttribute(OutCollection);
 	for (const auto& AttachmentPoint : AttachmentPoints)
 	{
+		const int32 MaxBudNumber = *Algo::MaxElement(BudNumberAttribute);
+		int32 CurrentBudNumber = MaxBudNumber + 1;
 		const int32 PickedCandidateIndex = ConditionPicker::PickCandidateIndex(
 			GraftSourcesExtracted.Num(), DistributionConditions, AttachmentPoint.ConditionSample, RandomStream, GetCandidateValue, IsValidCandidate);
 
@@ -806,6 +830,7 @@ void FPVGrafter::DistributeGraftWithHormoneBasedSettings(
 		FPVGraftAttributes Attributes{
 			.BranchIndex = AttachmentPoint.BranchIndex,
 			.PointIndex = AttachmentPointIndex,
+			.CurrentBudNumber = CurrentBudNumber,
 			.AttachmentPoint = AttachmentPoint.Position,
 			.UpVector = AttachmentPoint.UpDirection,
 			.NormalVector = AttachmentPoint.NormalDirection,
@@ -875,8 +900,6 @@ void FPVGrafter::DistributeGraftWithParametricSettings(FManagedArrayCollection&
 
 	auto BudNumberAttribute = PV::FPointBudNumberAttribute::GetAttribute(OutCollection);
 	check(BudNumberAttribute.Num() > 0);
-	const int32 MaxBudNumber = *Algo::MaxElement(BudNumberAttribute);
-	int32 CurrentBudNumber = MaxBudNumber + 1;
 
 	auto IsValidCandidate = [&](int32 CandidateIndex) -> bool
 		{
@@ -911,6 +934,8 @@ void FPVGrafter::DistributeGraftWithParametricSettings(FManagedArrayCollection&
 	auto PointLengthFromRootAttribute = PV::FPointLengthFromRootAttribute::GetAttribute(OutCollection);
 	for (const auto& AttachmentPoint : AttachmentPoints)
 	{
+		const int32 MaxBudNumber = *Algo::MaxElement(BudNumberAttribute);
+		int32 CurrentBudNumber = MaxBudNumber + 1;
 		const int32 PickedCandidateIndex = ConditionPicker::PickCandidateIndex(
 			GraftSourcesExtracted.Num(),
 			DistributionConditions,
@@ -954,6 +979,7 @@ void FPVGrafter::DistributeGraftWithParametricSettings(FManagedArrayCollection&
 		{
 			.BranchIndex = AttachmentPoint.BranchIndex,
 			.PointIndex = AttachmentPointIndex,
+			.CurrentBudNumber = CurrentBudNumber,
 			.AttachmentPoint = AttachmentPoint.Position,
 			.UpVector = AttachmentPoint.UpDirection,
 			.NormalVector = AttachmentPoint.NormalDirection,
diff --git a/Engine/Plugins/Experimental/ProceduralVegetationEditor/Source/ProceduralVegetation/Private/Implementations/PVGrafter.h b/Engine/Plugins/Experimental/ProceduralVegetationEditor/Source/ProceduralVegetation/Private/Implementations/PVGrafter.h
--- a/Engine/Plugins/Experimental/ProceduralVegetationEditor/Source/ProceduralVegetation/Private/Implementations/PVGrafter.h
+++ b/Engine/Plugins/Experimental/ProceduralVegetationEditor/Source/ProceduralVegetation/Private/Implementations/PVGrafter.h
@@ -20,6 +20,7 @@ struct FPVGrafter
 	{
 		int32 BranchIndex;
 		int32 PointIndex;
+		int32 CurrentBudNumber;
 		FVector3f AttachmentPoint;
 		FVector3f UpVector;
 		FVector3f NormalVector;
