[PVE] Fix issue where mesh output was missing branches for grafter data

#Synced-CL 51784263

diff --git a/Engine/Plugins/Experimental/ProceduralVegetationEditor/Source/ProceduralVegetation/Private/Implementations/PVMeshBuilder.cpp b/Engine/Plugins/Experimental/ProceduralVegetationEditor/Source/ProceduralVegetation/Private/Implementations/PVMeshBuilder.cpp
--- a/Engine/Plugins/Experimental/ProceduralVegetationEditor/Source/ProceduralVegetation/Private/Implementations/PVMeshBuilder.cpp
+++ b/Engine/Plugins/Experimental/ProceduralVegetationEditor/Source/ProceduralVegetation/Private/Implementations/PVMeshBuilder.cpp
@@ -934,6 +934,27 @@ void FPVMeshBuilder::GenerateBranchMeshData(const bool bPrimitiveIsTrunk, const
 			const FVector3f& NextPointPosition = PointPositions[NextPointIndex];
 			TangentVector = (NextPointPosition - PointPosition).GetSafeNormal();
 		}
+		
+		// Guard against near-zero tangent when consecutive branch points are coincident
+		if (TangentVector.IsNearlyZero())
+		{
+			for (int32 SearchIndex = Index + 1; SearchIndex < PrimitivePoints.Num(); ++SearchIndex)
+			{
+				TangentVector = (PointPositions[PrimitivePoints[SearchIndex]] - PointPosition).GetSafeNormal();
+				if (!TangentVector.IsNearlyZero())
+				{
+					break;
+				}
+			}
+			if (TangentVector.IsNearlyZero())
+			{
+				TangentVector = PointBudDirections[PointIndex][0].GetSafeNormal(); // apical bud direction
+			}
+			if (TangentVector.IsNearlyZero())
+			{
+				TangentVector = FVector3f::UpVector;
+			}
+		}
 
 		/*
 		 * This creates a perpendicular Up Vector to the Current Tangent
@@ -950,6 +971,13 @@ void FPVMeshBuilder::GenerateBranchMeshData(const bool bPrimitiveIsTrunk, const
 		FVector3f CrossVector = FVector3f::CrossProduct(TangentVector, PreviousUpVector);
 		FVector3f UpVector = FVector3f::CrossProduct(CrossVector, TangentVector);
 
+		if (UpVector.IsNearlyZero())
+		{
+			FVector Axis1, Axis2;
+			FVector(TangentVector).FindBestAxisVectors(Axis1, Axis2);
+			UpVector = FVector3f(Axis1);
+		}
+
 		PreviousUpVector = UpVector;
 
 		FVector2f UV1 = FVector2f(0, 1.0 - Displacement_UV_V);
