[PVE] Fix scale related attribute computation for resume growth with grafter workflow

#Synced-CL 51784263

diff --git a/Engine/Plugins/Experimental/ProceduralVegetationEditor/Source/ProceduralVegetation/Private/Facades/PVPointFacade.cpp b/Engine/Plugins/Experimental/ProceduralVegetationEditor/Source/ProceduralVegetation/Private/Facades/PVPointFacade.cpp
--- a/Engine/Plugins/Experimental/ProceduralVegetationEditor/Source/ProceduralVegetation/Private/Facades/PVPointFacade.cpp
+++ b/Engine/Plugins/Experimental/ProceduralVegetationEditor/Source/ProceduralVegetation/Private/Facades/PVPointFacade.cpp
@@ -202,6 +202,22 @@ namespace PV::Facades
 		}
 	}
 
+	void FPointFacade::SetSeedPScale(int32 Index, float InSeedPScale)
+	{
+		if (SeedPScale.IsValid() && SeedPScale.IsValidIndex(Index))
+		{
+			SeedPScale.ModifyAt(Index, InSeedPScale);
+		}
+	}
+
+	void FPointFacade::SetSeedPScaleRatio(int32 Index, float InSeedPScaleRatio)
+	{
+		if (SeedPScaleRatio.IsValid() && SeedPScaleRatio.IsValidIndex(Index))
+		{
+			SeedPScaleRatio.ModifyAt(Index, InSeedPScaleRatio);
+		}
+	}
+
 	int32 FPointFacade::GetBudNumber(int32 Index) const
 	{
 		if (BudNumber.IsValid() && BudNumber.IsValidIndex(Index))
@@ -843,6 +859,16 @@ namespace PV::Facades
 		FILL_COLLECTION_ATTRIBUTE_FROM_INDEX(PointScale, InputArray, StartIndex);
 	}
 
+	void FPointFacade::SetSeedPScalesFromIndex(const TArray<float>& InputArray, const int32 StartIndex)
+	{
+		FILL_COLLECTION_ATTRIBUTE_FROM_INDEX(SeedPScale, InputArray, StartIndex);
+	}
+
+	void FPointFacade::SetSeedPScaleRatiosFromIndex(const TArray<float>& InputArray, const int32 StartIndex)
+	{
+		FILL_COLLECTION_ATTRIBUTE_FROM_INDEX(SeedPScaleRatio, InputArray, StartIndex);
+	}
+
 	void FPointFacade::SetLFRsFromIndex(const TArray<float>& InputArray, const int32 StartIndex)
 	{
 		FILL_COLLECTION_ATTRIBUTE_FROM_INDEX(LengthFromRoot, InputArray, StartIndex);
diff --git a/Engine/Plugins/Experimental/ProceduralVegetationEditor/Source/ProceduralVegetation/Private/Helpers/PVUtilities.cpp b/Engine/Plugins/Experimental/ProceduralVegetationEditor/Source/ProceduralVegetation/Private/Helpers/PVUtilities.cpp
--- a/Engine/Plugins/Experimental/ProceduralVegetationEditor/Source/ProceduralVegetation/Private/Helpers/PVUtilities.cpp
+++ b/Engine/Plugins/Experimental/ProceduralVegetationEditor/Source/ProceduralVegetation/Private/Helpers/PVUtilities.cpp
@@ -351,6 +351,18 @@ namespace PV::Utilities
 		const float CurrentPointScale = PointFacade.GetPointScale(CurrentPointIndex);
 		const float NextPointScale = PointFacade.GetPointScale(NextPointIndex);
 		const float PointScale = FMath::Lerp(CurrentPointScale, NextPointScale, Alpha);
+		
+		float CurrentSeedPScale = 0.f;
+		PointFacade.GetSeedPScale(CurrentPointIndex, CurrentSeedPScale);
+		float NextSeedPScale = 0.f;
+		PointFacade.GetSeedPScale(NextPointIndex, NextSeedPScale);
+		const float SeedPScale = FMath::Lerp(CurrentSeedPScale, NextSeedPScale, Alpha);
+		
+		float CurrentSeedPScaleRatio = 1.f;
+		PointFacade.GetSeedPScaleRatio(CurrentPointIndex, CurrentSeedPScaleRatio);
+		float NextSeedPScaleRatio = 1.f;
+		PointFacade.GetSeedPScaleRatio(NextPointIndex, NextSeedPScaleRatio);
+		const float SeedPScaleRatio = FMath::Lerp(CurrentSeedPScaleRatio, NextSeedPScaleRatio, Alpha);
 
 		const float CurrentPointLengthFromSeed = PointFacade.GetLengthFromSeed(CurrentPointIndex);
 		const float NextPointLengthFromSeed = PointFacade.GetLengthFromSeed(NextPointIndex);
@@ -417,6 +429,8 @@ namespace PV::Utilities
 		PointFacade.SetPosition(NewIndex, Position);
 		PointFacade.SetLengthFromRoot(NewIndex, LengthFromRoot);
 		PointFacade.SetPointScale(NewIndex, PointScale);
+		PointFacade.SetSeedPScale(NewIndex, SeedPScale);
+		PointFacade.SetSeedPScaleRatio(NewIndex, SeedPScaleRatio);
 		PointFacade.SetLengthFromSeed(NewIndex, LengthFromSeed);
 		PointFacade.SetNjordPixelIndex(NewIndex, NjordPixelIdx);
 		PointFacade.SetBudLightDetected(NewIndex, BudLightDetected);
diff --git a/Engine/Plugins/Experimental/ProceduralVegetationEditor/Source/ProceduralVegetation/Private/Implementations/PVGrafter.cpp b/Engine/Plugins/Experimental/ProceduralVegetationEditor/Source/ProceduralVegetation/Private/Implementations/PVGrafter.cpp
--- a/Engine/Plugins/Experimental/ProceduralVegetationEditor/Source/ProceduralVegetation/Private/Implementations/PVGrafter.cpp
+++ b/Engine/Plugins/Experimental/ProceduralVegetationEditor/Source/ProceduralVegetation/Private/Implementations/PVGrafter.cpp
@@ -252,12 +252,41 @@ void FPVGrafter::SetGraft(FManagedArrayCollection& OutCollection,
 
 	PointFacadeOut.SetPointScalesFromIndex(TransformedPointScales, PointFacadeOutStartIndex);
 	
+	TArray<float> TransformedSeedPScales;
+	TransformedSeedPScales.SetNum(NumberOfPointsToAddToSkeleton);
+
+	ForEachGraftPointExcludingRoot(NumOfPointsInSourceGraft, SourceGraftTrunkRootPointIndex,
+		[&](const int32 SourcePointIndex, const int32 TargetPointIndex)
+			{
+				float SeedPScaleValue = 0.f;
+				PointFacadeSourceGraft.GetSeedPScale(SourcePointIndex, SeedPScaleValue);
+				TransformedSeedPScales[TargetPointIndex] = SeedPScaleValue * ParentToGraftRootPointScaleRatio;
+			});
+
+	PointFacadeOut.SetSeedPScalesFromIndex(TransformedSeedPScales, PointFacadeOutStartIndex);
+	
+	TArray<float> TransformedSeedPScaleRatios;
+	TransformedSeedPScaleRatios.SetNum(NumberOfPointsToAddToSkeleton);
+
+	ForEachGraftPointExcludingRoot(NumOfPointsInSourceGraft, SourceGraftTrunkRootPointIndex,
+		[&](const int32 SourcePointIndex, const int32 TargetPointIndex)
+			{
+				float SeedPScaleRatioValue = 1.f;
+				PointFacadeSourceGraft.GetSeedPScaleRatio(SourcePointIndex, SeedPScaleRatioValue);
+				TransformedSeedPScaleRatios[TargetPointIndex] = SeedPScaleRatioValue;
+			});
+
+	PointFacadeOut.SetSeedPScaleRatiosFromIndex(TransformedSeedPScaleRatios, PointFacadeOutStartIndex);
+	
 	const TManagedArray<TArray<float>>& BudLateralMeristemArrays = PointFacadeSourceGraft.GetBudLateralMeristemArrays();
 
 	ForEachGraftPointExcludingRoot(NumOfPointsInSourceGraft, SourceGraftTrunkRootPointIndex,
 		[&](const int32 SourcePointIndex, const int32 TargetPointIndex)
 			{
-				PointFacadeOut.SetBudLateralMeristem(TargetPointIndex + PointFacadeOutStartIndex, BudLateralMeristemArrays[SourcePointIndex]);
+				TArray<float> LateralMeristemValues = BudLateralMeristemArrays[SourcePointIndex];
+				check(LateralMeristemValues.Num() > 0);
+				LateralMeristemValues[0] *= ParentToGraftRootPointScaleRatio;
+				PointFacadeOut.SetBudLateralMeristem(TargetPointIndex + PointFacadeOutStartIndex, LateralMeristemValues);
 			});
 
 	const TManagedArray<float>& LFRs = PointFacadeSourceGraft.GetLengthFromRootsArray();
diff --git a/Engine/Plugins/Experimental/ProceduralVegetationEditor/Source/ProceduralVegetation/Public/Facades/PVPointFacade.h b/Engine/Plugins/Experimental/ProceduralVegetationEditor/Source/ProceduralVegetation/Public/Facades/PVPointFacade.h
--- a/Engine/Plugins/Experimental/ProceduralVegetationEditor/Source/ProceduralVegetation/Public/Facades/PVPointFacade.h
+++ b/Engine/Plugins/Experimental/ProceduralVegetationEditor/Source/ProceduralVegetation/Public/Facades/PVPointFacade.h
@@ -62,6 +62,10 @@ Accessor.ModifyAt(i, InputArray[i]); \
 		float GetPointScale(int32 Index) const;
 
 		void SetPointScale(int32 Index, float InPointScale);
+		
+		void SetSeedPScale(int32 Index, float InSeedPScale);
+		
+		void SetSeedPScaleRatio(int32 Index, float InSeedPScaleRatio);
 
 		int32 GetBudNumber(int32 Index) const;
 
@@ -179,6 +183,10 @@ Accessor.ModifyAt(i, InputArray[i]); \
 		
 		void SetPointScalesFromIndex(const TArray<float>& InputArray, const int32 StartIndex);
 		
+		void SetSeedPScalesFromIndex(const TArray<float>& InputArray, const int32 StartIndex);
+		
+		void SetSeedPScaleRatiosFromIndex(const TArray<float>& InputArray, const int32 StartIndex);
+		
 		void SetLFRsFromIndex(const TArray<float>& InputArray, const int32 StartIndex);
 		
 		const TManagedArray<float>& GetLengthFromSeedsArray() const;
