[PVE] Added IsFusedPoint() to FPointFacade which checks generation and source bud number
to identify fused points. Branch pruning now skips removal of fused points to avoid
incorrectly removing points shared with other branches.

#Synced-CL 51757984

diff --git a/Engine/Plugins/Experimental/ProceduralVegetationEditor/Source/ProceduralVegetation/Private/Facades/PVPointFacade.cpp b/Engine/Plugins/Experimental/ProceduralVegetationEditor/Source/ProceduralVegetation/Private/Facades/PVPointFacade.cpp
--- a/Engine/Plugins/Experimental/ProceduralVegetationEditor/Source/ProceduralVegetation/Private/Facades/PVPointFacade.cpp
+++ b/Engine/Plugins/Experimental/ProceduralVegetationEditor/Source/ProceduralVegetation/Private/Facades/PVPointFacade.cpp
@@ -513,6 +513,13 @@ namespace PV::Facades
 		}
 	}
 
+	bool FPointFacade::IsFusedPoint(const int Generation, const int SourceBudNumber, const int PointIndex) const
+	{
+		const int Number = GetBudNumber(PointIndex);
+		
+		return Generation > 1 && Number == SourceBudNumber;
+	}
+
 	TManagedArray<FVector3f>& FPointFacade::ModifyPositions()
 	{
 		return Positions.Modify();
diff --git a/Engine/Plugins/Experimental/ProceduralVegetationEditor/Source/ProceduralVegetation/Private/Facades/PVTreeFacade.cpp b/Engine/Plugins/Experimental/ProceduralVegetationEditor/Source/ProceduralVegetation/Private/Facades/PVTreeFacade.cpp
--- a/Engine/Plugins/Experimental/ProceduralVegetationEditor/Source/ProceduralVegetation/Private/Facades/PVTreeFacade.cpp
+++ b/Engine/Plugins/Experimental/ProceduralVegetationEditor/Source/ProceduralVegetation/Private/Facades/PVTreeFacade.cpp
@@ -122,11 +122,19 @@ namespace PV::Facades
 	
 		for (int32 i = 0; i < InBranchesToRemove.Num(); i++)
 		{
+			const int BranchIndex = InBranchesToRemove[i];
+			
 			BranchesToRemove[InBranchesToRemove[i]] = true;
-		
+
+			const int Generation = BranchFacadeOut.GetHierarchyGenerationNumber(BranchIndex);
+			const int SourceBudNumber = BranchFacadeOut.GetBranchSourceBudNumber(BranchIndex);
+			
 			for (const int32& P : BranchFacadeOut.GetPoints(InBranchesToRemove[i]))
 			{
-				PointsToRemove[P] = true;
+				if (!PointFacadeOut.IsFusedPoint(Generation, SourceBudNumber, P))
+				{
+					PointsToRemove[P] = true;
+				}
 			}
 			for (const int32& FId : FoliageFacadeOut.GetFoliageEntryIdsForBranch(InBranchesToRemove[i]))
 			{
diff --git a/Engine/Plugins/Experimental/ProceduralVegetationEditor/Source/ProceduralVegetation/Public/Facades/PVPointFacade.h b/Engine/Plugins/Experimental/ProceduralVegetationEditor/Source/ProceduralVegetation/Public/Facades/PVPointFacade.h
--- a/Engine/Plugins/Experimental/ProceduralVegetationEditor/Source/ProceduralVegetation/Public/Facades/PVPointFacade.h
+++ b/Engine/Plugins/Experimental/ProceduralVegetationEditor/Source/ProceduralVegetation/Public/Facades/PVPointFacade.h
@@ -125,6 +125,8 @@ Accessor.ModifyAt(i, InputArray[i]); \
 		
 		void SetURange(int32 Index, FVector2f InURange);
 
+		bool IsFusedPoint(const int Generation, const int SourceBudNumber, const int PointIndex) const;
+
 		TManagedArray<FVector3f>& ModifyPositions();
 		
 		TManagedArray<float>& ModifyPointScales();
