Fix RemoveBranches with Generation Basis not correctly handling root branches when processing growth data from 2D/3D importers
Ensure root branches are assigned a valid parent entry instead of being left uninitialized, resolving compatibility with the import-generated data path

#Synced-CL 51761158

diff --git a/Engine/Plugins/Experimental/ProceduralVegetationEditor/Source/ProceduralVegetation/Private/Helpers/PVImportHelpers.cpp b/Engine/Plugins/Experimental/ProceduralVegetationEditor/Source/ProceduralVegetation/Private/Helpers/PVImportHelpers.cpp
--- a/Engine/Plugins/Experimental/ProceduralVegetationEditor/Source/ProceduralVegetation/Private/Helpers/PVImportHelpers.cpp
+++ b/Engine/Plugins/Experimental/ProceduralVegetationEditor/Source/ProceduralVegetation/Private/Helpers/PVImportHelpers.cpp
@@ -929,7 +929,12 @@ void PV::ImportHelper::FillAttributesFromBranchHierarchy(
 		}
 		BranchPointsAttribute[AttributeBranchIndex] = MoveTemp(AttributeBranchPoints);
 
-		if (AttributeParentBranchIndex != INDEX_NONE)
+		const bool bIsRootBranch = AttributeParentBranchIndex == INDEX_NONE;
+		if (bIsRootBranch)
+		{
+			BranchParentsAttribute[AttributeBranchIndex] = { 0 };
+		}
+		else
 		{
 			TArray<int32> Parents = BranchParentsAttribute[AttributeParentBranchIndex];
 			Parents.Add(BranchIndexToBranchNumber(AttributeParentBranchIndex));
