Fix collider transform and skeleton edit tool click interactions not working when connecting via remote desktop solutions
Move hit proxy validation from IsHitByClick into OnClicked to ensure selection clearing and gizmo visibility are handled correctly regardless of whether a focused viewport is available

#Synced-CL 51761158

diff --git a/Engine/Plugins/Experimental/ProceduralVegetationEditor/Source/ProceduralVegetationEditor/Private/Tools/NodeTools/PVColliderTransformTool.cpp b/Engine/Plugins/Experimental/ProceduralVegetationEditor/Source/ProceduralVegetationEditor/Private/Tools/NodeTools/PVColliderTransformTool.cpp
--- a/Engine/Plugins/Experimental/ProceduralVegetationEditor/Source/ProceduralVegetationEditor/Private/Tools/NodeTools/PVColliderTransformTool.cpp
+++ b/Engine/Plugins/Experimental/ProceduralVegetationEditor/Source/ProceduralVegetationEditor/Private/Tools/NodeTools/PVColliderTransformTool.cpp
@@ -79,47 +79,37 @@ void UPVColliderTransformTool::OnNodeSettingsUpdated()
 
 FInputRayHit UPVColliderTransformTool::IsHitByClick(const FInputDeviceRay& ClickPos)
 {
-	if (FViewport* FocusedViewport = GetToolManager()->GetContextQueriesAPI()->GetFocusedViewport())
-	{
-		if (FocusedViewport->GetHitProxy(ClickPos.ScreenPosition.X, ClickPos.ScreenPosition.Y))
-		{
-			return FInputRayHit(1.0f);
-		}
-	}
-
-	CombinedGizmo->SetVisibility(false);
-	return FInputRayHit();
+	return FInputRayHit(1.0f);
 }
 
 void UPVColliderTransformTool::OnClicked(const FInputDeviceRay& ClickPos)
 {
 	FViewport* FocusedViewport = GetToolManager()->GetContextQueriesAPI()->GetFocusedViewport();
-	if (!FocusedViewport)
-	{
-		return;
-	}
-	const HHitProxy* HitProxy = FocusedViewport->GetHitProxy(ClickPos.ScreenPosition.X, ClickPos.ScreenPosition.Y);
-	if (!HitProxy)
+	const HHitProxy* HitProxy = FocusedViewport ? FocusedViewport->GetHitProxy(ClickPos.ScreenPosition.X, ClickPos.ScreenPosition.Y) : nullptr;
+	if (!HitProxy || !HitProxy->IsA(HActor::StaticGetType()))
 	{
+		SelectedComponentIndex = INDEX_NONE;
+		CombinedGizmo->SetVisibility(false);
 		return;
 	}
 
-	if (HitProxy->IsA(HActor::StaticGetType()))
+	const HActor* ActorHitProxy = static_cast<const HActor*>(HitProxy);
+	if (ActorHitProxy->Actor == PreviewActor)
 	{
-		const HActor* ActorHitProxy = static_cast<const HActor*>(HitProxy);
-		if (ActorHitProxy->Actor == PreviewActor)
-		{
-			const TObjectPtr<const UPrimitiveComponent> HitComponent = ActorHitProxy->PrimComponent;
-			SelectedComponentIndex = MeshComponents.IndexOfByPredicate([&](const auto& Comp)
-				{
-					return Comp == HitComponent;
-				}
-			);
-			if (SelectedComponentIndex != INDEX_NONE)
+		const TObjectPtr<const UPrimitiveComponent> HitComponent = ActorHitProxy->PrimComponent;
+		SelectedComponentIndex = MeshComponents.IndexOfByPredicate([&](const auto& Comp)
 			{
-				CombinedGizmo->SetVisibility(true);
-				CombinedGizmo->SetNewGizmoTransform(ActorHitProxy->PrimComponent->GetComponentTransform());
+				return Comp == HitComponent;
 			}
+		);
+		if (ensure(SelectedComponentIndex != INDEX_NONE))
+		{
+			CombinedGizmo->SetVisibility(true);
+			CombinedGizmo->SetNewGizmoTransform(ActorHitProxy->PrimComponent->GetComponentTransform());
+		}
+		else
+		{
+			CombinedGizmo->SetVisibility(false);
 		}
 	}
 }
diff --git a/Engine/Plugins/Experimental/ProceduralVegetationEditor/Source/ProceduralVegetationEditor/Private/Tools/NodeTools/PVSkeletonEditTool.cpp b/Engine/Plugins/Experimental/ProceduralVegetationEditor/Source/ProceduralVegetationEditor/Private/Tools/NodeTools/PVSkeletonEditTool.cpp
--- a/Engine/Plugins/Experimental/ProceduralVegetationEditor/Source/ProceduralVegetationEditor/Private/Tools/NodeTools/PVSkeletonEditTool.cpp
+++ b/Engine/Plugins/Experimental/ProceduralVegetationEditor/Source/ProceduralVegetationEditor/Private/Tools/NodeTools/PVSkeletonEditTool.cpp
@@ -470,188 +470,168 @@ void UPVSkeletonEditTool::OnNodeSettingsUpdated()
 
 FInputRayHit UPVSkeletonEditTool::IsHitByClick(const FInputDeviceRay& ClickPos)
 {
-	if (FViewport* FocusedViewport = GetToolManager()->GetContextQueriesAPI()->GetFocusedViewport())
-	{
-		if (FocusedViewport->GetHitProxy(ClickPos.ScreenPosition.X, ClickPos.ScreenPosition.Y))
-		{
-			return FInputRayHit(1.0f);
-		}
-	}
-
-	Selection->ClearSelection();
-	CombinedGizmo->SetVisibility(false);
-	return FInputRayHit();
+	return FInputRayHit(1.0f);
 }
 
 void UPVSkeletonEditTool::OnClicked(const FInputDeviceRay& ClickPos)
 {
 	FViewport* FocusedViewport = GetToolManager()->GetContextQueriesAPI()->GetFocusedViewport();
-	if (!FocusedViewport)
-	{
-		return;
-	}
-	const HHitProxy* HitProxy = FocusedViewport->GetHitProxy(ClickPos.ScreenPosition.X, ClickPos.ScreenPosition.Y);
-	if (!HitProxy)
+	const HHitProxy* HitProxy = FocusedViewport ? FocusedViewport->GetHitProxy(ClickPos.ScreenPosition.X, ClickPos.ScreenPosition.Y) : nullptr;
+	if (!HitProxy || !HitProxy->IsA(HPVSkeletonPointProxy::StaticGetType()))
 	{
+		PreviousRotation = FQuat4f::Identity;
+		Selection->ClearSelection();
+		CombinedGizmo->SetVisibility(false);
 		return;
 	}
 
-	if (HitProxy->IsA(HPVSkeletonPointProxy::StaticGetType()))
-	{
-		const PV::Facades::FPointFacade PointFacade(GetCollection());
-		const PV::Facades::FBranchFacade BranchFacade(GetCollection());
+	const PV::Facades::FPointFacade PointFacade(GetCollection());
+	const PV::Facades::FBranchFacade BranchFacade(GetCollection());
 
-		const auto GetParentPoint = [&](const int32 BranchIndex) -> TPair<int32, int32>
+	const auto GetParentPoint = [&](const int32 BranchIndex) -> TPair<int32, int32>
+		{
+			const int32 RootPointIndex = BranchFacade.GetRootPoint(BranchIndex);
+			if (RootPointIndex != INDEX_NONE)
 			{
-				const int32 RootPointIndex = BranchFacade.GetRootPoint(BranchIndex);
-				if (RootPointIndex != INDEX_NONE)
+				const float PointLFR = PointFacade.GetLengthFromRoot(RootPointIndex);
+				const int32 ParentBranchIndex = BranchFacade.GetParentBranchIndex(BranchIndex);
+				if (ParentBranchIndex != INDEX_NONE)
 				{
-					const float PointLFR = PointFacade.GetLengthFromRoot(RootPointIndex);
-					const int32 ParentBranchIndex = BranchFacade.GetParentBranchIndex(BranchIndex);
-					if (ParentBranchIndex != INDEX_NONE)
+					const TArray<int32>& ParentBranchPoints = BranchFacade.GetPoints(ParentBranchIndex);
+					for (int32 BranchPointIndex = 0; BranchPointIndex < ParentBranchPoints.Num(); BranchPointIndex++)
 					{
-						const TArray<int32>& ParentBranchPoints = BranchFacade.GetPoints(ParentBranchIndex);
-						for (int32 BranchPointIndex = 0; BranchPointIndex < ParentBranchPoints.Num(); BranchPointIndex++)
+						if (PointFacade.GetLengthFromRoot(ParentBranchPoints[BranchPointIndex]) >= PointLFR)
 						{
-							if (PointFacade.GetLengthFromRoot(ParentBranchPoints[BranchPointIndex]) >= PointLFR)
-							{
-								return {ParentBranchIndex, BranchPointIndex};
-							}
+							return {ParentBranchIndex, BranchPointIndex};
 						}
 					}
 				}
-				return {INDEX_NONE, INDEX_NONE};
-			};
+			}
+			return {INDEX_NONE, INDEX_NONE};
+		};
 
-		const HPVSkeletonPointProxy* SkeletonPointProxy = static_cast<const HPVSkeletonPointProxy*>(HitProxy);
+	const HPVSkeletonPointProxy* SkeletonPointProxy = static_cast<const HPVSkeletonPointProxy*>(HitProxy);
 
-		const TObjectPtr<UPVSkeletonEditSettings> SkeletonEditSettings = GetNodeSettings<UPVSkeletonEditSettings>();
+	const TObjectPtr<UPVSkeletonEditSettings> SkeletonEditSettings = GetNodeSettings<UPVSkeletonEditSettings>();
 
-		if (SkeletonEditSettings->SkeletonEditSettings.SkeletonEditMode == ESkeletonEditMode::TRS)
-		{
-			const FPVSkeletonSelectionParams& SelectionParams = SkeletonEditSettings->SkeletonEditSettings.SkeletonSelectionParams;
+	if (SkeletonEditSettings->SkeletonEditSettings.SkeletonEditMode == ESkeletonEditMode::TRS)
+	{
+		const FPVSkeletonSelectionParams& SelectionParams = SkeletonEditSettings->SkeletonEditSettings.SkeletonSelectionParams;
 
-			if (!SelectionParams.bRootSelectsBranch && SkeletonPointProxy->BranchPointIndex == 0)
-			{
-				const auto [ParentBranchIndex, ParentBranchPointIndex] = GetParentPoint(SkeletonPointProxy->BranchIndex);
-				if (ParentBranchIndex == INDEX_NONE || ParentBranchPointIndex == INDEX_NONE)
-				{
-					Selection->SelectPoint(SkeletonPointProxy->BranchIndex, SkeletonPointProxy->BranchPointIndex);
-				}
-				else
-				{
-					Selection->SelectPoint(ParentBranchIndex, ParentBranchPointIndex);
-				}
-			}
-			else
+		if (!SelectionParams.bRootSelectsBranch && SkeletonPointProxy->BranchPointIndex == 0)
+		{
+			const auto [ParentBranchIndex, ParentBranchPointIndex] = GetParentPoint(SkeletonPointProxy->BranchIndex);
+			if (ParentBranchIndex == INDEX_NONE || ParentBranchPointIndex == INDEX_NONE)
 			{
 				Selection->SelectPoint(SkeletonPointProxy->BranchIndex, SkeletonPointProxy->BranchPointIndex);
 			}
-
-			if (SelectionParams.SkeletonSelectionMode == ESkeletonSelectionMode::SelectByDistance)
-			{
-				Selection->ExtendSelectionByDistance(
-					SelectionParams.SelectionDistance,
-					SkeletonEditSettings->SkeletonEditSettings.RelativeOffsets,
-					SelectionParams.SelectionFalloff
-				);
-			}
 			else
 			{
-				Selection->ExtendSelection(
-					SelectionParams.NumPointsSelection,
-					SelectionParams.SelectionFalloff
-				);
+				Selection->SelectPoint(ParentBranchIndex, ParentBranchPointIndex);
 			}
+		}
+		else
+		{
+			Selection->SelectPoint(SkeletonPointProxy->BranchIndex, SkeletonPointProxy->BranchPointIndex);
+		}
 
-			TArray<FVector3f>& Offsets = SkeletonEditSettings->SkeletonEditSettings.RelativeOffsets;
-
-			const int32 SelectedPointIndex = BranchFacade.GetPoints(SkeletonPointProxy->BranchIndex)[SkeletonPointProxy->BranchPointIndex];
-			const int32 SelectedPointBudNumber = PointFacade.GetBudNumber(SelectedPointIndex);
-			const FVector Position = static_cast<FVector>(PointFacade.GetPosition(SelectedPointIndex) + Offsets[SelectedPointBudNumber]);
-			PreviousRotation = FQuat4f::Identity;
-			CombinedGizmo->SetVisibility(true);
-			CombinedGizmo->ReinitializeGizmoTransform(FTransform(Position));
+		if (SelectionParams.SkeletonSelectionMode == ESkeletonSelectionMode::SelectByDistance)
+		{
+			Selection->ExtendSelectionByDistance(
+				SelectionParams.SelectionDistance,
+				SkeletonEditSettings->SkeletonEditSettings.RelativeOffsets,
+				SelectionParams.SelectionFalloff
+			);
 		}
-		else if (SkeletonEditSettings->SkeletonEditSettings.SkeletonEditMode == ESkeletonEditMode::BranchRemoval)
+		else
 		{
-			PreviousRotation = FQuat4f::Identity;
-			Selection->ClearSelection();
-			CombinedGizmo->SetVisibility(false);
+			Selection->ExtendSelection(
+				SelectionParams.NumPointsSelection,
+				SelectionParams.SelectionFalloff
+			);
+		}
+
+		TArray<FVector3f>& Offsets = SkeletonEditSettings->SkeletonEditSettings.RelativeOffsets;
+
+		const int32 SelectedPointIndex = BranchFacade.GetPoints(SkeletonPointProxy->BranchIndex)[SkeletonPointProxy->BranchPointIndex];
+		const int32 SelectedPointBudNumber = PointFacade.GetBudNumber(SelectedPointIndex);
+		const FVector Position = static_cast<FVector>(PointFacade.GetPosition(SelectedPointIndex) + Offsets[SelectedPointBudNumber]);
+		PreviousRotation = FQuat4f::Identity;
+		CombinedGizmo->SetVisibility(true);
+		CombinedGizmo->ReinitializeGizmoTransform(FTransform(Position));
+	}
+	else if (SkeletonEditSettings->SkeletonEditSettings.SkeletonEditMode == ESkeletonEditMode::BranchRemoval)
+	{
+		PreviousRotation = FQuat4f::Identity;
+		Selection->ClearSelection();
+		CombinedGizmo->SetVisibility(false);
 
-			TArray<bool>& RemovedPoints = SkeletonEditSettings->SkeletonEditSettings.RemovedPoints;
+		TArray<bool>& RemovedPoints = SkeletonEditSettings->SkeletonEditSettings.RemovedPoints;
 
-			const TArray<int32>& BranchPoints = BranchFacade.GetPoints(SkeletonPointProxy->BranchIndex);
+		const TArray<int32>& BranchPoints = BranchFacade.GetPoints(SkeletonPointProxy->BranchIndex);
 
-			int32 StartIndex =
-				SkeletonEditSettings->SkeletonEditSettings.BranchRemovalParams.BranchRemovalMode == EBranchRemovalMode::TrimBranch
-				? SkeletonPointProxy->BranchPointIndex
-				: 0;
+		int32 StartIndex =
+			SkeletonEditSettings->SkeletonEditSettings.BranchRemovalParams.BranchRemovalMode == EBranchRemovalMode::TrimBranch
+			? SkeletonPointProxy->BranchPointIndex
+			: 0;
 
-			bool bNextPointRemoved = false;
-			if (BranchPoints.IsValidIndex(StartIndex + 1))
-			{
-				const int32 NextPointBudNumber = PointFacade.GetBudNumber(BranchPoints[StartIndex + 1]);
-				bNextPointRemoved = RemovedPoints[NextPointBudNumber];
-			}
-			if (StartIndex > 0 && StartIndex < BranchPoints.Num() - 1 && !bNextPointRemoved)
-			{
-				// Remove from the next point onwards from currently selected point
-				// unless it's the first or the last point, we want to remove the first or last point if selected
-				StartIndex += 1;
-			}
-			if (StartIndex <= 1)
-			{
-				// Remove the complete branch if first point was selected (can't have a branch with 1 point)
-				StartIndex = 0;
-			}
+		bool bNextPointRemoved = false;
+		if (BranchPoints.IsValidIndex(StartIndex + 1))
+		{
+			const int32 NextPointBudNumber = PointFacade.GetBudNumber(BranchPoints[StartIndex + 1]);
+			bNextPointRemoved = RemovedPoints[NextPointBudNumber];
+		}
+		if (StartIndex > 0 && StartIndex < BranchPoints.Num() - 1 && !bNextPointRemoved)
+		{
+			// Remove from the next point onwards from currently selected point
+			// unless it's the first or the last point, we want to remove the first or last point if selected
+			StartIndex += 1;
+		}
+		if (StartIndex <= 1)
+		{
+			// Remove the complete branch if first point was selected (can't have a branch with 1 point)
+			StartIndex = 0;
+		}
 			
-			for (int32 BranchPointIndex = StartIndex; BranchPointIndex < BranchPoints.Num(); BranchPointIndex++)
-			{
-				const int32 PointIndex = BranchPoints[BranchPointIndex];
-				const int32 PointBudNumber = PointFacade.GetBudNumber(PointIndex);
-				RemovedPoints[PointBudNumber] = true;
-			}
+		for (int32 BranchPointIndex = StartIndex; BranchPointIndex < BranchPoints.Num(); BranchPointIndex++)
+		{
+			const int32 PointIndex = BranchPoints[BranchPointIndex];
+			const int32 PointBudNumber = PointFacade.GetBudNumber(PointIndex);
+			RemovedPoints[PointBudNumber] = true;
+		}
 
-			const int32 RemovedPointIndex = BranchPoints[StartIndex];
-			const float PointLFR = PointFacade.GetLengthFromRoot(RemovedPointIndex);
-			const TManagedArray<int32>& BranchNumbers = BranchFacade.GetBranchNumbers();
-			for (int32 BranchChildNumber : BranchFacade.GetImmediateChildren(SkeletonPointProxy->BranchIndex))
+		const int32 RemovedPointIndex = BranchPoints[StartIndex];
+		const float PointLFR = PointFacade.GetLengthFromRoot(RemovedPointIndex);
+		const TManagedArray<int32>& BranchNumbers = BranchFacade.GetBranchNumbers();
+		for (int32 BranchChildNumber : BranchFacade.GetImmediateChildren(SkeletonPointProxy->BranchIndex))
+		{
+			if (const int32 BranchIndex = BranchNumbers.Find(BranchChildNumber); BranchIndex != INDEX_NONE)
 			{
-				if (const int32 BranchIndex = BranchNumbers.Find(BranchChildNumber); BranchIndex != INDEX_NONE)
+				const TArray<int32>& ChildBranchPoints = BranchFacade.GetPoints(BranchIndex);
+				if (PointFacade.GetLengthFromRoot(ChildBranchPoints[0]) >= PointLFR)
 				{
-					const TArray<int32>& ChildBranchPoints = BranchFacade.GetPoints(BranchIndex);
-					if (PointFacade.GetLengthFromRoot(ChildBranchPoints[0]) >= PointLFR)
+					for (const int32 ChildBranchPointIndex : ChildBranchPoints)
 					{
-						for (const int32 ChildBranchPointIndex : ChildBranchPoints)
-						{
-							const int32 ChildPointBudNumber = PointFacade.GetBudNumber(ChildBranchPointIndex);
-							RemovedPoints[ChildPointBudNumber] = true;
-						}
+						const int32 ChildPointBudNumber = PointFacade.GetBudNumber(ChildBranchPointIndex);
+						RemovedPoints[ChildPointBudNumber] = true;
+					}
 
-						for (int32 SubBranchChildNumber : BranchFacade.GetChildren(BranchIndex))
+					for (int32 SubBranchChildNumber : BranchFacade.GetChildren(BranchIndex))
+					{
+						if (const int32 SubBranchIndex = BranchNumbers.Find(SubBranchChildNumber); SubBranchIndex != INDEX_NONE)
 						{
-							if (const int32 SubBranchIndex = BranchNumbers.Find(SubBranchChildNumber); SubBranchIndex != INDEX_NONE)
+							for (const int32 ChildSubBranchPointIndex : BranchFacade.GetPoints(SubBranchIndex))
 							{
-								for (const int32 ChildSubBranchPointIndex : BranchFacade.GetPoints(SubBranchIndex))
-								{
-									const int32 SubBranchChildPointBudNumber = PointFacade.GetBudNumber(ChildSubBranchPointIndex);
-									RemovedPoints[SubBranchChildPointBudNumber] = true;
-								}
+								const int32 SubBranchChildPointBudNumber = PointFacade.GetBudNumber(ChildSubBranchPointIndex);
+								RemovedPoints[SubBranchChildPointBudNumber] = true;
 							}
 						}
 					}
 				}
 			}
-
-			GetNodeSettings()->PostEditChange();
 		}
-	}
-	else
-	{
-		PreviousRotation = FQuat4f::Identity;
-		Selection->ClearSelection();
-		CombinedGizmo->SetVisibility(false);
+
+		GetNodeSettings()->PostEditChange();
 	}
 }
 
