Fix incorrect parent branch selection in procedural vegetation import

Remove radius from the parent branch scoring heuristic as noisy radius values were skewing selection results
Replace min-max normalization with sum-based normalization for branch length to preserve relative magnitude differences between candidates
Do not normalize dot products as they are already in 0-1 range, and do not make sense to compare relativly

#Synced-CL 51879788

diff --git a/Engine/Plugins/Experimental/ProceduralVegetationEditor/Source/ProceduralVegetation/Private/Helpers/PVImportHelpers.cpp b/Engine/Plugins/Experimental/ProceduralVegetationEditor/Source/ProceduralVegetation/Private/Helpers/PVImportHelpers.cpp
--- a/Engine/Plugins/Experimental/ProceduralVegetationEditor/Source/ProceduralVegetation/Private/Helpers/PVImportHelpers.cpp
+++ b/Engine/Plugins/Experimental/ProceduralVegetationEditor/Source/ProceduralVegetation/Private/Helpers/PVImportHelpers.cpp
@@ -227,46 +227,43 @@ namespace PV::ImportHelper
 
 			TArray<float, TInlineAllocator<4>> Dots;
 			TArray<float, TInlineAllocator<4>> Lengths;
-			TArray<float, TInlineAllocator<4>> Radii;
 			Dots.SetNumUninitialized(NextPoints.Num());
 			Lengths.SetNumUninitialized(NextPoints.Num());
-			Radii.SetNumUninitialized(NextPoints.Num());
 
 			for (int32 i = 0; i < NextPoints.Num(); ++i)
 			{
 				const int32 NextVertexID = NextPoints[i];
 				const FVector3d NextPos = DynamicMesh.GetVertex(NextVertexID);
 				const FVector3d NextDir = (NextPos - VertexPos).GetSafeNormal();
-				Dots[i] = FVector3d::DotProduct(ApicalDirection, NextDir);
+				Dots[i] = (1.0 + FVector3d::DotProduct(ApicalDirection, NextDir)) / 2.0;
 				Lengths[i] = MaxSubtreeLengthFromRoot[NextVertexID] - LengthFromRootAttribute.GetValue(NextVertexID);
-				Radii[i] = FMath::Abs(VertexScaleAttribute.GetValue(NextVertexID) - CurrentRadius);
 			}
 
-			// Normalize a raw float array to range [0,1].
-			// bInvert: output [1,0] instead of [0,1]
+			// Normalize a raw float array (all values >= 0) by dividing by the sum.
+			// This preserves proportionality: a value of 10000 vs 1 stays near 1.0 vs 0.0001,
+			// whereas a value of 10 vs 9 stays near 0.526 vs 0.473.
+			// Min-max normalization collapses both cases to {1, 0}, losing magnitude information.
+			// bInvert: output 1-normalized instead of normalized
 			const auto NormalizeArray = [](TArray<float, TInlineAllocator<4>>& InOutArray, bool bInvert)
 			{
-				float MinVal = InOutArray[0];
-				float MaxVal = InOutArray[0];
-				for (int32 i = 1; i < InOutArray.Num(); ++i)
+				float Sum = 0.f;
+				for (int32 i = 0; i < InOutArray.Num(); ++i)
 				{
-					MinVal = FMath::Min(MinVal, InOutArray[i]);
-					MaxVal = FMath::Max(MaxVal, InOutArray[i]);
+					Sum += InOutArray[i];
 				}
 
-				const float Range = MaxVal - MinVal;
-				if (Range < UE_SMALL_NUMBER)
+				if (Sum < UE_SMALL_NUMBER)
 				{
 					for (int32 i = 0; i < InOutArray.Num(); ++i)
-					{ 
-						InOutArray[i] = 0.5f;
+					{
+						InOutArray[i] = 0.f;
 					}
 				}
 				else
 				{
 					for (int32 i = 0; i < InOutArray.Num(); ++i)
 					{
-						InOutArray[i] = (InOutArray[i] - MinVal) / Range;
+						InOutArray[i] = InOutArray[i] / Sum;
 					}
 
 					if (bInvert)
@@ -279,15 +276,13 @@ namespace PV::ImportHelper
 				}
 			};
 
-			NormalizeArray(Dots, false);
-			NormalizeArray(Lengths,false);
-			NormalizeArray(Radii, true);
+			NormalizeArray(Lengths, false);
 
 			float BestScore = -MAX_flt;
 			int32 BestIndex = INDEX_NONE;
 			for (int32 i = 0; i < NextPoints.Num(); ++i)
 			{
-				const float Score = Dots[i] + Lengths[i] + Radii[i];
+				const float Score = Dots[i] + Lengths[i];
 				if (Score > BestScore)
 				{
 					BestScore = Score;
