[PVE] Fix issue for grafter node where ramps for scale and axil angle need to be remapped to relative start and end

#Synced-CL 51894126

diff --git a/Engine/Plugins/Experimental/ProceduralVegetationEditor/Source/ProceduralVegetation/Private/Helpers/PVParametricDistributionHelper.cpp b/Engine/Plugins/Experimental/ProceduralVegetationEditor/Source/ProceduralVegetation/Private/Helpers/PVParametricDistributionHelper.cpp
--- a/Engine/Plugins/Experimental/ProceduralVegetationEditor/Source/ProceduralVegetation/Private/Helpers/PVParametricDistributionHelper.cpp
+++ b/Engine/Plugins/Experimental/ProceduralVegetationEditor/Source/ProceduralVegetation/Private/Helpers/PVParametricDistributionHelper.cpp
@@ -34,13 +34,19 @@ float PV::ParametricDistributionHelper::ComputeAttachmentScale(
 	const float PlantGradient,
 	const float BranchGradient,
 	const float LengthRatio,
-	const FParametricSettings::FScaleSettings& ScaleSettings)
+	const FParametricSettings::FScaleSettings& ScaleSettings, 
+	const float RelativeStart, 
+	const float RelativeEnd)
 {
-	const float RampLookup =
+	float RampLookup =
 		(ScaleSettings.ScaleRampBasis == FParametricSettings::EDistributionBasis::Branch)
 		? BranchGradient
 		: PlantGradient;
 
+	RampLookup = FMath::GetMappedRangeValueClamped(
+			FVector2f(RelativeStart, RelativeEnd),
+			FVector2f(0.0f, 1.0f),
+			RampLookup);
 	float ScaleRampValue = ScaleSettings.ScaleRamp->GetRichCurveConst()->Eval(RampLookup);
 
 	if (ScaleSettings.ScaleRampBasis == FParametricSettings::EDistributionBasis::Branch)
@@ -356,7 +362,9 @@ void PV::ParametricDistributionHelper::ComputeAttachmentPoints(
 					PlantGradient,
 					BranchGradient,
 					LengthRatio,
-					ScaleSettings);
+					ScaleSettings, 
+					SpacingSettings.RelativeStart, 
+					SpacingSettings.RelativeEnd);
 
 				FVector3f UpDirection;
 				FVector3f NormalDirection;
@@ -373,6 +381,7 @@ void PV::ParametricDistributionHelper::ComputeAttachmentPoints(
 					GlobalBudIndex,
 					RandomSeed,
 					AngleSettings,
+					SpacingSettings,
 					UpDirection,
 					NormalDirection);
 
@@ -634,6 +643,7 @@ void PV::ParametricDistributionHelper::ComputeRotationAnglesForAttachment(
 	const int32 BudIndex,
 	const int32 RandomSeed,
 	const FParametricSettings::FAngleSettings& AngleSettings,
+	const FParametricSettings::FSpacingSettings& SpacingSettings,
 	FVector3f& OutUpDirection,
 	FVector3f& OutNormalDirection)
 {
@@ -678,6 +688,10 @@ void PV::ParametricDistributionHelper::ComputeRotationAnglesForAttachment(
 		AxilRampLookup = PlantGradient;
 	}
 
+	AxilRampLookup = FMath::GetMappedRangeValueClamped(
+		FVector2f(SpacingSettings.RelativeStart, SpacingSettings.RelativeEnd),
+		FVector2f(0.0f, 1.0f),
+		AxilRampLookup);
 	const float AxilRampValue = AngleSettings.AxilAngleRamp->GetRichCurveConst()->Eval(AxilRampLookup);
 
 	float AxilAngleInDegrees = AngleSettings.AxilAngle * AxilRampValue;
diff --git a/Engine/Plugins/Experimental/ProceduralVegetationEditor/Source/ProceduralVegetation/Private/Helpers/PVParametricDistributionHelper.h b/Engine/Plugins/Experimental/ProceduralVegetationEditor/Source/ProceduralVegetation/Private/Helpers/PVParametricDistributionHelper.h
--- a/Engine/Plugins/Experimental/ProceduralVegetationEditor/Source/ProceduralVegetation/Private/Helpers/PVParametricDistributionHelper.h
+++ b/Engine/Plugins/Experimental/ProceduralVegetationEditor/Source/ProceduralVegetation/Private/Helpers/PVParametricDistributionHelper.h
@@ -76,7 +76,9 @@ namespace PV::ParametricDistributionHelper
 		const float PlantGradient,
 		const float BranchGradient,
 		const float LengthRatio,
-		const FParametricSettings::FScaleSettings& ScaleSettings);
+		const FParametricSettings::FScaleSettings& ScaleSettings, 
+		float RelativeStart, 
+		float RelativeEnd);
 
 	void ComputeAttachmentPoints(
 		FManagedArrayCollection& OutCollection,
@@ -126,6 +128,7 @@ namespace PV::ParametricDistributionHelper
 		const int32 BudIndex,
 		const int32 RandomSeed,
 		const FParametricSettings::FAngleSettings& AngleSettings,
+		const FParametricSettings::FSpacingSettings& SpacingSettings,
 		FVector3f& OutUpDirection,
 		FVector3f& OutNormalDirection);
 }
