[PVE] Fix issue where default value for base scale is too low for parametric settings in the foliage distributor and grafter node

#Synced-CL 51914762

diff --git a/Engine/Plugins/Experimental/ProceduralVegetationEditor/Source/ProceduralVegetation/Private/Nodes/PVFoliageDistributorSettings.h b/Engine/Plugins/Experimental/ProceduralVegetationEditor/Source/ProceduralVegetation/Private/Nodes/PVFoliageDistributorSettings.h
--- a/Engine/Plugins/Experimental/ProceduralVegetationEditor/Source/ProceduralVegetation/Private/Nodes/PVFoliageDistributorSettings.h
+++ b/Engine/Plugins/Experimental/ProceduralVegetationEditor/Source/ProceduralVegetation/Private/Nodes/PVFoliageDistributorSettings.h
@@ -200,7 +200,7 @@ struct FPVFoliageDistributionParametricSettings
 	EPVFoliageDistributionBasis ScaleRampBasis = EPVFoliageDistributionBasis::Plant;
 
 	UPROPERTY(EditAnywhere, Category="Parametric Scale Settings", meta=(PCG_Overridable, ClampMin=0.0f, ClampMax=10.0f))
-	float BaseScale = 0.2f;
+	float BaseScale = 1.0f;
 
 	UPROPERTY(EditAnywhere, Category="Parametric Scale Settings", meta=(PCG_Overridable, ClampMin=0.0f, ClampMax=1.0f))
 	float RandomizeScaleAngleMinimum = 1.0f;
diff --git a/Engine/Plugins/Experimental/ProceduralVegetationEditor/Source/ProceduralVegetation/Private/Nodes/PVGrafterSettings.h b/Engine/Plugins/Experimental/ProceduralVegetationEditor/Source/ProceduralVegetation/Private/Nodes/PVGrafterSettings.h
--- a/Engine/Plugins/Experimental/ProceduralVegetationEditor/Source/ProceduralVegetation/Private/Nodes/PVGrafterSettings.h
+++ b/Engine/Plugins/Experimental/ProceduralVegetationEditor/Source/ProceduralVegetation/Private/Nodes/PVGrafterSettings.h
@@ -125,7 +125,7 @@ struct FPVGrafterDistributionParametricSettings
 	EPVGrafterDistributionBasis ScaleRampBasis = EPVGrafterDistributionBasis::Plant;
 
 	UPROPERTY(EditAnywhere, Category="Parametric Scale Settings", meta=(PCG_Overridable, ClampMin=0.0f, ClampMax=10.0f))
-	float BaseScale = 0.2f;
+	float BaseScale = 1.0f;
 
 	UPROPERTY(EditAnywhere, Category="Parametric Scale Settings", meta=(PCG_Overridable, ClampMin=0.0f, ClampMax=1.0f))
 	float RandomizeScaleMinimum = 1.0f;
