[PVE] Fix issue where graft node does not add recursively computed descendent children

#Synced-CL 51930022

diff --git a/Engine/Plugins/Experimental/ProceduralVegetationEditor/Source/ProceduralVegetation/Private/Facades/PVBranchFacade.cpp b/Engine/Plugins/Experimental/ProceduralVegetationEditor/Source/ProceduralVegetation/Private/Facades/PVBranchFacade.cpp
--- a/Engine/Plugins/Experimental/ProceduralVegetationEditor/Source/ProceduralVegetation/Private/Facades/PVBranchFacade.cpp
+++ b/Engine/Plugins/Experimental/ProceduralVegetationEditor/Source/ProceduralVegetation/Private/Facades/PVBranchFacade.cpp
@@ -501,8 +501,8 @@ namespace PV::Facades
 	void FBranchFacade::RecomputeBranchChildren()
 	{
 		const int32 NumberOfElements = BranchParentNumbers.Num();
-		TMap<int32, TArray<int32>> BranchNumbersToBranchChildren;
-		BranchNumbersToBranchChildren.Reserve(NumberOfElements);
+		TMap<int32, TArray<int32>> ImmediateChildren;
+		ImmediateChildren.Reserve(NumberOfElements);
 
 		for (int32 i = 0; i < NumberOfElements; i++)
 		{
@@ -512,17 +512,61 @@ namespace PV::Facades
 			if (BranchParentNumber != 0)
 			[[likely]]
 			{
-				BranchNumbersToBranchChildren.FindOrAdd(BranchParentNumber).Add(BranchNumber);
+				ImmediateChildren.FindOrAdd(BranchParentNumber).Add(BranchNumber);
 			}
 		}
 
+		TMap<int32, TArray<int32>> Memo;
+		Memo.Reserve(NumberOfElements);
+
 		for (int32 i = 0; i < NumberOfElements; i++)
 		{
 			const int32 BranchNumber = BranchNumbers[i];
-			Children.ModifyAt(i, BranchNumbersToBranchChildren.Contains(BranchNumber)
-				? BranchNumbersToBranchChildren[BranchNumber]
-				: TArray<int32>());
+			TSet<int32> VisitingSet;
+			Children.ModifyAt(i, ComputeDescendantsForBranchNumber(BranchNumber, ImmediateChildren, Memo, VisitingSet));
+		}
+	}
+
+	TArray<int32> FBranchFacade::ComputeDescendantsForBranchNumber(
+		const int32 BranchNumber,
+		const TMap<int32, TArray<int32>>& ImmediateChildren,
+		TMap<int32, TArray<int32>>& Memo,
+		TSet<int32>& VisitingSet)
+	{
+		if (const TArray<int32>* Cached = Memo.Find(BranchNumber))
+		{
+			return *Cached;
+		}
+
+		// Cycle detected — return empty to break the recursion
+		if (VisitingSet.Contains(BranchNumber))
+		{
+			return TArray<int32>();
 		}
+
+		const TArray<int32>* DirectChildren = ImmediateChildren.Find(BranchNumber);
+		if (!DirectChildren)
+		{
+			Memo.Add(BranchNumber, TArray<int32>());
+			return TArray<int32>();
+		}
+
+		VisitingSet.Add(BranchNumber);
+
+		TSet<int32> DescendantSet(*DirectChildren);
+		for (const int32 Child : *DirectChildren)
+		{
+			for (const int32 Descendant : ComputeDescendantsForBranchNumber(Child, ImmediateChildren, Memo, VisitingSet))
+			{
+				DescendantSet.Add(Descendant);
+			}
+		}
+		
+		VisitingSet.Remove(BranchNumber);
+
+		TArray<int32> Descendants = DescendantSet.Array();
+		Memo.Add(BranchNumber, Descendants);
+		return Descendants;
 	}
 
 	TArray<int32> FBranchFacade::ComputeParentsForBranchNumber(const int32 BranchNumber, const TMap<int32, int32>& BranchNumbersToBranchParentNumbers,
diff --git a/Engine/Plugins/Experimental/ProceduralVegetationEditor/Source/ProceduralVegetation/Private/Implementations/PVGrafter.cpp b/Engine/Plugins/Experimental/ProceduralVegetationEditor/Source/ProceduralVegetation/Private/Implementations/PVGrafter.cpp
--- a/Engine/Plugins/Experimental/ProceduralVegetationEditor/Source/ProceduralVegetation/Private/Implementations/PVGrafter.cpp
+++ b/Engine/Plugins/Experimental/ProceduralVegetationEditor/Source/ProceduralVegetation/Private/Implementations/PVGrafter.cpp
@@ -594,14 +594,52 @@ void FPVGrafter::SetGraft(FManagedArrayCollection& OutCollection,
 	}
 	else
 	{
-		// The graft trunk is a new branch; add it as a direct child of the host branch.
-		const int32 GraftTrunkNewBranchNumber = SourceOldBranchNumbersToNewBranchNumbers[BranchNumbers[SourceGraftTrunkBranchIndex]];
+		// The graft trunk is a new branch; add it as a direct child of the host branch
+		const int32 GraftTrunkNewBranchNumber =
+			SourceOldBranchNumbersToNewBranchNumbers[BranchNumbers[SourceGraftTrunkBranchIndex]];
 		HostBranchChildren.Add(GraftTrunkNewBranchNumber);
+		
+		// Also add all descendants of the graft trunk
+		const TArray<int32>& SourceTrunkDescendants = BranchFacadeSourceGraft.GetChildren(SourceGraftTrunkBranchIndex);
+		for (const int32 SourceDescendantBranchNumber : SourceTrunkDescendants)
+		{
+			if (const int32* NewNumber = SourceOldBranchNumbersToNewBranchNumbers.Find(SourceDescendantBranchNumber))
+			{
+				HostBranchChildren.Add(*NewNumber);
+			}
+		}
 	}
 
 	BranchFacadeOut.SetChildren(InAttributes.BranchIndex, MoveTemp(HostBranchChildren));
-
 	
+	// Update all ancestors of the host branch to include the new graft branches as descendants
+	const TArray<int32>& HostAncestorBranchNumbers = BranchFacadeOut.GetParents(InAttributes.BranchIndex);
+	if (!HostAncestorBranchNumbers.IsEmpty())
+	{
+		const TManagedArray<int32>& AllBranchNumbers = BranchFacadeOut.GetBranchNumbers();
+		TMap<int32, int32> BranchNumberToIndex;
+		BranchNumberToIndex.Reserve(AllBranchNumbers.Num());
+		for (int32 i = 0; i < AllBranchNumbers.Num(); ++i)
+		{
+			BranchNumberToIndex.Add(AllBranchNumbers[i], i);
+		}
+	
+		for (const int32 AncestorBranchNumber : HostAncestorBranchNumbers)
+		{
+			if (AncestorBranchNumber == 0)
+			{
+				continue;
+			}
+			if (const int32* AncestorIndex = BranchNumberToIndex.Find(AncestorBranchNumber))
+			{
+				TArray<int32> AncestorChildren = BranchFacadeOut.GetChildren(*AncestorIndex);
+				AncestorChildren.Append(TransformedBranchNumbers);
+				BranchFacadeOut.SetChildren(*AncestorIndex, MoveTemp(AncestorChildren));
+			}
+		}
+	}
+
+
 	TArray<int32> HostAncestorPrefix = BranchFacadeOut.GetParents(InAttributes.BranchIndex);
 	HostAncestorPrefix.Add(BranchFacadeOut.GetBranchNumber(InAttributes.BranchIndex));
 	SourceIndex = 0;
diff --git a/Engine/Plugins/Experimental/ProceduralVegetationEditor/Source/ProceduralVegetation/Public/Facades/PVBranchFacade.h b/Engine/Plugins/Experimental/ProceduralVegetationEditor/Source/ProceduralVegetation/Public/Facades/PVBranchFacade.h
--- a/Engine/Plugins/Experimental/ProceduralVegetationEditor/Source/ProceduralVegetation/Public/Facades/PVBranchFacade.h
+++ b/Engine/Plugins/Experimental/ProceduralVegetationEditor/Source/ProceduralVegetation/Public/Facades/PVBranchFacade.h
@@ -109,6 +109,12 @@ namespace PV::Facades
 		
 		void RecomputeBranchChildren();
 
+		static TArray<int32> ComputeDescendantsForBranchNumber(
+			const int32 BranchNumber,
+			const TMap<int32, TArray<int32>>& ImmediateChildren,
+			TMap<int32, TArray<int32>>& Memo,
+			TSet<int32>& VisitingSet);
+
 		static TArray<int32> ComputeParentsForBranchNumber(const int32 BranchNumber, const TMap<int32, int32>& BranchNumbersToBranchParentNumbers,
 		                                                   TMap<int32, TArray<int32>>& Memo, TSet<int32>& VisitedParents);
 
