Fix Simplify Skeleton node failing when processing output from the Grower node by making foliage-related attributes optional rather than required for validation
Move foliage attribute lookup after the required attribute validation check and guard foliage processing behind a validity check

#Synced-CL 52013200

diff --git a/Engine/Plugins/Experimental/ProceduralVegetationEditor/Source/ProceduralVegetation/Private/Implementations/PVSimplifyPlantSkeleton.cpp b/Engine/Plugins/Experimental/ProceduralVegetationEditor/Source/ProceduralVegetation/Private/Implementations/PVSimplifyPlantSkeleton.cpp
--- a/Engine/Plugins/Experimental/ProceduralVegetationEditor/Source/ProceduralVegetation/Private/Implementations/PVSimplifyPlantSkeleton.cpp
+++ b/Engine/Plugins/Experimental/ProceduralVegetationEditor/Source/ProceduralVegetation/Private/Implementations/PVSimplifyPlantSkeleton.cpp
@@ -16,17 +16,14 @@ void PV::SimplifyPlantSkeleton(FManagedArrayCollection& InOutCollection, float S
 	const PV::FBranchNumberAttributeConstView BranchNumberAttribute = PV::FBranchNumberAttribute::FindAttribute(InOutCollection);
 	const PV::FBranchPointsAttributeConstView BranchPointsAttribute = PV::FBranchPointsAttribute::FindAttribute(InOutCollection);
 	const PV::FPointLengthFromRootAttributeConstView PointLengthFromRootAttribute = PV::FPointLengthFromRootAttribute::FindAttribute(InOutCollection);
-	const PV::FBranchFoliageIDsAttributeConstView BranchFoliageIDsAttribute = PV::FBranchFoliageIDsAttribute::FindAttribute(InOutCollection);
-	const PV::FFoliageLengthFromRootAttributeConstView FoliageLengthFromRootAttribute = PV::FFoliageLengthFromRootAttribute::FindAttribute(InOutCollection);
 
-	if (!PointPositionAttribute
-		|| !BranchChildrenAttribute
-		|| !BranchParentNumberAttribute
-		|| !BranchNumberAttribute
-		|| !BranchPointsAttribute
-		|| !PointLengthFromRootAttribute
-		|| !BranchFoliageIDsAttribute
-		|| !FoliageLengthFromRootAttribute)
+	if (!ValidateAttributeCollection(
+		PointPositionAttribute, 
+		BranchChildrenAttribute,
+		BranchParentNumberAttribute, 
+		BranchNumberAttribute,
+		BranchPointsAttribute,
+		PointLengthFromRootAttribute))
 	{
 		return;
 	}
@@ -65,9 +62,6 @@ void PV::SimplifyPlantSkeleton(FManagedArrayCollection& InOutCollection, float S
 	TArray<bool> BranchesToRemove;
 	BranchesToRemove.Init(false, BranchNumberAttribute.Num());
 
-	TArray<bool> FoliageInstancesToRemove;
-	FoliageInstancesToRemove.Init(false, FoliageLengthFromRootAttribute.Num());
-
 	TArray<int32> NewVertexIndices;
 	const auto SimplifySegment = [&](
 		const TArray<int32>& BranchPoints,
@@ -87,20 +81,6 @@ void PV::SimplifyPlantSkeleton(FManagedArrayCollection& InOutCollection, float S
 			for (int32 BranchPointIndex = PrevBranchPointIndex + 1; BranchPointIndex < ThisBranchPointIndex; ++BranchPointIndex)
 			{
 				PointsToRemove[BranchPoints[BranchPointIndex]] = true;
-
-				const TArray<int32> BranchSegmentFoliageEntryIds = GetBranchPointFoliageEntryIds(
-					BranchPointsAttribute,
-					PointLengthFromRootAttribute,
-					BranchFoliageIDsAttribute,
-					FoliageLengthFromRootAttribute,
-					BranchIndex,
-					BranchPointIndex
-				);
-
-				for (int32 FoliageEntryIndex : BranchSegmentFoliageEntryIds)
-				{
-					FoliageInstancesToRemove[FoliageEntryIndex] = true;
-				}
 			}
 		}
 	};
@@ -128,6 +108,7 @@ void PV::SimplifyPlantSkeleton(FManagedArrayCollection& InOutCollection, float S
 	}
 
 	const FManagedArrayCollection SourceCollection = InOutCollection;
+	TArray<bool> FoliageInstancesToRemove;
 	FPVCarve::RemoveEntriesAndRecomputeAttributes(
 		InOutCollection,
 		SourceCollection,
