[PVE] Texture Visualization Fillepd issue Fix
Jira UE-370465

#Synced-CL 51773451

diff --git a/Engine/Plugins/Experimental/ProceduralVegetationEditor/Source/ProceduralVegetationEditor/Private/Visualizations/PVDataVisualization.cpp b/Engine/Plugins/Experimental/ProceduralVegetationEditor/Source/ProceduralVegetationEditor/Private/Visualizations/PVDataVisualization.cpp
--- a/Engine/Plugins/Experimental/ProceduralVegetationEditor/Source/ProceduralVegetationEditor/Private/Visualizations/PVDataVisualization.cpp
+++ b/Engine/Plugins/Experimental/ProceduralVegetationEditor/Source/ProceduralVegetationEditor/Private/Visualizations/PVDataVisualization.cpp
@@ -578,7 +578,7 @@ void FPVDataVisualization::TexturesRenderer(const UPVData* Data, const FManagedA
 	
 	InOutParams.ManagedResources.Add(StaticMeshComponent);
 	InOutParams.Scene->AddComponent(StaticMeshComponent, FTransform(
-				FQuat::MakeFromEuler(FVector(90, 0,90)),
+				FQuat::MakeFromEuler(FVector( 90, 0,0)),
 				static_cast<FVector>(Position),
 				FVector::One() * 10
 			));
@@ -597,16 +597,16 @@ void FPVDataVisualization::TexturesRenderer(const UPVData* Data, const FManagedA
 		float DilationFactor = (GenerationUV.OffsetXStart * Scale) - PreviousEnd;
 		const FLinearColor Color = PV::Utilities::GetRandomHueColor(NumGenerations > 0 ? Generation / static_cast<float>(NumGenerations) : Generation);
 		LineBatchComponent->AddLine(
-			FVector(0, PreviousEnd - (Scale/ 2.0), Scale + 100 ),
-			FVector(0, ((GenerationUV.OffsetXEnd * Scale) + DilationFactor) - (Scale/ 2.0), Scale + 100),
+			FVector(PreviousEnd - (Scale/ 2.0), 0, Scale + 100 ),
+			FVector(((GenerationUV.OffsetXEnd * Scale) + DilationFactor) - (Scale/ 2.0), 0, Scale + 100),
 			Color,
 			SDPG_World,
 			EPointDrawSettings::None
 		);
 
 		LineBatchComponent->AddLine(
-			FVector(0, PreviousEnd - (Scale/ 2.0), 0 ),
-			FVector(0, ((GenerationUV.OffsetXEnd * Scale) + DilationFactor) - (Scale/ 2.0), 0),
+			FVector(PreviousEnd - (Scale/ 2.0), 0, 0 ),
+			FVector(((GenerationUV.OffsetXEnd * Scale) + DilationFactor) - (Scale/ 2.0), 0, 0),
 			Color,
 			SDPG_World,
 			EPointDrawSettings::None
