Fix Object interaction Trim Outside leaving floating points when trimming a skeleton from StaticMesh importer node

- Fix BranchChildren attribute not being populated with all descendants when generating growth data from importer nodes by performing a reversed tree traversal to propagate children recursively
- Replace the recursive MarkSubtreeForRemoval helper in the Carve node with an inline flat iteration over BranchChildren, which is now sufficient since BranchChildren contains all descendants

#Synced-CL 52013200

diff --git a/Engine/Plugins/Experimental/ProceduralVegetationEditor/Source/ProceduralVegetation/Private/Helpers/PVImportHelpers.cpp b/Engine/Plugins/Experimental/ProceduralVegetationEditor/Source/ProceduralVegetation/Private/Helpers/PVImportHelpers.cpp
--- a/Engine/Plugins/Experimental/ProceduralVegetationEditor/Source/ProceduralVegetation/Private/Helpers/PVImportHelpers.cpp
+++ b/Engine/Plugins/Experimental/ProceduralVegetationEditor/Source/ProceduralVegetation/Private/Helpers/PVImportHelpers.cpp
@@ -7,6 +7,7 @@
 #include "GeometryCollection/GeometryCollection.h"
 
 #include "Helpers/PVAttributesHelper.h"
+#include "Helpers/PVPlantTraversalHelper.h"
 #include "Helpers/PVDynamicMeshHelpers.h"
 
 #include "Engine/StaticMesh.h"
@@ -895,6 +896,13 @@ void PV::ImportHelper::FillAttributesFromBranchHierarchy(
 		return BranchIndex + 1; 
 	}; 
 
+	const static auto BranchNumberToBranchIndex = [](int32 BranchNumber)
+	{ 
+		// Note that BranchNumberToBranchIndex only works within the local scope of FillAttributesFromBranchHierarchy due to how
+		// we compute the BranchNumber in BranchIndexToBranchNumber
+		return BranchNumber - 1;
+	}; 
+
 	using FFillBranchesRecursive = TFunction<void(const FPVBranchHierarchyDescription&, int32, int32, int32, int32, int32)>;
 
 	const FFillBranchesRecursive FillBranchesRecursive = [&](
@@ -991,6 +999,29 @@ void PV::ImportHelper::FillAttributesFromBranchHierarchy(
 		AttributePointIndexOffset += BranchHierarchy.Points.Num();
 		PlantNumber += 1;
 	}
+	
+	// Fill all descendants in the BranchChildren attribute by walking the tree hierarchy backwards
+	// and appending the children of our children. This only works since we know for certain
+	// all branches only contain the immediate children at this stage.
+	PV::PlantTraversalHelper::RecursiveWalkBranches_Reversed(
+		BranchParentNumberAttribute,
+		BranchChildrenAttribute,
+		BranchNumberAttribute,
+		[&](int32 BranchIndex)
+		{
+			const TArray<int32> ImmediateChildBranchNumbers = BranchChildrenAttribute[BranchIndex]; // We know these are the immediate children due to not filling the descendants above.
+			for (const int32 ChildBranchNumber : ImmediateChildBranchNumbers)
+			{
+				const int32 ChildBranchIndex = BranchNumberToBranchIndex(ChildBranchNumber);
+				if (ensure(BranchChildrenAttribute.IsValidIndex(ChildBranchIndex)))
+				{
+					BranchChildrenAttribute[BranchIndex].Append(BranchChildrenAttribute[ChildBranchIndex]);
+				}
+			}
+
+			return PV::PlantTraversalHelper::EForEachResult::Continue;
+		}
+	);
 }
 
 void PV::ImportHelper::ConnectDisconnectedBranches(FPVBranchHierarchyDescription& InOutBranchHierarchy)
