[PVE] Texture Setup Fixes
Fixed Confusing Folder paths
Moved the Bake Texture path in Bake Textures Category
Made the Path relative to Game Dir
Added validation for baking
Black Textures when selecting image file from a folder with in the project path

#Synced-CL 51773451

diff --git a/Engine/Plugins/Experimental/ProceduralVegetationEditor/Source/ProceduralVegetation/Private/Implementations/PVTextureSetup.cpp b/Engine/Plugins/Experimental/ProceduralVegetationEditor/Source/ProceduralVegetation/Private/Implementations/PVTextureSetup.cpp
--- a/Engine/Plugins/Experimental/ProceduralVegetationEditor/Source/ProceduralVegetation/Private/Implementations/PVTextureSetup.cpp
+++ b/Engine/Plugins/Experimental/ProceduralVegetationEditor/Source/ProceduralVegetation/Private/Implementations/PVTextureSetup.cpp
@@ -10,6 +10,7 @@
 #include "AssetRegistry/AssetRegistryModule.h"
 #include "Engine/Texture2D.h"
 #include "Engine/TextureRenderTarget2D.h"
+#include "Helpers/PVUtilities.h"
 #include "Materials/MaterialInstanceConstant.h"
 #include "Misc/PackageName.h"
 #include "UObject/SavePackage.h"
@@ -336,17 +337,23 @@ void FPVTextureSetupImplementation::CreateTrimSheetTexture( FPVTextureSetupParam
 	TextureSetupInfo.PreviewChannelName = StaticEnum<EPVTextureChannel>()->GetNameStringByValue((int64)Params.PreviewChannel);
 }
 
-UTexture2D* FPVTextureSetupImplementation::ConvertRenderTargetToTexture2D(UTextureRenderTarget2D* InRenderTarget, const FString& Path)
+UTexture2D* FPVTextureSetupImplementation::ConvertRenderTargetToTexture2D(UTextureRenderTarget2D* InRenderTarget, const FString& Path, FString& OutError)
 {
 	if (!InRenderTarget)
 	{
+		OutError = "Cannot convert a null render target to a Texture2D.";
 		return nullptr;
 	}
 	
 #if WITH_EDITOR
 	FString AssetName = InRenderTarget->GetName();
 	FString PackageName = Path / AssetName;
-	FString ObjectPath = PackageName + TEXT(".") + AssetName; 
+	FString ObjectPath = PackageName + TEXT(".") + AssetName;
+
+	if (!PV::Utilities::ValidateAssetPathAndName(AssetName, Path, UTexture2D::StaticClass(), OutError))
+	{
+		return nullptr;
+	}
 
 	UTexture2D* Texture = LoadObject<UTexture2D>(nullptr, ObjectPath);
 	UPackage* Package = nullptr;
diff --git a/Engine/Plugins/Experimental/ProceduralVegetationEditor/Source/ProceduralVegetation/Private/Implementations/PVTextureSetup.h b/Engine/Plugins/Experimental/ProceduralVegetationEditor/Source/ProceduralVegetation/Private/Implementations/PVTextureSetup.h
--- a/Engine/Plugins/Experimental/ProceduralVegetationEditor/Source/ProceduralVegetation/Private/Implementations/PVTextureSetup.h
+++ b/Engine/Plugins/Experimental/ProceduralVegetationEditor/Source/ProceduralVegetation/Private/Implementations/PVTextureSetup.h
@@ -34,5 +34,5 @@ struct FPVTextureSetupImplementation
 	static TObjectPtr<UTexture>  GetChannelTextureForMaterialParam(FPVTextureSetupInfo& TextureSetupInfo, FString MaterialParamName);
 	static void SetupMaterial(TObjectPtr<UMaterialInterface> Material, FPVTextureSetupInfo& TextureSetupInfo);
 	static void CreateTrimSheetTexture( FPVTextureSetupParams& Params, FPVTextureSetupInfo& TextureSetupInfo);
-	static UTexture2D* ConvertRenderTargetToTexture2D(UTextureRenderTarget2D* InRenderTarget, const FString& Path);
+	static UTexture2D* ConvertRenderTargetToTexture2D(UTextureRenderTarget2D* InRenderTarget, const FString& Path, FString& OutError);
 };
diff --git a/Engine/Plugins/Experimental/ProceduralVegetationEditor/Source/ProceduralVegetation/Private/Nodes/PVTextureSetupSettings.cpp b/Engine/Plugins/Experimental/ProceduralVegetationEditor/Source/ProceduralVegetation/Private/Nodes/PVTextureSetupSettings.cpp
--- a/Engine/Plugins/Experimental/ProceduralVegetationEditor/Source/ProceduralVegetation/Private/Nodes/PVTextureSetupSettings.cpp
+++ b/Engine/Plugins/Experimental/ProceduralVegetationEditor/Source/ProceduralVegetation/Private/Nodes/PVTextureSetupSettings.cpp
@@ -11,7 +11,7 @@ class UTexture2D;
 
 UPVTextureSetupSettings::UPVTextureSetupSettings()
 {
-	TextureSetupParams.ContentBrowserFolder.Path = "/Game/ProceduralVegetation/";
+	TextureSetupParams.BakeTextureFolder.Path = "/Game/ProceduralVegetation";
 }
 
 #if WITH_EDITOR
@@ -38,6 +38,7 @@ void UPVTextureSetupSettings::PostEditChangeChainProperty(struct FPropertyChange
 	{
 		FPVTextureSetupImplementation::CreateTrimSheetTexture(TextureSetupParams, TextureSetupInfo);
 		bTexturesBaked = false;
+		BakingErrors.Empty();
 	}
 }
 
@@ -86,17 +87,28 @@ FPCGDataTypeIdentifier UPVTextureSetupSettings::GetOutputPinTypeIdentifier() con
 void UPVTextureSetupSettings::BakeTextures()
 {
 #if WITH_EDITOR
+	BakingErrors.Empty();
 	for (auto& Channel : TextureSetupInfo.Channels)
 	{
+		FString Error;
 		if (UTexture* ConvertedTexture = Cast<UTexture>(FPVTextureSetupImplementation::ConvertRenderTargetToTexture2D(Cast<UTextureRenderTarget2D>(Channel.Value),
-				TextureSetupParams.ContentBrowserFolder.Path)))
+				TextureSetupParams.BakeTextureFolder.Path, Error)))
 		{
 			Channel.Value = ConvertedTexture;
 		}
-	}
 
-	FPVTextureSetupImplementation::SetupMaterial(TextureSetupParams.Material, TextureSetupInfo);
-	bTexturesBaked = true;
+		if (!Error.IsEmpty())
+		{
+			Error = Error.Replace(TEXT("Export Path"), TEXT("Bake Path"));
+			BakingErrors += Error + "\n";
+		}
+	}
+	
+	if (BakingErrors.IsEmpty())
+	{
+		FPVTextureSetupImplementation::SetupMaterial(TextureSetupParams.Material, TextureSetupInfo);
+		bTexturesBaked = true;
+	}
 #else
 	UE_LOG(LogTemp, Display, TEXT("Baking Textures is not supported"));
 #endif
@@ -131,6 +143,11 @@ bool FPVTextureSetupElement::ExecuteInternal(FPCGContext* InContext) const
 	{
 		PCGLog::LogWarningOnGraph(FText::FromString(TEXT("Generations are not fitting in the texture, consider scaling down the XScale values in generation settings.")), InContext);
 	}
+
+	if (!Settings->BakingErrors.IsEmpty())
+	{
+		PCGLog::LogWarningOnGraph(FText::FromString(FString::Printf(TEXT("Texture Baking errors, %s"), *Settings->BakingErrors)), InContext);
+	}
 	
 	return true;
 }
diff --git a/Engine/Plugins/Experimental/ProceduralVegetationEditor/Source/ProceduralVegetation/Public/Nodes/PVTextureSetupSettings.h b/Engine/Plugins/Experimental/ProceduralVegetationEditor/Source/ProceduralVegetation/Public/Nodes/PVTextureSetupSettings.h
--- a/Engine/Plugins/Experimental/ProceduralVegetationEditor/Source/ProceduralVegetation/Public/Nodes/PVTextureSetupSettings.h
+++ b/Engine/Plugins/Experimental/ProceduralVegetationEditor/Source/ProceduralVegetation/Public/Nodes/PVTextureSetupSettings.h
@@ -43,6 +43,9 @@ public:
 	UPROPERTY()
 	FPVTextureSetupInfo TextureSetupInfo;
 
+	UPROPERTY()
+	FString BakingErrors;
+
 private:
 	UPROPERTY()
 	bool bTexturesBaked = false;
diff --git a/Engine/Plugins/Experimental/ProceduralVegetationEditor/Source/ProceduralVegetation/Public/Params/PVTextureSetupParams.h b/Engine/Plugins/Experimental/ProceduralVegetationEditor/Source/ProceduralVegetation/Public/Params/PVTextureSetupParams.h
--- a/Engine/Plugins/Experimental/ProceduralVegetationEditor/Source/ProceduralVegetation/Public/Params/PVTextureSetupParams.h
+++ b/Engine/Plugins/Experimental/ProceduralVegetationEditor/Source/ProceduralVegetation/Public/Params/PVTextureSetupParams.h
@@ -30,7 +30,7 @@ struct FPVTextureChannelParams
 	UPROPERTY(EditAnywhere, Category=Channel, meta=(EditCondition="Channel == EPVTextureChannel::Custom"))
 	FString ChannelName;
 
-	UPROPERTY(EditAnywhere, Category= Channel, meta = (AllowPrivateAccess = true, FilePathFilter = "Image Files (*.png;*.PNG;*.jpg;*.jpeg;*.exr;)|*.png;*.PNG;*.jpg;*.jpeg;*.exr;", RelativeToGameDir))
+	UPROPERTY(EditAnywhere, Category= Channel, meta = (AllowPrivateAccess = true, FilePathFilter = "Image Files (*.png;*.PNG;*.jpg;*.jpeg;*.exr;)|*.png;*.PNG;*.jpg;*.jpeg;*.exr;"))
 	FFilePath TextureFilePath;
 
 	UPROPERTY(VisibleAnywhere, Category=Channel)
@@ -49,10 +49,7 @@ USTRUCT()
 struct FPVTextureGenerationParams
 {
 	GENERATED_BODY()
-
-	UPROPERTY(EditAnywhere, Category= Generation)
-	FDirectoryPath BakeTextureFolder;
-
+	
 	UPROPERTY(EditAnywhere, Category= Generation)
 	float Dilation = 0.1f;
 
@@ -84,13 +81,13 @@ struct FPVTextureSetupParams
 {
 	GENERATED_BODY()
 
-	UPROPERTY(EditAnywhere, Category= NoCategory)
+	UPROPERTY(EditAnywhere, Category= "Baked Textures")
 	FString AssetNamePrefix;
 
-	UPROPERTY(EditAnywhere, Category= NoCategory)
-	FDirectoryPath ContentBrowserFolder;
+	UPROPERTY(EditAnywhere, Category= "Baked Textures", meta=(ContentDir, RelativeToGameDir, ToolTip="Select a folder from content browser to save the baked textures."))
+	FDirectoryPath BakeTextureFolder;
 
-	UPROPERTY( EditAnywhere, Category= NoCategory)
+	UPROPERTY( EditAnywhere, Category= "Baked Textures")
 	bool bVirtualTexture = false;
 
 	UPROPERTY( EditAnywhere, Category= NoCategory)
