[PVE] Fix carve output discrepancy between fused and split points

- Normalization: moved CarveValue min/max and normalization into a single
pre-pass to prevent fused points being normalized twice with an
inconsistent range.

- MarkSubtreeForRemoval: added bSkipFirstPoint to prevent fused tip points
being marked for removal by child subtree removal. Added
PointsToRemove[LastPointIndex] = false as safety net for recursive cases.

- UpdatePointScales: added InSkipFirstPoint to skip BranchPoints[0] for
non-trunk branches in fused mode, preventing double PScale writes at the
attachment point.

- CarveFromBottom: added VisitedPoints TSet to ensure fused points shared
across branches are shifted and recomputed exactly once.

#Synced-CL 52095318

diff --git a/Engine/Plugins/Experimental/ProceduralVegetationEditor/Source/ProceduralVegetation/Private/Implementations/PVCarve.cpp b/Engine/Plugins/Experimental/ProceduralVegetationEditor/Source/ProceduralVegetation/Private/Implementations/PVCarve.cpp
--- a/Engine/Plugins/Experimental/ProceduralVegetationEditor/Source/ProceduralVegetation/Private/Implementations/PVCarve.cpp
+++ b/Engine/Plugins/Experimental/ProceduralVegetationEditor/Source/ProceduralVegetation/Private/Implementations/PVCarve.cpp
@@ -3,7 +3,6 @@
 #include "PVCarve.h"
 
 #include "Facades/PVAttributesNames.h"
-#include "Facades/PVFoliageFacade.h"
 #include "Facades/PVPlantFacade.h"
 #include "Facades/PVPointFacade.h"
 #include "Facades/PVTreeFacade.h"
@@ -13,29 +12,28 @@ namespace
 	static void MarkSubtreeForRemoval(
 		const int32 BranchIndex,
 		const PV::Facades::FBranchFacade& BranchFacadeSource,
-		const PV::Facades::FFoliageFacade& FoliageFacadeSource,
 		const TMap<int32, int32>& BranchNumbersToBranchIDs,
 		TArray<bool>& BranchesToRemove,
 		TArray<bool>& PointsToRemove,
-		TArray<bool>& FoliageInstancesToRemove)
+		const bool bSkipFirstPoint = false)
 	{
 		BranchesToRemove[BranchIndex] = true;
-		for (const int32& P : BranchFacadeSource.GetPoints(BranchIndex))
-		{
-			PointsToRemove[P] = true;
-		}
-		for (const int32& FId : FoliageFacadeSource.GetFoliageEntryIdsForBranch(BranchIndex))
+		const TArray<int32>& BranchPoints = BranchFacadeSource.GetPoints(BranchIndex);
+
+		const int32 StartingIndex = bSkipFirstPoint? 1: 0;
+		
+		for (int32 i = StartingIndex; i < BranchPoints.Num(); ++i)
 		{
-			FoliageInstancesToRemove[FId] = true;
+			PointsToRemove[BranchPoints[i]] = true;
 		}
 
-		for (const int32& ChildBranchNumber : BranchFacadeSource.GetChildren(BranchIndex))
+		for (const int32& ChildBranchNumber : BranchFacadeSource.GetImmediateChildren(BranchIndex))
 		{
 			const int32* ChildIndexPtr = BranchNumbersToBranchIDs.Find(ChildBranchNumber);
 			if (ChildIndexPtr && !BranchesToRemove[*ChildIndexPtr])
 			{
-				MarkSubtreeForRemoval(*ChildIndexPtr, BranchFacadeSource, FoliageFacadeSource,
-					BranchNumbersToBranchIDs, BranchesToRemove, PointsToRemove, FoliageInstancesToRemove);
+				MarkSubtreeForRemoval(*ChildIndexPtr, BranchFacadeSource, BranchNumbersToBranchIDs, BranchesToRemove,
+					PointsToRemove);
 			}
 		}
 	}
@@ -61,7 +59,7 @@ void FPVCarve::ApplyCarve(FManagedArrayCollection& OutCollection, const FManaged
 
 void FPVCarve::UpdatePointScales(PV::Facades::FPointFacade& PointFacadeOut, const PV::Facades::FPointFacade& PointFacadeSource,
                                  const TArray<int>& BranchPoints, const float LastPointScale, const int32 LastPointIndex,
-                                 const float CarveRatio, const int32 EndIndex, const float FirstPointTargetPScale)
+                                 const float CarveRatio, const int32 EndIndex, const float FirstPointTargetPScale, const bool InSkipFirstPoint)
 {
 	TRACE_CPUPROFILER_EVENT_SCOPE(PVCarve::UpdatePointScales);
 
@@ -71,7 +69,9 @@ void FPVCarve::UpdatePointScales(PV::Facades::FPointFacade& PointFacadeOut, cons
 	const float OldMin = PointFacadeSource.GetPointScale(LastPointIndex);
 	const float NewMin = LastPointScale * Ratio;
 
-	for (int32 i = 0; i <= EndIndex; ++i)
+	const int32 StartingIndex = InSkipFirstPoint? 1: 0;
+	
+	for (int32 i = StartingIndex; i <= EndIndex; ++i)
 	{
 		const float UpdatedScale = FMath::GetMappedRangeValueClamped(
 			FVector2f(OldMin, OldMax),
@@ -176,7 +176,6 @@ void FPVCarve::CarveFromTop(FManagedArrayCollection& OutCollection, const FManag
 
 	const PV::Facades::FBranchFacade BranchFacadeSource(SourceCollection);
 	const PV::Facades::FPointFacade PointFacadeSource(SourceCollection);
-	const PV::Facades::FFoliageFacade FoliageFacadeSource(SourceCollection);
 	const PV::Facades::FPlantFacade PlantFacadeSource(SourceCollection);
 
 	PV::Facades::FBranchFacade BranchFacadeOut(OutCollection);
@@ -193,14 +192,12 @@ void FPVCarve::CarveFromTop(FManagedArrayCollection& OutCollection, const FManag
 	PointsToRemove.Init(false, PointFacadeSource.GetElementCount());
 	TArray<bool> BranchesToRemove;
 	BranchesToRemove.Init(false, NumOfBranchesInSource);
-	TArray<bool> FoliageInstancesToRemove;
-	FoliageInstancesToRemove.Init(false, FoliageFacadeSource.NumFoliageEntries());
+	
 
 	TMap<int32, int32> BranchNumbersToBranchIDs;
 	TMap<int32, float> BranchNumbersToLengthFromRoots;
 	ComputeMetadata(BranchNumbersToBranchIDs, BranchNumbersToLengthFromRoots, BranchFacadeSource, PointFacadeSource);
 
-	// Compute carve values according to carve basis metric
 	TArray<float> CarveValues;
 	CarveValues.Reserve(PointFacadeSource.GetElementCount());
 
@@ -226,33 +223,27 @@ void FPVCarve::CarveFromTop(FManagedArrayCollection& OutCollection, const FManag
 		}
 	}
 
+	const float MinCarveValue = *Algo::MinElement(CarveValues);
+	const float MaxCarveValue = *Algo::MaxElement(CarveValues);
+
 	FVector2f MappedRange(0.0f, 1.0f);
 	if (CarveBasis == ECarveBasis::Radius)
 	{
 		MappedRange = FVector2f(1.0f, 0.0f);
 	}
 
-	for (const int32& PlantNumber : PlantFacadeSource.GetPlantNumbers())
+	for (int32 i = 0; i < PointFacadeSource.GetElementCount(); ++i)
 	{
-		float MinimumCarveValue = TNumericLimits<float>::Max();
-		float MaximumCarveValue = TNumericLimits<float>::Min();
+		CarveValues[i] = FMath::GetMappedRangeValueClamped(
+			FVector2f(MinCarveValue, MaxCarveValue),
+			MappedRange,
+			CarveValues[i]);
+	}
 
+	for (const int32& PlantNumber : PlantFacadeSource.GetPlantNumbers())
+	{
 		for (const int32& BranchIndex : PlantFacadeSource.GetBranchIndices(PlantNumber))
 		{
-			for (const int32& PointIndex : BranchFacadeSource.GetPoints(BranchIndex))
-			{
-				MinimumCarveValue = FMath::Min(MinimumCarveValue, CarveValues[PointIndex]);
-				MaximumCarveValue = FMath::Max(MaximumCarveValue, CarveValues[PointIndex]);
-			}
-
-			for (const int32& PointIndex : BranchFacadeSource.GetPoints(BranchIndex))
-			{
-				CarveValues[PointIndex] = FMath::GetMappedRangeValueClamped(
-					FVector2f(MinimumCarveValue, MaximumCarveValue),
-					MappedRange,
-					CarveValues[PointIndex]);
-			}
-
 			// Compute branches, foliage instances and points to remove
 			if (BranchesToRemove[BranchIndex])
 			{
@@ -268,65 +259,51 @@ void FPVCarve::CarveFromTop(FManagedArrayCollection& OutCollection, const FManag
 				continue;
 			}
 
-			// Compute PScale for the point in parent branch that is closest to the first point
+			//  Find Source point Pscales
 			// (For non-trunk branches)
-			float FirstPointTargetPScale = PointFacadeSource.GetPointScale(BranchPoints[0]);
-			if (!PlantFacadeSource.IsTrunkIndex(BranchIndex)) [[likely]]
-			{
-				const FVector3f& FirstPointPosition = PointFacadeSource.GetPosition(BranchPoints[0]);
-
-				const int32 ParentBranchNumber = BranchFacadeSource.GetParentBranchNumber(BranchIndex);
-				const int32 ParentBranchIndex = BranchNumbersToBranchIDs[ParentBranchNumber];
-				const TArray<int32>& ParentBranchPoints = BranchFacadeSource.GetPoints(ParentBranchIndex);
+			const int32 SourceBudNumber = BranchFacadeSource.GetBranchSourceBudNumber(BranchIndex);
 
-				int32 ClosestPointIndex = 0;
-				float MinimumDistance = FLT_MAX;
-				for (int32 PointIndex = 0; PointIndex < ParentBranchPoints.Num(); ++PointIndex)
+			const int32* RootPointIndexPtr = BranchPoints.FindByPredicate([&PointFacadeSource, &SourceBudNumber](const int32& PointIndex)
 				{
-					const FVector3f& ParentPointPosition = PointFacadeSource.GetPosition(ParentBranchPoints[PointIndex]);
-					if (const float Distance = FVector3f::Distance(FirstPointPosition, ParentPointPosition); Distance < MinimumDistance)
-					{
-						MinimumDistance = Distance;
-						ClosestPointIndex = PointIndex;
-					}
-				}
-
-				FirstPointTargetPScale = PointFacadeOut.GetPointScale(ParentBranchPoints[ClosestPointIndex]);
-			}
-
+					return PointFacadeSource.GetBudNumber(PointIndex) == SourceBudNumber;
+				});
+
+			// Safety check
+			checkf(RootPointIndexPtr != nullptr, TEXT("Crave: No point with matching bud number found"));
+			check(BranchPoints[0] == *RootPointIndexPtr);
+			
+			const float FirstPointTargetPScale = PointFacadeSource.GetPointScale(*RootPointIndexPtr);
+			const bool IsTrunk = PlantFacadeSource.IsTrunkIndex(BranchIndex);
+			
 			int EndIndex = 1;
 			while (EndIndex < NumOfPointsInCurrentBranch && (CarveValues[BranchPoints[EndIndex]] < CarveDistance + 0.0001f))
 			{
 				EndIndex++;
 			}
 			EndIndex--;
-
+			
 			if (EndIndex == NumOfPointsInCurrentBranch - 1)
 			{
 				const int32 LastPointIndex = BranchPoints[EndIndex];
 				const float LastPointScale = PointFacadeSource.GetPointScale(LastPointIndex);
 				UpdatePointScales(PointFacadeOut, PointFacadeSource, BranchPoints, LastPointScale, LastPointIndex, 1.0f, EndIndex,
-					FirstPointTargetPScale);
+					FirstPointTargetPScale, !IsTrunk);
 				continue;
 			}
 
 			// Remove all branches in the subtree beyond the last kept point
-			for (const int32& BranchNumber : BranchFacadeSource.GetChildren(BranchIndex))
+			for (const int32& BranchNumber : BranchFacadeSource.GetImmediateChildren(BranchIndex))
 			{
 				if (BranchNumbersToLengthFromRoots.Contains(BranchNumber) && BranchNumbersToLengthFromRoots[BranchNumber] > PointFacadeSource.
 					GetLengthFromRoot(BranchPoints[EndIndex]))
 				{
 					const int32 Index = BranchNumbersToBranchIDs[BranchNumber];
-					MarkSubtreeForRemoval(Index, BranchFacadeSource, FoliageFacadeSource,
-						BranchNumbersToBranchIDs, BranchesToRemove, PointsToRemove, FoliageInstancesToRemove);
-				}
-			}
+					const TArray<int32>& ChildBranchPoints = BranchFacadeSource.GetPoints(Index);
+					const bool bSkipFirstPoint = ChildBranchPoints.Num() > 0
+						&& BranchPoints.IsValidIndex(EndIndex + 1)
+						&& ChildBranchPoints[0] == BranchPoints[EndIndex + 1];
 
-			for (const int32& FoliageIndex : FoliageFacadeSource.GetFoliageEntryIdsForBranch(BranchIndex))
-			{
-				if (FoliageFacadeSource.GetLengthFromRoot(FoliageIndex) > PointFacadeSource.GetLengthFromRoot(BranchPoints[EndIndex]))
-				{
-					FoliageInstancesToRemove[FoliageIndex] = true;
+					MarkSubtreeForRemoval(Index, BranchFacadeSource, BranchNumbersToBranchIDs, BranchesToRemove, PointsToRemove, bSkipFirstPoint);
 				}
 			}
 
@@ -365,15 +342,15 @@ void FPVCarve::CarveFromTop(FManagedArrayCollection& OutCollection, const FManag
 			PointFacadeOut.ModifyLengthFromRoots()[LastPointIndex] = LengthFromRoot;
 			PointFacadeOut.ModifyLengthFromSeeds()[LastPointIndex] = LengthFromRoot;
 			PointFacadeOut.ModifyPointScales()[LastPointIndex] = PointScale;
-
+			PointsToRemove[LastPointIndex] = false;
 
 			// Update point scales
 			const float CarveRatio = (EndIndex + BlendValue) / NumOfPointsInCurrentBranch;
 			UpdatePointScales(PointFacadeOut, PointFacadeSource, BranchPoints, PointScale, LastPointIndex, CarveRatio, EndIndex,
-				FirstPointTargetPScale);
+				FirstPointTargetPScale, !IsTrunk);
 		}
 	}
-
+	TArray<bool> FoliageInstancesToRemove;
 	RemoveEntriesAndRecomputeAttributes(OutCollection, SourceCollection, PointsToRemove, BranchesToRemove, FoliageInstancesToRemove);
 }
 
@@ -388,7 +365,6 @@ void FPVCarve::CarveFromBottom(FManagedArrayCollection& OutCollection, const FMa
 
 	PV::Facades::FBranchFacade BranchFacadeOut(OutCollection);
 	PV::Facades::FPointFacade PointFacadeOut(OutCollection);
-	PV::Facades::FFoliageFacade FoliageFacadeOut(OutCollection);
 
 	if (PointFacadeSource.GetElementCount() == 0 || BranchFacadeSource.GetElementCount() == 0)
 	{
@@ -399,8 +375,6 @@ void FPVCarve::CarveFromBottom(FManagedArrayCollection& OutCollection, const FMa
 	PointsToRemove.Init(false, PointFacadeSource.GetElementCount());
 	TArray<bool> BranchesToRemove;
 	BranchesToRemove.Init(false, BranchFacadeSource.GetElementCount());
-	TArray<bool> FoliageInstancesToRemove;
-	FoliageInstancesToRemove.Init(false, FoliageFacadeSource.NumFoliageEntries());
 
 	TMap<int32, int32> BranchNumbersToBranchIDs;
 	TMap<int32, float> BranchNumbersToLengthFromRoots;
@@ -435,21 +409,12 @@ void FPVCarve::CarveFromBottom(FManagedArrayCollection& OutCollection, const FMa
 		EndIndex--;
 
 		// Remove all branches in the subtree before current point
-		for (const int32& BranchNumber : BranchFacadeSource.GetChildren(TrunkIndex))
+		for (const int32& BranchNumber : BranchFacadeSource.GetImmediateChildren(TrunkIndex))
 		{
 			if (BranchNumbersToLengthFromRoots.Contains(BranchNumber) && BranchNumbersToLengthFromRoots[BranchNumber] < CarveLengthFromRoot)
 			{
 				const int32 Index = BranchNumbersToBranchIDs[BranchNumber];
-				MarkSubtreeForRemoval(Index, BranchFacadeSource, FoliageFacadeSource,
-					BranchNumbersToBranchIDs, BranchesToRemove, PointsToRemove, FoliageInstancesToRemove);
-			}
-		}
-
-		for (const int32& FoliageIndex : FoliageFacadeSource.GetFoliageEntryIdsForBranch(TrunkIndex))
-		{
-			if (FoliageFacadeSource.GetLengthFromRoot(FoliageIndex) < CarveLengthFromRoot)
-			{
-				FoliageInstancesToRemove[FoliageIndex] = true;
+				MarkSubtreeForRemoval(Index, BranchFacadeSource, BranchNumbersToBranchIDs, BranchesToRemove,PointsToRemove);
 			}
 		}
 
@@ -488,6 +453,8 @@ void FPVCarve::CarveFromBottom(FManagedArrayCollection& OutCollection, const FMa
 
 		// Move all points down and recompute length from roots
 		const FVector3f Delta(0.0f, 0.0f, Position.Z);
+		
+		TSet<int32> VisitedPoints;
 
 		for (const int32& BranchIndex : PlantFacadeSource.GetBranchIndices(PlantNumber))
 		{
@@ -498,25 +465,22 @@ void FPVCarve::CarveFromBottom(FManagedArrayCollection& OutCollection, const FMa
 
 			for (const int32& PointIndex : BranchFacadeOut.GetPoints(BranchIndex))
 			{
+				if (VisitedPoints.Contains(PointIndex))
+					continue;
+
 				PointFacadeOut.ModifyPositions()[PointIndex] = PointFacadeOut.GetPosition(PointIndex) - Delta;
 
 				PointFacadeOut.ModifyLengthFromRoots()[PointIndex] = FMath::Max(PointFacadeOut.GetLengthFromRoot(PointIndex) - CarveLengthFromRoot,
 					0.0f);
 				PointFacadeOut.ModifyLengthFromSeeds()[PointIndex] = FMath::Max(PointFacadeOut.GetLengthFromSeed(PointIndex) - CarveLengthFromRoot,
 					0.0f);
-			}
 
-			for (const int32& FoliageIndex : FoliageFacadeOut.GetFoliageEntryIdsForBranch(BranchIndex))
-			{
-				const FVector3f UpdatedPivotPoint = FoliageFacadeOut.GetPivotPoint(FoliageIndex) - Delta;
-				FoliageFacadeOut.SetPivotPoint(FoliageIndex, UpdatedPivotPoint);
-
-				const float UpdatedLengthFromRoot = FMath::Max(FoliageFacadeOut.GetLengthFromRoot(FoliageIndex) - CarveLengthFromRoot);
-				FoliageFacadeOut.SetLengthFromRoot(FoliageIndex, UpdatedLengthFromRoot);
+				VisitedPoints.Add(PointIndex);
 			}
 		}
 	}
-
+	
+	TArray<bool> FoliageInstancesToRemove;
 	RemoveEntriesAndRecomputeAttributes(OutCollection, SourceCollection, PointsToRemove, BranchesToRemove, FoliageInstancesToRemove);
 }
 
@@ -568,4 +532,4 @@ void FPVCarve::RemoveEntriesAndRecomputeAttributes(FManagedArrayCollection& OutC
 	}
 
 	RecomputeAttributes(OutCollection, SourceCollection, BranchesNewIDsToOldIDs);
-}
+}
\ No newline at end of file
diff --git a/Engine/Plugins/Experimental/ProceduralVegetationEditor/Source/ProceduralVegetation/Private/Implementations/PVCarve.h b/Engine/Plugins/Experimental/ProceduralVegetationEditor/Source/ProceduralVegetation/Private/Implementations/PVCarve.h
--- a/Engine/Plugins/Experimental/ProceduralVegetationEditor/Source/ProceduralVegetation/Private/Implementations/PVCarve.h
+++ b/Engine/Plugins/Experimental/ProceduralVegetationEditor/Source/ProceduralVegetation/Private/Implementations/PVCarve.h
@@ -38,7 +38,7 @@ struct FPVCarve
 
 	static void UpdatePointScales(PV::Facades::FPointFacade& PointFacadeOut, const PV::Facades::FPointFacade& PointFacadeSource,
 	                              const TArray<int>& BranchPoints, const float LastPointScale, const int32 LastPointIndex, const float CarveRatio,
-	                              const int32 EndIndex, const float FirstPointTargetPScale);
+	                              const int32 EndIndex, const float FirstPointTargetPScale, const bool InSkipFirstPoint);
 
 	static void RecomputeAttributes(FManagedArrayCollection& OutCollection, const FManagedArrayCollection& SourceCollection,
 	                                const TMap<int32, int32>& BranchesNewIDsToOldIDs);
