[PVE] Update gravity to cater fused points

#Synced-CL 52128411

diff --git a/Engine/Plugins/Experimental/ProceduralVegetationEditor/Source/ProceduralVegetation/Private/Implementations/PVGravity.cpp b/Engine/Plugins/Experimental/ProceduralVegetationEditor/Source/ProceduralVegetation/Private/Implementations/PVGravity.cpp
--- a/Engine/Plugins/Experimental/ProceduralVegetationEditor/Source/ProceduralVegetation/Private/Implementations/PVGravity.cpp
+++ b/Engine/Plugins/Experimental/ProceduralVegetationEditor/Source/ProceduralVegetation/Private/Implementations/PVGravity.cpp
@@ -83,8 +83,7 @@ void FPVGravity::ApplyGravity(
 	const TArray<FVector3f>& PhototropicDirections,
 	FManagedArrayCollection& OutCollection,
 	FQuat4f TotalDownForce,
-	FVector3f PreviousPosition,
-	FVector3f PivotPosition
+	FVector3f PreviousPosition
 )
 {
 	PVE_OUTER_LOOP_DEBUG_CHECK(return);
@@ -114,7 +113,15 @@ void FPVGravity::ApplyGravity(
 	PV::Facades::FBranchFacade BranchFacade = PV::Facades::FBranchFacade(OutCollection);
 	PV::Facades::FPointFacade PointFacade = PV::Facades::FPointFacade(OutCollection);
 	PV::Facades::FBudVectorsFacade BudVectorFacade = PV::Facades::FBudVectorsFacade(OutCollection);
-	PV::Facades::FFoliageFacade FoliageFacade = PV::Facades::FFoliageFacade(OutCollection);
+
+	const TManagedArray<float>& PointScales = PointFacade.GetPointScales();
+	if (PointScales.Num() == 0)
+	{
+		return;
+	}
+	
+	const float MinScale = *Algo::MinElement(PointScales);
+	const float MaxScale = *Algo::MaxElement(PointScales);
 
 	TManagedArray<FVector3f>& PointPositions = PointFacade.ModifyPositions();
 	TManagedArray<TArray<FVector3f>>& PointBudDirections = BudVectorFacade.ModifyBudDirections();
@@ -125,7 +132,6 @@ void FPVGravity::ApplyGravity(
 
 	TArray<int32> ChildrenIndices;
 	ChildrenIndices.Reserve(CurrentBranchChildren.Num());
-
 	for (const int32 Child : CurrentBranchChildren)
 	{
 		const int32 ChildIndex = BranchNumbers.Find(Child);
@@ -137,45 +143,19 @@ void FPVGravity::ApplyGravity(
 		}
 	}
 
-	if (TotalDownForce != FQuat4f::Identity)
-	[[likely]]
-	{
-		const int32 PointIndex = CurrentBranchPoints[0];
-		FVector3f& Position = PointPositions[PointIndex];
-		TArray<FVector3f>& CurrentBudDirections = PointBudDirections[PointIndex];
-		const float& PScaleGradient = PointFacade.GetPointScaleGradient(PointIndex);
-
-		const FVector3f RelativeDirection = TotalDownForce * (Position - PreviousPosition);
-		const FQuat4f DownForce = CalculateDownForce(RelativeDirection.GetUnsafeNormal(), PointIndex, PScaleGradient, GravitySettings,
-			PhototropicDirections);
-
-		TotalDownForce = DownForce * TotalDownForce;
-		PreviousPosition = Position;
-		Position = PivotPosition + DownForce * RelativeDirection;
-
-		for (int32 i = 0; i < CurrentBudDirections.Num(); ++i)
-		{
-			if (i == PV::Facades::BudDirectionsLightOptimalIndex || i == PV::Facades::BudDirectionsLightSubOptimal)
-			{
-				continue;
-			}
-				
-			CurrentBudDirections[i] = TotalDownForce * CurrentBudDirections[i];
-		}
-	}
-	else
+	if (TotalDownForce == FQuat4f::Identity)
+	[[unlikely]]
 	{
 		PreviousPosition = PointPositions[CurrentBranchPoints[0]];
 	}
 
-	const int32 MaxIndex = CurrentBranchPoints.Num() - 1;
-	for (int32 Index = 1; Index <= MaxIndex; ++Index)
+	for (int32 Index = 1; Index < CurrentBranchPoints.Num(); ++Index)
 	{
 		const int32 PointIndex = CurrentBranchPoints[Index];
 		const int32 PreviousPointIndex = CurrentBranchPoints[Index - 1];
 
 		FVector3f& Position = PointPositions[PointIndex];
-		PivotPosition = PointPositions[PreviousPointIndex];
+		const FVector3f& PivotPosition = PointPositions[PreviousPointIndex];
 
 		TArray<FVector3f>& CurrentBudDirections = PointBudDirections[PointIndex];
 
@@ -185,35 +165,28 @@ void FPVGravity::ApplyGravity(
 		const FVector3f RelativeDirection = TotalDownForce * (Position - PreviousPosition);
 		const FVector3f Tangent = RelativeDirection.GetUnsafeNormal();
 
-		const FQuat4f DownForce = CalculateDownForce(Tangent, PointIndex, PointFacade.GetPointScaleGradient(PointIndex), GravitySettings,
-			PhototropicDirections);
+		const FQuat4f DownForce = CalculateDownForce(
+			Tangent,
+			PointIndex,
+			FMath::GetRangePct(MinScale, MaxScale, PointFacade.GetPointScale(PointIndex)),
+			GravitySettings,
+			PhototropicDirections
+		);
 		TotalDownForce = DownForce * TotalDownForce;
 
 		PVE_LOOP_DEBUG_PARAMS_START()
+			PVE_LOOP_DEBUG_POINT_PARAM("Pivot", PivotPosition);
 			PVE_LOOP_DEBUG_DIRECTION_PARAM("Tangent", Position, Tangent);
+			PVE_LOOP_DEBUG_DIRECTION_PARAM("Updated Relative Direction", Position, DownForce * RelativeDirection);
 			PVE_LOOP_DEBUG_TEXT_PARAM("Length from Root", Position, FText::FromString(FString::Printf(TEXT("%.2f"), LengthFromRoot)));
 			PVE_LOOP_DEBUG_VECTOR_PARAM("Offset", Position, PivotPosition + DownForce * RelativeDirection);
 		PVE_LOOP_DEBUG_PARAMS_END()
 
 		PVE_LOOP_DEBUG_STEP(return);
 
-		// Update attachment points for all foliage elements for current branch between the current and previous point 
-		TArray<int32> FoliageEntryIds = FoliageFacade.GetFoliageEntryIdsForBranch(BranchIndex);
-		for (const int32 FoliageEntryId : FoliageEntryIds)
-		{
-			PV::Facades::FFoliageEntryData Entry = FoliageFacade.GetFoliageEntry(FoliageEntryId);
-
-			if (Entry.LengthFromRoot <= LengthFromRoot && Entry.LengthFromRoot > PreviousLengthFromRoot)
-			{
-				// Pivot point is updated based on Down force while Up Vector is updated based on Total Down force
-				// same as point positions and bud directions respectively
-				Entry.PivotPoint = PivotPosition + DownForce * RelativeDirection;
-				Entry.UpVector = TotalDownForce * Entry.UpVector;
-				Entry.NormalVector = TotalDownForce * Entry.NormalVector;
-				FoliageFacade.SetFoliageEntry(FoliageEntryId, Entry);
-			}
-		}
-
+		PreviousPosition = Position;
+		Position = PivotPosition + DownForce * RelativeDirection;
+		
 		for (const int32 ChildIndex : ChildrenIndices)
 		{
 			const TArray<int32>& BranchPoints = BranchFacade.GetPoints(ChildIndex);
@@ -233,15 +206,11 @@ void FPVGravity::ApplyGravity(
 					PhototropicDirections,
 					OutCollection,
 					TotalDownForce,
-					PreviousPosition,
-					PivotPosition
+					PreviousPosition
 				);
 			}
 		}
 
-		PreviousPosition = Position;
-		Position = PivotPosition + DownForce * RelativeDirection;
-
 		for (int32 i = 0; i < CurrentBudDirections.Num(); ++i)
 		{
 			if (i == PV::Facades::BudDirectionsLightOptimalIndex || i == PV::Facades::BudDirectionsLightSubOptimal)
diff --git a/Engine/Plugins/Experimental/ProceduralVegetationEditor/Source/ProceduralVegetation/Private/Implementations/PVGravity.h b/Engine/Plugins/Experimental/ProceduralVegetationEditor/Source/ProceduralVegetation/Private/Implementations/PVGravity.h
--- a/Engine/Plugins/Experimental/ProceduralVegetationEditor/Source/ProceduralVegetation/Private/Implementations/PVGravity.h
+++ b/Engine/Plugins/Experimental/ProceduralVegetationEditor/Source/ProceduralVegetation/Private/Implementations/PVGravity.h
@@ -53,5 +53,5 @@ private :
 	static void GeneratePhototropicData(const FPVGravityParams& InGravityParams, const PV::Facades::FBranchFacade& InBranchFacade, const PV::Facades::FPointFacade& InPointFacade, TArray<FVector3f>& OutPhototropicDirections);
 
 	static void ApplyGravity(const int32 BranchIndex, const FPVGravityParams& GravitySettings, const TArray<FVector3f>& PhototropicDirections, FManagedArrayCollection& OutCollection,
-		FQuat4f TotalDownForce = FQuat4f::Identity, FVector3f PreviousPosition = FVector3f::ZeroVector, FVector3f PivotPosition = FVector3f::ZeroVector);
+		FQuat4f TotalDownForce = FQuat4f::Identity, FVector3f PreviousPosition = FVector3f::ZeroVector);
 };
