Improve branch stitching algorithm for the StaticMesh importer

- Add a new branch connection code path that uses branch tracing to determine how to connect child branches to parent geometry, selecting attachment edges that lie on the active trace path rather than nearest geometry
- Expose ComputeLengthFromRoot, TraceBranches, and ExtractBranchHierarchies as public API and add a CVar to toggle between the old and new branch connection algorithms
- Re-structure PVImportHelpers.cpp (move ComputeLengthFromRoot, TraceBranches, and ExtractBranchHierarchies out of local namespace and move remaining functions to anonymous namespace)

#Synced-CL 52101121

#Synced-CL 52101121

diff --git a/Engine/Plugins/Experimental/ProceduralVegetationEditor/Source/ProceduralVegetation/Private/Helpers/PVImportHelpers.cpp b/Engine/Plugins/Experimental/ProceduralVegetationEditor/Source/ProceduralVegetation/Private/Helpers/PVImportHelpers.cpp
--- a/Engine/Plugins/Experimental/ProceduralVegetationEditor/Source/ProceduralVegetation/Private/Helpers/PVImportHelpers.cpp
+++ b/Engine/Plugins/Experimental/ProceduralVegetationEditor/Source/ProceduralVegetation/Private/Helpers/PVImportHelpers.cpp
@@ -12,447 +12,459 @@
 
 #include "Engine/StaticMesh.h"
 
-namespace PV::ImportHelper
+namespace
 {
-	static void ComputeLengthFromRoot(
-		const UE::Geometry::FDynamicMesh3& DynamicMesh,
-		const FBooleanVertexAttribute& RootPointsAttribute,
-		const FBooleanVertexAttribute& EndPointsAttribute,
-		FFloatVertexAttribute& LengthFromRootAttribute,
-		const TFunction<bool(int32)>& VertexFilterPredicate
-	)
-	{
-		TRACE_CPUPROFILER_EVENT_SCOPE(ComputeLengthFromRoot);
+	const static TSet<FName> GeometryGroups = {
+		FGeometryCollection::VerticesGroup,
+		FGeometryCollection::FacesGroup,
+		FGeometryCollection::GeometryGroup,
+		FGeometryCollection::BreakingGroup,
+		FGeometryCollection::MaterialGroup,
+		FTransformCollection::TransformGroup,
+		FTransformCollection::ConvexGroup
+	};
 
-		// This algorithm starts at the root point and walks along every vertex, keeping track of its distance to the root point and stores that distance
-		// in the LengthFromRootAttribute.
+	const static FName MaterialPathAttributeName("MaterialPath");
 
-		for (int32 VertexID : DynamicMesh.VertexIndicesItr())
+	static void SetManagedArrayCollectionGeometry(const FGeometryCollection& InGeometryCollection, FManagedArrayCollection& InOutManagedArrayCollection)
+	{
+		for (const FName GroupName : GeometryGroups)
 		{
-			LengthFromRootAttribute.SetValue(VertexID, -1);
+			InOutManagedArrayCollection.RemoveGroup(GroupName);
 		}
 
-		TArray<int32> SearchQueue;
-		SearchQueue.Reserve(DynamicMesh.VertexCount());
-		for (int32 VertexID : RootPointsAttribute.FindAllNonZero())
+		InGeometryCollection.CopyTo(&InOutManagedArrayCollection);
+	}
+
+	static void AppendGeometryFromManagedArrayCollection(FGeometryCollection& InOutGeometryCollection, const FManagedArrayCollection& InManagedArrayCollection)
+	{
+		TArray<FName> GroupsToSkip;
+		GroupsToSkip.Reserve(InManagedArrayCollection.GroupNames().Num());
+
+		for (FName GroupName : InManagedArrayCollection.GroupNames())
 		{
-			SearchQueue.Add(VertexID);
-			LengthFromRootAttribute.SetValue(VertexID, 0);
+			if (!GeometryGroups.Contains(GroupName))
+			{
+				GroupsToSkip.Add(GroupName);
+			}
 		}
 
-		while (SearchQueue.Num() > 0)
+		FGeometryCollection GeometryCollection;
+		InManagedArrayCollection.CopyTo(&GeometryCollection, GroupsToSkip);
+
+		const auto* ExistingMaterialPathsAttribute = InOutGeometryCollection.FindAttribute<FString>(MaterialPathAttributeName, FGeometryCollection::MaterialGroup);
+		const int32 NumExistingMaterials = ExistingMaterialPathsAttribute ? ExistingMaterialPathsAttribute->Num() : 1;
+		const bool bHasCopiedGeometry = InOutGeometryCollection.AppendGeometry(GeometryCollection, NumExistingMaterials, false) != INDEX_NONE;
+
+		// Copy materials. Ideally we'd want to re-use materials where possible, but given this is for debugging we just assign 
+		// one material per mesh.
+		if (bHasCopiedGeometry)
 		{
-			const int32& VertexID = SearchQueue.Pop(EAllowShrinking::No);
+			if (!InOutGeometryCollection.HasGroup(FGeometryCollection::MaterialGroup))
+			{
+				InOutGeometryCollection.AddGroup(FGeometryCollection::MaterialGroup);
+			}
 
-			const float PrevLengthFromRoot = LengthFromRootAttribute.GetValue(VertexID);
-			const FVector PrevVertexPosition = DynamicMesh.GetVertex(VertexID);
+			auto* MaterialPathsAttributeToCopy = GeometryCollection.FindAttribute<FString>(MaterialPathAttributeName, FGeometryCollection::MaterialGroup);
+			const int32 NumMaterialsToCopy = MaterialPathsAttributeToCopy ? MaterialPathsAttributeToCopy->Num() : 1;
+			InOutGeometryCollection.Resize(NumExistingMaterials + NumMaterialsToCopy, FGeometryCollection::MaterialGroup);
 
-			const TArray<int32> Neighbours = PV::DynamicMeshHelper::GetVertexNeighbours(DynamicMesh, VertexID);
-			for (int32 NeighbourVertexID : Neighbours)
+			if (MaterialPathsAttributeToCopy)
 			{
-				if (!EndPointsAttribute.GetValue(NeighbourVertexID) // Ingore vertex filter for end points as that means we've reached the end of the branch
-					&& !VertexFilterPredicate(NeighbourVertexID))
+				auto* MaterialPathsAttribute = InOutGeometryCollection.FindOrAddAttributeTyped<FString>(MaterialPathAttributeName, FGeometryCollection::MaterialGroup);
+				for (int32 i = NumExistingMaterials; i < NumExistingMaterials + NumMaterialsToCopy; ++i)
 				{
-					continue;
+					const int32 MaterialIndexToCopy = i - NumExistingMaterials;
+					if (MaterialPathsAttributeToCopy->IsValidIndex(MaterialIndexToCopy))
+					{
+						(*MaterialPathsAttribute)[i] = (*MaterialPathsAttributeToCopy)[MaterialIndexToCopy];
+					}
 				}
+			}
+		}
+	}
+};
 
-				const FVector VertexPosition = DynamicMesh.GetVertex(NeighbourVertexID);
-				const float NewLengthFromRoot = PrevLengthFromRoot + (VertexPosition - PrevVertexPosition).Size();
-				const float ExistingLengthFromRoot = LengthFromRootAttribute.GetValue(NeighbourVertexID);
+void PV::ImportHelper::ComputeLengthFromRoot(
+	const UE::Geometry::FDynamicMesh3& DynamicMesh,
+	const FBooleanVertexAttribute& RootPointsAttribute,
+	const FBooleanVertexAttribute& EndPointsAttribute,
+	FFloatVertexAttribute& LengthFromRootAttribute,
+	const TFunction<bool(int32)>& VertexFilterPredicate
+)
+{
+	TRACE_CPUPROFILER_EVENT_SCOPE(ComputeLengthFromRoot);
 
-				if (ExistingLengthFromRoot >= 0 
-					&& ExistingLengthFromRoot <= NewLengthFromRoot)
-				{
-					// If we find a shorter path then we should ideally remove vertices from the queue
-					// which are computing the distance using old values. 
-					// The current algorithm will still produce correct results, but may compute a large number
-					// of unecessary distances which will later be over-written when a closer path is found.
-					continue;
-				}
+	// This algorithm starts at the root point and walks along every vertex, keeping track of its distance to the root point and stores that distance
+	// in the LengthFromRootAttribute.
 
-				LengthFromRootAttribute.SetValue(NeighbourVertexID, NewLengthFromRoot);
+	for (int32 VertexID : DynamicMesh.VertexIndicesItr())
+	{
+		LengthFromRootAttribute.SetValue(VertexID, -1);
+	}
 
-				SearchQueue.Insert(NeighbourVertexID, 0);
-			}
-		}
+	TArray<int32> SearchQueue;
+	SearchQueue.Reserve(DynamicMesh.VertexCount());
+	for (int32 VertexID : RootPointsAttribute.FindAllNonZero())
+	{
+		SearchQueue.Add(VertexID);
+		LengthFromRootAttribute.SetValue(VertexID, 0);
 	}
 
-	static void TraceBranches(
-		const UE::Geometry::FDynamicMesh3& DynamicMesh,
-		const FFloatVertexAttribute& LengthFromRootAttribute,
-		const FBooleanVertexAttribute& RootPointAttribute,
-		const FBooleanVertexAttribute& EndPointAttribute,
-		FIntVertexAttribute& NextBranchPointAttribute,
-		const TFunction<bool(int32)>& VertexFilterPredicate
-	)
+	while (SearchQueue.Num() > 0)
 	{
-		using namespace PV::DynamicMeshHelper;
+		const int32& VertexID = SearchQueue.Pop(EAllowShrinking::No);
 
-		TRACE_CPUPROFILER_EVENT_SCOPE(TraceBranches);
+		const float PrevLengthFromRoot = LengthFromRootAttribute.GetValue(VertexID);
+		const FVector PrevVertexPosition = DynamicMesh.GetVertex(VertexID);
 
-		for (int32 VertexID : DynamicMesh.VertexIndicesItr())
+		const TArray<int32> Neighbours = PV::DynamicMeshHelper::GetVertexNeighbours(DynamicMesh, VertexID);
+		for (int32 NeighbourVertexID : Neighbours)
 		{
-			NextBranchPointAttribute.SetValue(VertexID, INDEX_NONE);
-		}
+			if (!EndPointsAttribute.GetValue(NeighbourVertexID) // Ingore vertex filter for end points as that means we've reached the end of the branch
+				&& !VertexFilterPredicate(NeighbourVertexID))
+			{
+				continue;
+			}
 
-		const auto FindNextPoint = [&](int32 PointIndex, float MinLengthFromRoot)->int32
-		{
-			int32 OutNextPointIndex = INDEX_NONE;
-			const TArray<int32> Neighbours = GetVertexNeighbours(DynamicMesh, PointIndex);
-			for (int32 Neighbour : Neighbours)
+			const FVector VertexPosition = DynamicMesh.GetVertex(NeighbourVertexID);
+			const float NewLengthFromRoot = PrevLengthFromRoot + (VertexPosition - PrevVertexPosition).Size();
+			const float ExistingLengthFromRoot = LengthFromRootAttribute.GetValue(NeighbourVertexID);
+
+			if (ExistingLengthFromRoot >= 0 
+				&& ExistingLengthFromRoot <= NewLengthFromRoot)
 			{
-				if (RootPointAttribute.GetValue(Neighbour))
-				{
-					OutNextPointIndex = Neighbour;
-					break;
-				}
+				// If we find a shorter path then we should ideally remove vertices from the queue
+				// which are computing the distance using old values. 
+				// The current algorithm will still produce correct results, but may compute a large number
+				// of unecessary distances which will later be over-written when a closer path is found.
+				continue;
+			}
 
-				if (!VertexFilterPredicate(Neighbour))
-				{
-					continue;
-				}
+			LengthFromRootAttribute.SetValue(NeighbourVertexID, NewLengthFromRoot);
 
-				const float LengthFromRoot = LengthFromRootAttribute.GetValue(Neighbour);
-				check(LengthFromRoot >= 0);
+			SearchQueue.Insert(NeighbourVertexID, 0);
+		}
+	}
+}
 
-				if (LengthFromRoot >= MinLengthFromRoot)
-				{
-					continue;
-				}
-				MinLengthFromRoot = LengthFromRoot;
-				OutNextPointIndex = Neighbour;
-			}
-			return OutNextPointIndex;
-		};
+void PV::ImportHelper::TraceBranches(
+	const UE::Geometry::FDynamicMesh3& DynamicMesh,
+	const FFloatVertexAttribute& LengthFromRootAttribute,
+	const FBooleanVertexAttribute& RootPointAttribute,
+	const FBooleanVertexAttribute& EndPointAttribute,
+	FIntVertexAttribute& NextBranchPointAttribute,
+	const TFunction<bool(int32)>& VertexFilterPredicate
+)
+{
+	using namespace PV::DynamicMeshHelper;
 
-		const TArray<int32> EndPoints = EndPointAttribute.FindAllNonZero();
+	TRACE_CPUPROFILER_EVENT_SCOPE(TraceBranches);
 
-		TArray<bool> bHasProcessedVertex;
+	for (int32 VertexID : DynamicMesh.VertexIndicesItr())
+	{
+		NextBranchPointAttribute.SetValue(VertexID, INDEX_NONE);
+	}
 
-		for (int32 i = 0; i < EndPoints.Num(); i++)
+	const auto FindNextPoint = [&](int32 PointIndex, float MaxLengthFromRoot)->int32
+	{
+		int32 OutNextPointIndex = INDEX_NONE;
+		float ClosestLengthFromRoot = MAX_flt;
+	
+		const FVector ThisVertexPosition = DynamicMesh.GetVertex(PointIndex);
+		const float ThisLengthFromRoot = LengthFromRootAttribute.GetValue(PointIndex);
+		const TArray<int32> Neighbours = GetVertexNeighbours(DynamicMesh, PointIndex);
+		for (int32 Neighbour : Neighbours)
 		{
-			int32 CurrentPointIndex = EndPoints[i];
+			if (RootPointAttribute.GetValue(Neighbour))
+			{
+				OutNextPointIndex = Neighbour;
+				break;
+			}
 
-			float MinLengthFromRoot = LengthFromRootAttribute.GetValue(CurrentPointIndex);
-			if (MinLengthFromRoot < 0)
+			if (!VertexFilterPredicate(Neighbour))
 			{
-				continue; // This end point was unable to trace down to a root point, perhaps log a warning?
+				continue;
+			}
+
+			const float NeighbourLengthFromRoot = LengthFromRootAttribute.GetValue(Neighbour);
+			check(NeighbourLengthFromRoot >= 0);
+			if (NeighbourLengthFromRoot >= MaxLengthFromRoot)
+			{
+				continue;
+			}
+		
+			const FVector NeighbourVertexPosition = DynamicMesh.GetVertex(Neighbour);
+			const float DistToNeighbour = (NeighbourVertexPosition - ThisVertexPosition).Size();
+			const float PathLengthFromRoot = NeighbourLengthFromRoot + DistToNeighbour;
+			if (PathLengthFromRoot >= ClosestLengthFromRoot)
+			{
+				continue;
 			}
+		
+			ClosestLengthFromRoot = PathLengthFromRoot;
+			OutNextPointIndex = Neighbour;
+		}
+		return OutNextPointIndex;
+	};
+
+	const TArray<int32> EndPoints = EndPointAttribute.FindAllNonZero();
+
+	TArray<bool> bHasProcessedVertex;
 
-			bHasProcessedVertex.Reset();
-			bHasProcessedVertex.SetNumZeroed(DynamicMesh.MaxVertexID());
+	for (int32 i = 0; i < EndPoints.Num(); i++)
+	{
+		int32 CurrentPointIndex = EndPoints[i];
+
+		float MaxLengthFromRoot = LengthFromRootAttribute.GetValue(CurrentPointIndex);
+		if (MaxLengthFromRoot < 0)
+		{
+			continue; // This end point was unable to trace down to a root point, perhaps log a warning?
+		}
 
-			while (true)
+		bHasProcessedVertex.Reset();
+		bHasProcessedVertex.SetNumZeroed(DynamicMesh.MaxVertexID());
+
+		while (true)
+		{
+			const int32 NextPointIndex = FindNextPoint(CurrentPointIndex, MaxLengthFromRoot);
+			if (NextPointIndex == INDEX_NONE)
 			{
-				const int32 NextPointIndex = FindNextPoint(CurrentPointIndex, MinLengthFromRoot);
-				if (NextPointIndex == INDEX_NONE)
-				{
-					break;
-				}
+				break;
+			}
 				
-				NextBranchPointAttribute.SetValue(CurrentPointIndex, NextPointIndex);
+			NextBranchPointAttribute.SetValue(CurrentPointIndex, NextPointIndex);
 
-				if (NextBranchPointAttribute.GetValue(NextPointIndex) != INDEX_NONE
-					|| RootPointAttribute.GetValue(NextPointIndex))
-				{
-					break;
-				}
+			if (NextBranchPointAttribute.GetValue(NextPointIndex) != INDEX_NONE
+				|| RootPointAttribute.GetValue(NextPointIndex))
+			{
+				break;
+			}
 
-				if (!ensureMsgf(!bHasProcessedVertex[NextPointIndex], TEXT("Infinite loop detected")))
-				{
-					break;
-				}
+			if (!ensureMsgf(!bHasProcessedVertex[NextPointIndex], TEXT("Infinite loop detected")))
+			{
+				break;
+			}
 
-				bHasProcessedVertex[NextPointIndex] = true;
+			bHasProcessedVertex[NextPointIndex] = true;
 
-				CurrentPointIndex = NextPointIndex;
-				MinLengthFromRoot = LengthFromRootAttribute.GetValue(CurrentPointIndex);
-			}
+			CurrentPointIndex = NextPointIndex;
+			MaxLengthFromRoot = LengthFromRootAttribute.GetValue(CurrentPointIndex);
 		}
 	}
+}
 
-	static void ExtractBranchHierarchies(
-		const UE::Geometry::FDynamicMesh3& DynamicMesh,
-		const FBooleanVertexAttribute& RootPointAttribute,
-		const FBooleanVertexAttribute& EndPointAttribute,
-		const FIntVertexAttribute& NextBranchPointAttribute,
-		const FFloatVertexAttribute& VertexScaleAttribute,
-		const FFloatVertexAttribute& LengthFromRootAttribute,
-		TArray<FMeshBranchHierarchy>& OutMeshBranchHierarchies
-	)
-	{
-		using namespace PV::DynamicMeshHelper;
+void PV::ImportHelper::ExtractBranchHierarchies(
+	const UE::Geometry::FDynamicMesh3& DynamicMesh,
+	const FBooleanVertexAttribute& RootPointAttribute,
+	const FBooleanVertexAttribute& EndPointAttribute,
+	const FIntVertexAttribute& NextBranchPointAttribute,
+	const FFloatVertexAttribute& VertexScaleAttribute,
+	const FFloatVertexAttribute& LengthFromRootAttribute,
+	TArray<FMeshBranchHierarchy>& OutMeshBranchHierarchies
+)
+{
+	using namespace PV::DynamicMeshHelper;
 
-		TRACE_CPUPROFILER_EVENT_SCOPE(ExtractBranchHierarchies);
+	TRACE_CPUPROFILER_EVENT_SCOPE(ExtractBranchHierarchies);
 
-		TArray<bool> bProcessedVertices;
-		bProcessedVertices.SetNumZeroed(DynamicMesh.MaxVertexID());
+	TArray<bool> bProcessedVertices;
+	bProcessedVertices.SetNumZeroed(DynamicMesh.MaxVertexID());
 
-		TArray<float> MaxSubtreeLengthFromRoot;
-		{
-			MaxSubtreeLengthFromRoot.SetNumZeroed(DynamicMesh.MaxVertexID());
+	TArray<float> MaxSubtreeLengthFromRoot;
+	{
+		MaxSubtreeLengthFromRoot.SetNumZeroed(DynamicMesh.MaxVertexID());
 
-			const TArray<int32> EndPoints = EndPointAttribute.FindAllNonZero();
-			for (int32 EndVertexID : EndPoints)
+		const TArray<int32> EndPoints = EndPointAttribute.FindAllNonZero();
+		for (int32 EndVertexID : EndPoints)
+		{
+			const float EndpointLength = LengthFromRootAttribute.GetValue(EndVertexID);
+			if (EndpointLength <= 0.f)
 			{
-				const float EndpointLength = LengthFromRootAttribute.GetValue(EndVertexID);
-				if (EndpointLength <= 0.f)
-				{
-					continue; // endpoint not reachable from root, skip
-				}
-				int32 Current = EndVertexID;
-				while (Current != INDEX_NONE)
+				continue; // endpoint not reachable from root, skip
+			}
+			int32 Current = EndVertexID;
+			while (Current != INDEX_NONE)
+			{
+				if (EndpointLength <= MaxSubtreeLengthFromRoot[Current])
 				{
-					if (EndpointLength <= MaxSubtreeLengthFromRoot[Current])
-					{
-						break;
-					}
-					MaxSubtreeLengthFromRoot[Current] = EndpointLength;
-					Current = NextBranchPointAttribute.GetValue(Current);
+					break;
 				}
+				MaxSubtreeLengthFromRoot[Current] = EndpointLength;
+				Current = NextBranchPointAttribute.GetValue(Current);
 			}
 		}
+	}
 		
-		const auto FindBestParentPoint = [&](int32 PrevVertexID, int32 VertexID, const TArray<int32>& NextPoints) -> int32
+	const auto FindBestParentPoint = [&](int32 PrevVertexID, int32 VertexID, const TArray<int32>& NextPoints) -> int32
+	{
+		if (NextPoints.Num() == 1)
 		{
-			if (NextPoints.Num() == 1)
-			{
-				return NextPoints[0];
-			}
-			check(NextPoints.Num() > 0);
+			return NextPoints[0];
+		}
+		check(NextPoints.Num() > 0);
+
+		const FVector3d VertexPos = DynamicMesh.GetVertex(VertexID);
+		const FVector3d PrevVertexPos = DynamicMesh.GetVertex(PrevVertexID);
+		const FVector3d ApicalDirection = (VertexPos - PrevVertexPos).GetSafeNormal();
+		const float CurrentRadius = VertexScaleAttribute.GetValue(VertexID);
 
-			const FVector3d VertexPos = DynamicMesh.GetVertex(VertexID);
-			const FVector3d PrevVertexPos = DynamicMesh.GetVertex(PrevVertexID);
-			const FVector3d ApicalDirection = (VertexPos - PrevVertexPos).GetSafeNormal();
-			const float CurrentRadius = VertexScaleAttribute.GetValue(VertexID);
+		TArray<float, TInlineAllocator<4>> Dots;
+		TArray<float, TInlineAllocator<4>> Lengths;
+		Dots.SetNumUninitialized(NextPoints.Num());
+		Lengths.SetNumUninitialized(NextPoints.Num());
 
-			TArray<float, TInlineAllocator<4>> Dots;
-			TArray<float, TInlineAllocator<4>> Lengths;
-			Dots.SetNumUninitialized(NextPoints.Num());
-			Lengths.SetNumUninitialized(NextPoints.Num());
+		for (int32 i = 0; i < NextPoints.Num(); ++i)
+		{
+			const int32 NextVertexID = NextPoints[i];
+			const FVector3d NextPos = DynamicMesh.GetVertex(NextVertexID);
+			const FVector3d NextDir = (NextPos - VertexPos).GetSafeNormal();
+			Dots[i] = (1.0 + FVector3d::DotProduct(ApicalDirection, NextDir)) / 2.0;
+			Lengths[i] = MaxSubtreeLengthFromRoot[NextVertexID] - LengthFromRootAttribute.GetValue(NextVertexID);
+		}
 
-			for (int32 i = 0; i < NextPoints.Num(); ++i)
+		// Normalize a raw float array (all values >= 0) by dividing by the sum.
+		// This preserves proportionality: a value of 10000 vs 1 stays near 1.0 vs 0.0001,
+		// whereas a value of 10 vs 9 stays near 0.526 vs 0.473.
+		// Min-max normalization collapses both cases to {1, 0}, losing magnitude information.
+		// bInvert: output 1-normalized instead of normalized
+		const auto NormalizeArray = [](TArray<float, TInlineAllocator<4>>& InOutArray, bool bInvert)
+		{
+			float Sum = 0.f;
+			for (int32 i = 0; i < InOutArray.Num(); ++i)
 			{
-				const int32 NextVertexID = NextPoints[i];
-				const FVector3d NextPos = DynamicMesh.GetVertex(NextVertexID);
-				const FVector3d NextDir = (NextPos - VertexPos).GetSafeNormal();
-				Dots[i] = (1.0 + FVector3d::DotProduct(ApicalDirection, NextDir)) / 2.0;
-				Lengths[i] = MaxSubtreeLengthFromRoot[NextVertexID] - LengthFromRootAttribute.GetValue(NextVertexID);
+				Sum += InOutArray[i];
 			}
 
-			// Normalize a raw float array (all values >= 0) by dividing by the sum.
-			// This preserves proportionality: a value of 10000 vs 1 stays near 1.0 vs 0.0001,
-			// whereas a value of 10 vs 9 stays near 0.526 vs 0.473.
-			// Min-max normalization collapses both cases to {1, 0}, losing magnitude information.
-			// bInvert: output 1-normalized instead of normalized
-			const auto NormalizeArray = [](TArray<float, TInlineAllocator<4>>& InOutArray, bool bInvert)
+			if (Sum < UE_SMALL_NUMBER)
 			{
-				float Sum = 0.f;
 				for (int32 i = 0; i < InOutArray.Num(); ++i)
 				{
-					Sum += InOutArray[i];
+					InOutArray[i] = 0.f;
 				}
-
-				if (Sum < UE_SMALL_NUMBER)
+			}
+			else
+			{
+				for (int32 i = 0; i < InOutArray.Num(); ++i)
 				{
-					for (int32 i = 0; i < InOutArray.Num(); ++i)
-					{
-						InOutArray[i] = 0.f;
-					}
+					InOutArray[i] = InOutArray[i] / Sum;
 				}
-				else
+
+				if (bInvert)
 				{
 					for (int32 i = 0; i < InOutArray.Num(); ++i)
 					{
-						InOutArray[i] = InOutArray[i] / Sum;
+						InOutArray[i] = 1.f - InOutArray[i];
 					}
-
-					if (bInvert)
-					{
-						for (int32 i = 0; i < InOutArray.Num(); ++i)
-						{
-							InOutArray[i] = 1.f - InOutArray[i];
-						}
-					}
-				}
-			};
-
-			NormalizeArray(Lengths, false);
-
-			float BestScore = -MAX_flt;
-			int32 BestIndex = INDEX_NONE;
-			for (int32 i = 0; i < NextPoints.Num(); ++i)
-			{
-				const float Score = Dots[i] + Lengths[i];
-				if (Score > BestScore)
-				{
-					BestScore = Score;
-					BestIndex = NextPoints[i];
 				}
 			}
-
-			check(BestIndex != INDEX_NONE);
-			return BestIndex;
 		};
 
-		const auto FindNextBranchPoints = [&](int32 VertexID)
-		{
-			TArray<int32> NextBranchPoints;
-			const TArray<int32> Neighbours = GetVertexNeighbours(DynamicMesh, VertexID);
+		NormalizeArray(Lengths, false);
 
-			for (int32 Neighbour : Neighbours)
-			{
-				if (NextBranchPointAttribute.GetValue(Neighbour) == VertexID
-					&& !bProcessedVertices[Neighbour])
-				{
-					NextBranchPoints.Add(Neighbour);
-				}
-			}
-			return NextBranchPoints;
-		};
-
-		const TFunction<void(int32, TArray<FMeshBranchHierarchy::FBranch>&)> TraceBranchesRecursive = [&](
-			int32 VertexID, 
-			TArray<FMeshBranchHierarchy::FBranch>& OutBranches
-		)
+		float BestScore = -MAX_flt;
+		int32 BestIndex = INDEX_NONE;
+		for (int32 i = 0; i < NextPoints.Num(); ++i)
 		{
-			const int32 BranchIndex = OutBranches.Num() - 1;
-			OutBranches[BranchIndex].BranchPoints.Add(VertexID);
-			
-			bProcessedVertices[VertexID] = true;
-
-			if (EndPointAttribute.GetValue(VertexID))
+			const float Score = Dots[i] + Lengths[i];
+			if (Score > BestScore)
 			{
-				return;
+				BestScore = Score;
+				BestIndex = NextPoints[i];
 			}
+		}
 
-			const TArray<int32> NextBranchPoints = FindNextBranchPoints(VertexID);
-			check(NextBranchPoints.Num() > 0);
+		check(BestIndex != INDEX_NONE);
+		return BestIndex;
+	};
 
-			for (int32 NextVertexID : NextBranchPoints)
+	const auto FindNextBranchPoints = [&](int32 VertexID)
+	{
+		TArray<int32> NextBranchPoints;
+		const TArray<int32> Neighbours = GetVertexNeighbours(DynamicMesh, VertexID);
+
+		for (int32 Neighbour : Neighbours)
+		{
+			if (NextBranchPointAttribute.GetValue(Neighbour) == VertexID
+				&& !bProcessedVertices[Neighbour])
 			{
-				bProcessedVertices[NextVertexID] = true;
+				NextBranchPoints.Add(Neighbour);
 			}
+		}
+		return NextBranchPoints;
+	};
 
-			const int32 PrevVertexID = OutBranches[BranchIndex].BranchPoints.Last(1);
-			const int32 NextBranchPointsIndex = FindBestParentPoint(PrevVertexID, VertexID, NextBranchPoints);
+	const TFunction<void(int32, TArray<FMeshBranchHierarchy::FBranch>&)> TraceBranchesRecursive = [&](
+		int32 VertexID, 
+		TArray<FMeshBranchHierarchy::FBranch>& OutBranches
+	)
+	{
+		const int32 BranchIndex = OutBranches.Num() - 1;
+		OutBranches[BranchIndex].BranchPoints.Add(VertexID);
 			
-			TraceBranchesRecursive(NextBranchPointsIndex, OutBranches);
-
-			for (int32 NextVertexID : NextBranchPoints)
-			{
-				if (NextVertexID != NextBranchPointsIndex)
-				{
-					FMeshBranchHierarchy::FBranch& NextBranch = OutBranches.AddDefaulted_GetRef();
-					NextBranch.BranchPoints.Add(VertexID);
-					NextBranch.ParentBranch = BranchIndex;
+		bProcessedVertices[VertexID] = true;
 
-					TraceBranchesRecursive(NextVertexID, OutBranches);
-				}
-			}
-		};
+		if (EndPointAttribute.GetValue(VertexID))
+		{
+			return;
+		}
 
-		const TArray<int32> RootPoints = RootPointAttribute.FindAllNonZero();
-		OutMeshBranchHierarchies.Reserve(RootPoints.Num());
+		const TArray<int32> NextBranchPoints = FindNextBranchPoints(VertexID);
+		check(NextBranchPoints.Num() > 0);
 
-		for (const int32 RootPointIndex : RootPoints)
+		for (int32 NextVertexID : NextBranchPoints)
 		{
-			TArray<TArray<FMeshBranchHierarchy::FBranch>> BranchHierarchies;
+			bProcessedVertices[NextVertexID] = true;
+		}
 
-			const TArray<int32> NextBranchPoints = FindNextBranchPoints(RootPointIndex);
-			for (int32 NextVertexID : NextBranchPoints)
-			{
-				bProcessedVertices[NextVertexID] = true;
-			}
+		const int32 PrevVertexID = OutBranches[BranchIndex].BranchPoints.Last(1);
+		const int32 NextBranchPointsIndex = FindBestParentPoint(PrevVertexID, VertexID, NextBranchPoints);
+			
+		TraceBranchesRecursive(NextBranchPointsIndex, OutBranches);
 
-			for (int32 NextVertexID : NextBranchPoints)
+		for (int32 NextVertexID : NextBranchPoints)
+		{
+			if (NextVertexID != NextBranchPointsIndex)
 			{
-				TArray<FMeshBranchHierarchy::FBranch>& Branches = BranchHierarchies.AddDefaulted_GetRef();
-				FMeshBranchHierarchy::FBranch& RootBranch = Branches.AddDefaulted_GetRef();
-				RootBranch.BranchPoints.Add(RootPointIndex);
-				TraceBranchesRecursive(NextVertexID, Branches);
-			}
+				FMeshBranchHierarchy::FBranch& NextBranch = OutBranches.AddDefaulted_GetRef();
+				NextBranch.BranchPoints.Add(VertexID);
+				NextBranch.ParentBranch = BranchIndex;
 
-			if (BranchHierarchies.Num() > 0)
-			{
-				// TODO: Handle multiple branches starting from root
-				TArray<FMeshBranchHierarchy::FBranch>* BestBranches = Algo::MaxElement(
-					BranchHierarchies, 
-					[](const auto& A, const auto& B) { return A.Num() < B.Num(); }
-				);
-				OutMeshBranchHierarchies.Emplace(MoveTemp(*BestBranches));
+				TraceBranchesRecursive(NextVertexID, OutBranches);
 			}
 		}
-	}
-
-	const static TSet<FName> GeometryGroups = {
-		FGeometryCollection::VerticesGroup,
-		FGeometryCollection::FacesGroup,
-		FGeometryCollection::GeometryGroup,
-		FGeometryCollection::BreakingGroup,
-		FGeometryCollection::MaterialGroup,
-		FTransformCollection::TransformGroup,
-		FTransformCollection::ConvexGroup
 	};
 
-	const static FName MaterialPathAttributeName("MaterialPath");
+	const TArray<int32> RootPoints = RootPointAttribute.FindAllNonZero();
+	OutMeshBranchHierarchies.Reserve(RootPoints.Num());
 
-	static void AppendGeometryFromManagedArrayCollection(FGeometryCollection& InOutGeometryCollection, const FManagedArrayCollection& InManagedArrayCollection)
+	for (const int32 RootPointIndex : RootPoints)
 	{
-		TArray<FName> GroupsToSkip;
-		GroupsToSkip.Reserve(InManagedArrayCollection.GroupNames().Num());
+		TArray<TArray<FMeshBranchHierarchy::FBranch>> BranchHierarchies;
 
-		for (FName GroupName : InManagedArrayCollection.GroupNames())
+		const TArray<int32> NextBranchPoints = FindNextBranchPoints(RootPointIndex);
+		for (int32 NextVertexID : NextBranchPoints)
 		{
-			if (!GeometryGroups.Contains(GroupName))
-			{
-				GroupsToSkip.Add(GroupName);
-			}
+			bProcessedVertices[NextVertexID] = true;
 		}
 
-		FGeometryCollection GeometryCollection;
-		InManagedArrayCollection.CopyTo(&GeometryCollection, GroupsToSkip);
-
-		const auto* ExistingMaterialPathsAttribute = InOutGeometryCollection.FindAttribute<FString>(MaterialPathAttributeName, FGeometryCollection::MaterialGroup);
-		const int32 NumExistingMaterials = ExistingMaterialPathsAttribute ? ExistingMaterialPathsAttribute->Num() : 1;
-		const bool bHasCopiedGeometry = InOutGeometryCollection.AppendGeometry(GeometryCollection, NumExistingMaterials, false) != INDEX_NONE;
-
-		// Copy materials. Ideally we'd want to re-use materials where possible, but given this is for debugging we just assign 
-		// one material per mesh.
-		if (bHasCopiedGeometry)
+		for (int32 NextVertexID : NextBranchPoints)
 		{
-			if (!InOutGeometryCollection.HasGroup(FGeometryCollection::MaterialGroup))
-			{
-				InOutGeometryCollection.AddGroup(FGeometryCollection::MaterialGroup);
-			}
-
-			auto* MaterialPathsAttributeToCopy = GeometryCollection.FindAttribute<FString>(MaterialPathAttributeName, FGeometryCollection::MaterialGroup);
-			const int32 NumMaterialsToCopy = MaterialPathsAttributeToCopy ? MaterialPathsAttributeToCopy->Num() : 1;
-			InOutGeometryCollection.Resize(NumExistingMaterials + NumMaterialsToCopy, FGeometryCollection::MaterialGroup);
-
-			if (MaterialPathsAttributeToCopy)
-			{
-				auto* MaterialPathsAttribute = InOutGeometryCollection.FindOrAddAttributeTyped<FString>(MaterialPathAttributeName, FGeometryCollection::MaterialGroup);
-				for (int32 i = NumExistingMaterials; i < NumExistingMaterials + NumMaterialsToCopy; ++i)
-				{
-					const int32 MaterialIndexToCopy = i - NumExistingMaterials;
-					if (MaterialPathsAttributeToCopy->IsValidIndex(MaterialIndexToCopy))
-					{
-						(*MaterialPathsAttribute)[i] = (*MaterialPathsAttributeToCopy)[MaterialIndexToCopy];
-					}
-				}
-			}
+			TArray<FMeshBranchHierarchy::FBranch>& Branches = BranchHierarchies.AddDefaulted_GetRef();
+			FMeshBranchHierarchy::FBranch& RootBranch = Branches.AddDefaulted_GetRef();
+			RootBranch.BranchPoints.Add(RootPointIndex);
+			TraceBranchesRecursive(NextVertexID, Branches);
 		}
-	}
 
-	static void SetManagedArrayCollectionGeometry(const FGeometryCollection& InGeometryCollection, FManagedArrayCollection& InOutManagedArrayCollection)
-	{
-		for (const FName GroupName : GeometryGroups)
+		if (BranchHierarchies.Num() > 0)
 		{
-			InOutManagedArrayCollection.RemoveGroup(GroupName);
+			// TODO: Handle multiple branches starting from root
+			TArray<FMeshBranchHierarchy::FBranch>* BestBranches = Algo::MaxElement(
+				BranchHierarchies, 
+				[](const auto& A, const auto& B) { return A.Num() < B.Num(); }
+			);
+			OutMeshBranchHierarchies.Emplace(MoveTemp(*BestBranches));
 		}
-
-		InGeometryCollection.CopyTo(&InOutManagedArrayCollection);
 	}
 }
 
@@ -708,7 +720,6 @@ void PV::ImportHelper::ConvertDynamicMeshBranchHierarchiesToBranchDescription(
 	}
 }
 
-
 FPVBranchHierarchyDescription::EValidateHierarchyResult FPVBranchHierarchyDescription::ValidateHierarchy() const
 {
 	if (!Branches.IsValidIndex(RootBranchIndex))
diff --git a/Engine/Plugins/Experimental/ProceduralVegetationEditor/Source/ProceduralVegetation/Private/Helpers/PVImportHelpers.h b/Engine/Plugins/Experimental/ProceduralVegetationEditor/Source/ProceduralVegetation/Private/Helpers/PVImportHelpers.h
--- a/Engine/Plugins/Experimental/ProceduralVegetationEditor/Source/ProceduralVegetation/Private/Helpers/PVImportHelpers.h
+++ b/Engine/Plugins/Experimental/ProceduralVegetationEditor/Source/ProceduralVegetation/Private/Helpers/PVImportHelpers.h
@@ -11,8 +11,6 @@ struct FManagedArrayCollection;
 namespace UE::Geometry
 {
 	class FDynamicMesh3;
-	template<class T> class TMeshAABBTree3;
-	using FDynamicMeshAABBTree3 = TMeshAABBTree3<FDynamicMesh3>;
 }
 
 struct FPVBranchDescription
@@ -155,6 +153,33 @@ namespace PV::ImportHelper
 
 	void FuseCollinearBranches(FPVBranchHierarchyDescription& InOutBranchHierarchy);
 
+	void ComputeLengthFromRoot(
+		const UE::Geometry::FDynamicMesh3& DynamicMesh,
+		const FBooleanVertexAttribute& RootPointsAttribute,
+		const FBooleanVertexAttribute& EndPointsAttribute,
+		FFloatVertexAttribute& LengthFromRootAttribute,
+		const TFunction<bool(int32)>& VertexFilterPredicate
+	);
+
+	void TraceBranches(
+		const UE::Geometry::FDynamicMesh3& DynamicMesh,
+		const FFloatVertexAttribute& LengthFromRootAttribute,
+		const FBooleanVertexAttribute& RootPointAttribute,
+		const FBooleanVertexAttribute& EndPointAttribute,
+		FIntVertexAttribute& NextBranchPointAttribute,
+		const TFunction<bool(int32)>& VertexFilterPredicate
+	);
+
+	void ExtractBranchHierarchies(
+		const UE::Geometry::FDynamicMesh3& DynamicMesh,
+		const FBooleanVertexAttribute& RootPointAttribute,
+		const FBooleanVertexAttribute& EndPointAttribute,
+		const FIntVertexAttribute& NextBranchPointAttribute,
+		const FFloatVertexAttribute& VertexScaleAttribute,
+		const FFloatVertexAttribute& LengthFromRootAttribute,
+		TArray<FMeshBranchHierarchy>& OutMeshBranchHierarchies
+	);
+
 	bool GenerateGrowthDataFromBranchHierarchies(FManagedArrayCollection& OutCollection, TArrayView<const FPVBranchHierarchyDescription> InBranchHierarchies, bool bSplitPoints = true);
 
 	bool GenerateGrowthDataFromBranchHierarchy(FManagedArrayCollection& OutCollection, const FPVBranchHierarchyDescription& InBranchHierarchy, bool bSplitPoints = true);
diff --git a/Engine/Plugins/Experimental/ProceduralVegetationEditor/Source/ProceduralVegetation/Private/Implementations/PVImporter_StaticMesh.cpp b/Engine/Plugins/Experimental/ProceduralVegetationEditor/Source/ProceduralVegetation/Private/Implementations/PVImporter_StaticMesh.cpp
--- a/Engine/Plugins/Experimental/ProceduralVegetationEditor/Source/ProceduralVegetation/Private/Implementations/PVImporter_StaticMesh.cpp
+++ b/Engine/Plugins/Experimental/ProceduralVegetationEditor/Source/ProceduralVegetation/Private/Implementations/PVImporter_StaticMesh.cpp
@@ -972,6 +972,383 @@ static void ConnectBranchesToParent(
 	}
 }
 
+static void ConnectAndTraceBranches(
+	FDynamicMesh3& DynamicMesh,
+	const FBooleanVertexAttribute& RootPointAttribute,
+	FFloatVertexAttribute& VertexScaleAttribute,
+	FBooleanVertexAttribute& EndPointAttribute,
+	FBooleanVertexAttribute& FillerVertexAttribute,
+	FFloatVertexAttribute& LengthFromRootAttribute,
+	FIntVertexAttribute& NextBranchPointAttribute
+)
+{
+	// Connects branch mesh components to their parent components (trunk or larger branches), choosing
+	// attachment edges that lie on the branch trace path rather than any geometrically-close edge.
+	// This avoids attaching at forks or off-path geometry which would confuse the downstream tracer.
+	//
+	// Algorithm overview:
+	//   Phase 0: Build a node graph (one node per connected component) and pre-compute all candidate
+	//            attachment edges for each (EndPoint, TargetNode) pair using a sphere-query, 
+	//            storing every valid edge instead of only the closest.
+	//   Phase 1: Run an initial ComputeLengthFromRoot + TraceBranches on the full mesh. Only the root
+	//            component obtains valid LengthFromRoot / NextBranchPoint values; all other components
+	//            have LengthFromRoot=-1 / NextBranchPoint=INDEX_NONE.
+	//   Phase 2: BFS level-by-level. For each level:
+	//            - For each already-connected node, find unvisited source nodes with EndPoints nearby.
+	//            - For each such source node, pick the candidate edge on the connected node that lies
+	//              on the current trace (NextBranchPoint[A]==B or NextBranchPoint[B]==A). Fall back to
+	//              the closest valid edge if no on-trace edge is found.
+	//            - Apply all accepted connections for this level (SplitEdge + MergeVertices).
+	//            - Re-trace so newly connected components obtain valid trace attributes for the next level.
+	// 
+	// Note that the ComputeLengthFromRoot+TraceBranches calls will re-trace the entire mesh. This is not 
+	// optimal, however, these calls are cheap compared to the rest of the import function (less than 1% 
+	// of the overal cost)
+
+	TRACE_CPUPROFILER_EVENT_SCOPE(ConnectAndTraceBranches);
+
+	using namespace PV::ImportHelper;
+
+	FMeshConnectedComponents ConnectedComponents(&DynamicMesh);
+	ConnectedComponents.FindConnectedVertices();
+
+	const TFunction<bool(int32)> VertexFilter = [&](int32 VertexID) { return !FillerVertexAttribute.GetValue(VertexID); };
+
+	if (ConnectedComponents.Components.Num() == 1)
+	{
+		// Single component: skip the connection step and go straight to trace.
+		ComputeLengthFromRoot(DynamicMesh, RootPointAttribute, EndPointAttribute, LengthFromRootAttribute, VertexFilter);
+		TraceBranches(DynamicMesh, LengthFromRootAttribute, RootPointAttribute, EndPointAttribute, NextBranchPointAttribute, VertexFilter);
+		return;
+	}
+
+	const TArray<int32> EndPointVertexIndices = EndPointAttribute.FindAllNonZero();
+	const int32 NumEndPoints = EndPointVertexIndices.Num();
+
+	TArray<int32> VertexIndexToNodeIndex;
+	{
+		VertexIndexToNodeIndex.SetNum(DynamicMesh.MaxVertexID());
+		for (int32 i = 0; i < ConnectedComponents.Components.Num(); ++i)
+		{
+			for (int32 VertexID : ConnectedComponents.Components[i].Indices)
+			{
+				VertexIndexToNodeIndex[VertexID] = i;
+			}
+		}
+	}
+
+	struct FCandidateEdge
+	{
+		FIndex2i Edge                 = FIndex2i(INDEX_NONE, INDEX_NONE);
+		FVector  TargetPoint          = FVector::ZeroVector;
+		float    Dist                 = 0.f;
+		float    TargetPointEdgeAlpha = 0.f;
+		float    TargetPointScale     = 0.f;
+	};
+
+	// EndPointsCandidateEdges[EndPointIndex][TargetNodeIndex] = list of valid candidate edges on TargetNode
+	TArray<TMap<int32, TArray<FCandidateEdge>>> EndPointsCandidateEdges;
+	EndPointsCandidateEdges.SetNum(NumEndPoints);
+
+	// NodesIncomingEndPoints[NodeIndex][IncomingNodeIndex] = list of incoming endpoints from a particular node
+	using FEndPointArray = TArray<int32, TInlineAllocator<4>>;
+	TArray<TMap<int32, FEndPointArray>> NodesIncomingEndPoints;
+	NodesIncomingEndPoints.SetNum(ConnectedComponents.Components.Num());
+
+	// Build candidate edges
+	const FDynamicMeshAABBTree3 DynamicMeshAABBTree(&DynamicMesh);
+	for (int32 EndPointIndex = 0; EndPointIndex < NumEndPoints; ++EndPointIndex)
+	{
+		const int32 EndPointVertexIndex = EndPointVertexIndices[EndPointIndex];
+		const FVector VertexPosition    = DynamicMesh.GetVertex(EndPointVertexIndex);
+		const int32 EndPointNodeIndex   = VertexIndexToNodeIndex[EndPointVertexIndex];
+		const float SearchRadius        = 100.f;
+
+		const TArray<int32> CloseTriangles = PV::DynamicMeshHelper::FindTrianglesInRadius(DynamicMeshAABBTree, VertexPosition, SearchRadius);
+
+		TMap<int32, TArray<FIndex3i>> NodeToTriangles;
+		for (const int32 TriangleID : CloseTriangles)
+		{
+			const FIndex3i Triangle = DynamicMesh.GetTriangle(TriangleID);
+			const int32 OtherNodeIndex = VertexIndexToNodeIndex[Triangle.A];
+			if (OtherNodeIndex == EndPointNodeIndex)
+			{
+				continue;
+			}
+			NodeToTriangles.FindOrAdd(OtherNodeIndex).Add(Triangle);
+		}
+
+		for (const auto& [NodeIndex, Triangles] : NodeToTriangles)
+		{
+			TArray<FCandidateEdge>& CandidateEdges = EndPointsCandidateEdges[EndPointIndex].FindOrAdd(NodeIndex);
+
+			for (const FIndex3i& Triangle : Triangles)
+			{
+				for (int32 j = 0; j < 3; ++j)
+				{
+					FIndex2i Edge = FIndex2i(Triangle[j], Triangle[(j + 1) % 3]);
+					Edge.Sort();
+
+					if (FillerVertexAttribute.GetValue(Edge.A) || FillerVertexAttribute.GetValue(Edge.B))
+					{
+						continue;
+					}
+
+					const FVector VertexA      = DynamicMesh.GetVertex(Edge.A);
+					const FVector VertexB      = DynamicMesh.GetVertex(Edge.B);
+					const FVector ClosestPoint = FMath::ClosestPointOnSegment(VertexPosition, VertexA, VertexB);
+
+					const float DistAlongEdge = (VertexA - ClosestPoint).Size();
+					const float EdgeSize      = (VertexB - VertexA).Size();
+					const float EdgeAlpha     = EdgeSize > 0.f ? DistAlongEdge / EdgeSize : 0.f;
+
+					const float VertexScaleA    = VertexScaleAttribute.GetValue(Edge.A);
+					const float VertexScaleB    = VertexScaleAttribute.GetValue(Edge.B);
+					const float TargetScale     = FMath::Lerp(VertexScaleA, VertexScaleB, EdgeAlpha);
+					const float ThisPointScale  = VertexScaleAttribute.GetValue(EndPointVertexIndex);
+
+					const float DistToEdge = (VertexPosition - ClosestPoint).Size();
+					if (DistToEdge > (TargetScale + ThisPointScale) * 1.5f)
+					{
+						continue; // radii do not overlap — not a valid attachment
+					}
+
+					const bool bAlreadyAdded = CandidateEdges.FindByPredicate([&](const auto& X) { return X.Edge == Edge; }) != nullptr;
+					if (bAlreadyAdded)
+					{
+						continue;
+					}
+
+					FCandidateEdge& Candidate = CandidateEdges.AddDefaulted_GetRef();
+					Candidate.Edge                 = Edge;
+					Candidate.TargetPoint          = ClosestPoint;
+					Candidate.Dist                 = DistToEdge;
+					Candidate.TargetPointEdgeAlpha = EdgeAlpha;
+					Candidate.TargetPointScale     = TargetScale;
+				}
+			}
+
+			if (CandidateEdges.Num() > 0)
+			{
+				NodesIncomingEndPoints[NodeIndex].FindOrAdd(EndPointNodeIndex).Add(EndPointIndex);
+			}
+		}
+	}
+
+	// Phase 1: Initial trace — seeds from real root points.  Disconnected components remain at
+	// LengthFromRoot=-1 / NextBranchPoint=INDEX_NONE, which is the signal used in Phase 2 to
+	// identify edges that are not yet on the trace.
+	ComputeLengthFromRoot(DynamicMesh, RootPointAttribute, EndPointAttribute, LengthFromRootAttribute, VertexFilter);
+	TraceBranches(DynamicMesh, LengthFromRootAttribute, RootPointAttribute, EndPointAttribute, NextBranchPointAttribute, VertexFilter);
+
+	// Phase 2: Iterative BFS.  Process one hop per loop iteration so we can re-trace after each
+	// level and use the updated NextBranchPoint values when selecting attachment edges for the next.
+	TArray<bool> VisitedNodes;
+	VisitedNodes.SetNumZeroed(ConnectedComponents.Components.Num());
+
+	int32 BFSQueueIndex = 0;
+	TArray<int32> BFSQueues[2];
+	{
+		const TArray<int32> RootPointVertexIndices = RootPointAttribute.FindAllNonZero();
+		for (int32 RootPointVertexIndex : RootPointVertexIndices)
+		{
+			const int32 RootNodeIndex = VertexIndexToNodeIndex[RootPointVertexIndex];
+			if (!VisitedNodes[RootNodeIndex])
+			{
+				VisitedNodes[RootNodeIndex] = true;
+				BFSQueues[BFSQueueIndex].Add(RootNodeIndex);
+			}
+		}
+	}
+
+	struct FLevelConnection
+	{
+		int32    TargetNode           = INDEX_NONE;
+		int32    SourceNode           = INDEX_NONE;
+		int32    EndPointIndex        = INDEX_NONE;
+		FIndex2i TargetEdge           = FIndex2i(INDEX_NONE, INDEX_NONE);
+		FVector  TargetPoint          = FVector::ZeroVector;
+		float    TargetPointEdgeAlpha = 0.f;
+		float    TargetPointScale     = 0.f;
+		float    Dist                 = MAX_flt;
+	};
+	TArray<FLevelConnection> LevelConnections;
+	TArray<int32> EndPointToAttachVertex;
+
+	while (BFSQueues[BFSQueueIndex].Num() > 0)
+	{
+		const int32 NextBFSQueueIndex = (BFSQueueIndex + 1) % 2;
+		TArray<int32>& NextBFSQueue = BFSQueues[NextBFSQueueIndex];
+		BFSQueues[NextBFSQueueIndex].Reset();
+		LevelConnections.Reset();
+
+		for (int32 NodeIndex : BFSQueues[BFSQueueIndex])
+		{
+			for (const auto& [SourceNodeIndex, EndPointIndices] : NodesIncomingEndPoints[NodeIndex])
+			{
+				if (VisitedNodes[SourceNodeIndex])
+				{
+					continue;
+				}
+
+				// Find the best (EndPoint, Edge) pair for connecting SourceNode to NodeIndex.
+				FLevelConnection BestConnection;
+				BestConnection.TargetNode = NodeIndex;
+				BestConnection.SourceNode = SourceNodeIndex;
+
+				for (int32 EndPointIndex : EndPointIndices)
+				{
+					const TArray<FCandidateEdge>* CandidatesPtr = EndPointsCandidateEdges[EndPointIndex].Find(NodeIndex);
+					if (!CandidatesPtr || CandidatesPtr->Num() == 0)
+					{
+						continue;
+					}
+
+					bool     bFoundOnTrace    = false;
+					FIndex2i BestEdge         = FIndex2i(INDEX_NONE, INDEX_NONE);
+					FVector  BestTargetPoint  = FVector::ZeroVector;
+					float    BestDist         = MAX_flt;
+					float    BestEdgeAlpha    = 0.f;
+					float    BestTargetScale  = 0.f;
+
+					for (const FCandidateEdge& Candidate : *CandidatesPtr)
+					{
+						const int32 NextA   = NextBranchPointAttribute.GetValue(Candidate.Edge.A);
+						const int32 NextB   = NextBranchPointAttribute.GetValue(Candidate.Edge.B);
+						const bool bOnTrace = (NextA == Candidate.Edge.B) || (NextB == Candidate.Edge.A);
+
+						if (bOnTrace && (!bFoundOnTrace || Candidate.Dist < BestDist))
+						{
+							bFoundOnTrace   = true;
+							BestEdge        = Candidate.Edge;
+							BestTargetPoint = Candidate.TargetPoint;
+							BestDist        = Candidate.Dist;
+							BestEdgeAlpha   = Candidate.TargetPointEdgeAlpha;
+							BestTargetScale = Candidate.TargetPointScale;
+						}
+						else if (!bOnTrace && !bFoundOnTrace && Candidate.Dist < BestDist)
+						{
+							// No on-trace candidate yet — keep the closest overall as fallback.
+							BestEdge        = Candidate.Edge;
+							BestTargetPoint = Candidate.TargetPoint;
+							BestDist        = Candidate.Dist;
+							BestEdgeAlpha   = Candidate.TargetPointEdgeAlpha;
+							BestTargetScale = Candidate.TargetPointScale;
+						}
+					}
+
+					if (BestEdge.A != INDEX_NONE && BestDist < BestConnection.Dist)
+					{
+						BestConnection.EndPointIndex        = EndPointIndex;
+						BestConnection.TargetEdge           = BestEdge;
+						BestConnection.TargetPoint          = BestTargetPoint;
+						BestConnection.Dist                 = BestDist;
+						BestConnection.TargetPointEdgeAlpha = BestEdgeAlpha;
+						BestConnection.TargetPointScale     = BestTargetScale;
+					}
+				}
+
+				if (BestConnection.EndPointIndex != INDEX_NONE)
+				{
+					LevelConnections.Add(BestConnection);
+					NextBFSQueue.Add(SourceNodeIndex);
+					VisitedNodes[SourceNodeIndex] = true;
+				}
+			}
+		}
+
+		if (LevelConnections.IsEmpty())
+		{
+			break;
+		}
+
+		// Apply connections: group by target edge, split with tail-advancing, then merge vertices.
+		EndPointToAttachVertex.Init(INDEX_NONE, NumEndPoints);
+
+		TMap<FIndex2i, TArray<int32>> EdgeToConnectionIndices;
+		for (int32 i = 0; i < LevelConnections.Num(); ++i)
+		{
+			EdgeToConnectionIndices.FindOrAdd(LevelConnections[i].TargetEdge).Add(i);
+		}
+
+		for (auto& [Edge, ConnectionIndices] : EdgeToConnectionIndices)
+		{
+			ConnectionIndices.Sort([&](int32 A, int32 B)
+			{
+				return LevelConnections[A].TargetPointEdgeAlpha < LevelConnections[B].TargetPointEdgeAlpha;
+			});
+
+			int32 TailA       = Edge.A;
+			const int32 TailB = Edge.B;
+			for (const int32 ConnIdx : ConnectionIndices)
+			{
+				const FLevelConnection& Conn = LevelConnections[ConnIdx];
+
+				const bool bSnapToTailA = DynamicMesh.GetVertex(TailA).Equals(Conn.TargetPoint, 1.f);
+				if (bSnapToTailA)
+				{
+					EndPointToAttachVertex[Conn.EndPointIndex] = TailA;
+					continue;
+				}
+
+				const bool bSnapToTailB = DynamicMesh.GetVertex(TailB).Equals(Conn.TargetPoint, 1.f);
+				if (bSnapToTailB)
+				{
+					EndPointToAttachVertex[Conn.EndPointIndex] = TailB;
+					continue;
+				}
+
+				FDynamicMesh3::FEdgeSplitInfo EdgeSplitInfo;
+				const UE::Geometry::EMeshResult SplitResult = DynamicMesh.SplitEdge(TailA, TailB, EdgeSplitInfo);
+				if (!ensure(SplitResult == UE::Geometry::EMeshResult::Ok))
+				{
+					continue;
+				}
+
+				DynamicMesh.SetVertex(EdgeSplitInfo.NewVertex, Conn.TargetPoint);
+				VertexScaleAttribute.SetValue(EdgeSplitInfo.NewVertex, Conn.TargetPointScale);
+				FillerVertexAttribute.SetValue(EdgeSplitInfo.NewVertex, false);
+
+				EndPointToAttachVertex[Conn.EndPointIndex] = EdgeSplitInfo.NewVertex;
+				TailA = EdgeSplitInfo.NewVertex;
+			}
+		}
+
+		for (const FLevelConnection& Conn : LevelConnections)
+		{
+			const int32 EndPointVertexIndex = EndPointVertexIndices[Conn.EndPointIndex];
+			const int32 AttachVertex        = EndPointToAttachVertex[Conn.EndPointIndex];
+			if (AttachVertex == INDEX_NONE)
+			{
+				continue;
+			}
+
+			EndPointAttribute.SetValue(EndPointVertexIndex, false);
+			EndPointAttribute.SetValue(AttachVertex, false);
+
+			FDynamicMesh3::FMergeVerticesOptions MergeOptions;
+			MergeOptions.bAllowNonBoundaryBowtieCreation = true; // AttachVertex is interior (from SplitEdge)
+
+			FDynamicMesh3::FMergeVerticesInfo MergeInfo;
+			const UE::Geometry::EMeshResult MergeResult = DynamicMesh.MergeVertices(
+				AttachVertex,        // KeepVid: attachment point on the parent skeleton
+				EndPointVertexIndex, // DiscardVid: endpoint being welded in
+				0.0,                 // InterpolationT: 0 = keep KeepVid position exactly
+				MergeOptions,
+				MergeInfo
+			);
+			ensure(MergeResult == UE::Geometry::EMeshResult::Ok);
+		}
+
+		// Re-trace so the newly connected components gain valid LengthFromRoot / NextBranchPoint
+		// values, which are used to select on-trace attachment edges in the next BFS level.
+		ComputeLengthFromRoot(DynamicMesh, RootPointAttribute, EndPointAttribute, LengthFromRootAttribute, VertexFilter);
+		TraceBranches(DynamicMesh, LengthFromRootAttribute, RootPointAttribute, EndPointAttribute, NextBranchPointAttribute, VertexFilter);
+
+		BFSQueueIndex = NextBFSQueueIndex;
+	}
+}
+
 static void FindRootPoints(
 	const FDynamicMesh3& DynamicMesh,
 	FBooleanVertexAttribute& RootPointAttribute,
@@ -1105,32 +1482,77 @@ PV::StaticMeshImport::EImportResult PV::StaticMeshImport::ImportGrowthDataFromSt
 		EndPointAttributeDefinition.CopyAttributeToCollection(DynamicMesh, DebugCollection);
 	}
 
-	ConnectBranchesToParent(
-		DynamicMesh,
-		RootPointAttributeDefinition.GetOrAttachAttribute(DynamicMesh),
-		VertexScaleAttributeDefinition.GetOrAttachAttribute(DynamicMesh),
-		EndPointAttributeDefinition.GetOrAttachAttribute(DynamicMesh),
-		FillerVertexAttributeDefinition.GetOrAttachAttribute(DynamicMesh)
+	TArray<FMeshBranchHierarchy> DynamicMeshBranchHierarchies;
+
+	static TAutoConsoleVariable<int32> ConnectBranchesToParentVersion(
+		TEXT("PVE.Import.ConnectBranchesToParent"),
+		2,
+		TEXT("PVE.Import.ConnectBranchesToParent {VERSION}"),
+		ECVF_Default
 	);
-	if (InDebugParams.DebugState == EPVImportStaticMeshDebugState::ConnectBranchesToParent)
+
+	const bool bUseNewConnectBranchesToParent = ConnectBranchesToParentVersion.GetValueOnAnyThread() == 2;
+	if (bUseNewConnectBranchesToParent)
 	{
-		FManagedArrayCollection& DebugCollection = Output.DebugCollection.Emplace();
-		AddDynamicMeshToCollection(DynamicMesh, DebugCollection);
-		EndPointAttributeDefinition.CopyAttributeToCollection(DynamicMesh, DebugCollection);
+		ConnectAndTraceBranches(
+			DynamicMesh,
+			RootPointAttributeDefinition.GetOrAttachAttribute(DynamicMesh),
+			VertexScaleAttributeDefinition.GetOrAttachAttribute(DynamicMesh),
+			EndPointAttributeDefinition.GetOrAttachAttribute(DynamicMesh),
+			FillerVertexAttributeDefinition.GetOrAttachAttribute(DynamicMesh),
+			LengthFromRootAttributeDefinition.GetOrAttachAttribute(DynamicMesh),
+			NextBranchPointAttributeDefinition.GetOrAttachAttribute(DynamicMesh)
+		);
+		if (InDebugParams.DebugState == EPVImportStaticMeshDebugState::ConnectBranchesToParent)
+		{
+			FManagedArrayCollection& DebugCollection = Output.DebugCollection.Emplace();
+			AddDynamicMeshToCollection(DynamicMesh, DebugCollection);
+			RootPointAttributeDefinition.CopyAttributeToCollection(DynamicMesh, DebugCollection);
+			VertexScaleAttributeDefinition.CopyAttributeToCollection(DynamicMesh, DebugCollection);
+			EndPointAttributeDefinition.CopyAttributeToCollection(DynamicMesh, DebugCollection);
+			FillerVertexAttributeDefinition.CopyAttributeToCollection(DynamicMesh, DebugCollection);
+			LengthFromRootAttributeDefinition.CopyAttributeToCollection(DynamicMesh, DebugCollection);
+			NextBranchPointAttributeDefinition.CopyAttributeToCollection(DynamicMesh, DebugCollection);
+		}
+
+		ExtractBranchHierarchies(
+			DynamicMesh, 
+			RootPointAttributeDefinition.GetOrAttachAttribute(DynamicMesh),
+			EndPointAttributeDefinition.GetOrAttachAttribute(DynamicMesh),
+			NextBranchPointAttributeDefinition.GetOrAttachAttribute(DynamicMesh),
+			VertexScaleAttributeDefinition.GetOrAttachAttribute(DynamicMesh),
+			LengthFromRootAttributeDefinition.GetOrAttachAttribute(DynamicMesh),
+			DynamicMeshBranchHierarchies
+		);
 	}
+	else
+	{
+		ConnectBranchesToParent(
+			DynamicMesh,
+			RootPointAttributeDefinition.GetOrAttachAttribute(DynamicMesh),
+			VertexScaleAttributeDefinition.GetOrAttachAttribute(DynamicMesh),
+			EndPointAttributeDefinition.GetOrAttachAttribute(DynamicMesh),
+			FillerVertexAttributeDefinition.GetOrAttachAttribute(DynamicMesh)
+		);
+		if (InDebugParams.DebugState == EPVImportStaticMeshDebugState::ConnectBranchesToParent)
+		{
+			FManagedArrayCollection& DebugCollection = Output.DebugCollection.Emplace();
+			AddDynamicMeshToCollection(DynamicMesh, DebugCollection);
+			EndPointAttributeDefinition.CopyAttributeToCollection(DynamicMesh, DebugCollection);
+		}
 
-	auto& FillerVertexAttribute = FillerVertexAttributeDefinition.GetOrAttachAttribute(DynamicMesh);
-	TArray<FMeshBranchHierarchy> DynamicMeshBranchHierarchies;
-	ComputeDynamicMeshBranchHierarchies(
-		DynamicMesh,
-		RootPointAttributeDefinition.GetOrAttachAttribute(DynamicMesh),
-		EndPointAttributeDefinition.GetOrAttachAttribute(DynamicMesh),
-		VertexScaleAttributeDefinition.GetOrAttachAttribute(DynamicMesh),
-		LengthFromRootAttributeDefinition.GetOrAttachAttribute(DynamicMesh),
-		NextBranchPointAttributeDefinition.GetOrAttachAttribute(DynamicMesh),
-		[&](int32 VertexID) { return !FillerVertexAttribute.GetValue(VertexID); },
-		DynamicMeshBranchHierarchies
-	);
+		auto& FillerVertexAttribute = FillerVertexAttributeDefinition.GetOrAttachAttribute(DynamicMesh);
+		ComputeDynamicMeshBranchHierarchies(
+			DynamicMesh,
+			RootPointAttributeDefinition.GetOrAttachAttribute(DynamicMesh),
+			EndPointAttributeDefinition.GetOrAttachAttribute(DynamicMesh),
+			VertexScaleAttributeDefinition.GetOrAttachAttribute(DynamicMesh),
+			LengthFromRootAttributeDefinition.GetOrAttachAttribute(DynamicMesh),
+			NextBranchPointAttributeDefinition.GetOrAttachAttribute(DynamicMesh),
+			[&](int32 VertexID) { return !FillerVertexAttribute.GetValue(VertexID); },
+			DynamicMeshBranchHierarchies
+		);
+	}
 
 	TArray<FPVBranchHierarchyDescription> BranchHierarchyDescriptions;
 	ConvertDynamicMeshBranchHierarchiesToBranchDescription(
