Fix UnityBuild break due to anonymous namespace leading to name collision

#Synced-CL 52140050

diff --git a/Engine/Plugins/Experimental/ProceduralVegetationEditor/Source/ProceduralVegetation/Private/Helpers/PVImportHelpers.cpp b/Engine/Plugins/Experimental/ProceduralVegetationEditor/Source/ProceduralVegetation/Private/Helpers/PVImportHelpers.cpp
--- a/Engine/Plugins/Experimental/ProceduralVegetationEditor/Source/ProceduralVegetation/Private/Helpers/PVImportHelpers.cpp
+++ b/Engine/Plugins/Experimental/ProceduralVegetationEditor/Source/ProceduralVegetation/Private/Helpers/PVImportHelpers.cpp
@@ -12,7 +12,7 @@
 
 #include "Engine/StaticMesh.h"
 
-namespace
+namespace PV::ImportHelper
 {
 	const static TSet<FName> GeometryGroups = {
 		FGeometryCollection::VerticesGroup,
