[PVE] Rename UPVOutputSettings to UPVExportSettings

For the subgraph, we require PCG's output node which conflicts with PVE's Output node
in the graph editor. To avoid confusion, rename PVE's Output node class
from UPVOutputSettings to UPVExportSettings and update all references.

Also adds a CoreRedirects entry in BaseProceduralVegetationEditor.ini so
existing assets referencing the old class name are automatically resolved
on load without requiring a resave.

#Synced-CL 52095318

diff --git a/Engine/Plugins/Experimental/ProceduralVegetationEditor/Config/BaseProceduralVegetationEditor.ini b/Engine/Plugins/Experimental/ProceduralVegetationEditor/Config/BaseProceduralVegetationEditor.ini
new file mode 100644
--- /dev/null
+++ b/Engine/Plugins/Experimental/ProceduralVegetationEditor/Config/BaseProceduralVegetationEditor.ini
@@ -0,0 +1,2 @@
+﻿[CoreRedirects]
++ClassRedirects=(OldName="/Script/ProceduralVegetation.PVOutputSettings",NewName="/Script/ProceduralVegetation.PVExportSettings")
\ No newline at end of file
diff --git a/Engine/Plugins/Experimental/ProceduralVegetationEditor/Source/ProceduralVegetation/Private/Nodes/PVExportSettings.cpp b/Engine/Plugins/Experimental/ProceduralVegetationEditor/Source/ProceduralVegetation/Private/Nodes/PVExportSettings.cpp
new file mode 100644
--- /dev/null
+++ b/Engine/Plugins/Experimental/ProceduralVegetationEditor/Source/ProceduralVegetation/Private/Nodes/PVExportSettings.cpp
@@ -0,0 +1,134 @@
+﻿// Copyright Epic Games, Inc. All Rights Reserved.
+
+#include "Nodes/PVExportSettings.h"
+#include "PCGContext.h"
+#include "DataTypes/PVMeshData.h"
+#include "DataTypes/PVFoliageMeshData.h"
+#include "Helpers/PCGHelpers.h"
+#include "Misc/PackageName.h"
+#include "Engine/StaticMesh.h"
+#include "Engine/SkeletalMesh.h"
+#include "Helpers/PVAnalyticsHelper.h"
+
+
+#define LOCTEXT_NAMESPACE "PVExportSettings"
+
+#if WITH_EDITOR
+FText UPVExportSettings::GetDefaultNodeTitle() const
+{
+	return ExportSettings.MeshName.IsNone() 
+		? LOCTEXT("NodeTitle", "Export") 
+		: FText::FromName(ExportSettings.MeshName);
+}
+
+FText UPVExportSettings::GetNodeTooltipText() const 
+{ 
+	return LOCTEXT("NodeTooltip", 
+		"This node allows the user to save the output of the graph as:\n"
+		"- Static Mesh\n"
+		"- Static Mesh with Nanite Foliage enabled\n"
+		"- Skeletal Mesh\n"
+		"- Skeletal Mesh with Nanite Foliage enabled\n\n"
+
+		"If the foliage meshes do not contain skeletal data and Skeletal mesh export is selected, it will build the skeletal data for foliage meshes"
+		"\n\nPress Ctrl + L to lock/unlock node output"
+	); 
+}
+#endif
+
+UPVExportSettings::UPVExportSettings()
+{
+	if (PCGHelpers::IsNewObjectAndNotDefault(this))
+	{
+		static FString DefaultMeshName = "Export";
+		static FString DefaultAssetPath = "/Game/ProceduralVegetation/";
+		static FString DefaultWindSettingsPath = "/ProceduralVegetationEditor/SampleAssets/WindSettings/DefaultTreeWindSettings.DefaultTreeWindSettings";
+		ExportSettings.Initialize(DefaultAssetPath, DefaultMeshName, DefaultWindSettingsPath);
+	}
+}
+
+FPCGDataTypeIdentifier UPVExportSettings::GetInputPinTypeIdentifier() const
+{
+	return FPCGDataTypeIdentifier::Construct(FPVDataTypeInfoMesh::AsId(), FPVDataTypeInfoFoliageMesh::AsId());
+}
+
+TArray<FPCGPinProperties> UPVExportSettings::OutputPinProperties() const
+{
+	TArray<FPCGPinProperties> Properties;
+
+	FPCGPinProperties& Pin = Properties.Emplace_GetRef(PCGPinConstants::DefaultOutputLabel, EPCGDataType::Any);
+	Pin.SetRequiredPin();
+	Pin.SetAllowMultipleConnections(false);
+	Pin.bInvisiblePin = true;
+
+	return Properties;
+}
+
+FPCGElementPtr UPVExportSettings::CreateElement() const
+{
+	return MakeShared<FPVExportElement>();
+}
+
+#if	WITH_EDITOR
+void UPVExportSettings::PostEditChangeProperty(FPropertyChangedEvent& PropertyChangedEvent)
+{
+	Super::PostEditChangeProperty(PropertyChangedEvent);
+
+	if (PropertyChangedEvent.Property
+		&& PropertyChangedEvent.Property->GetFName() == GET_MEMBER_NAME_CHECKED(FPVExportParams, WindSettings))
+	{
+		if (ExportSettings.WindSettings)
+		{
+			PV::Analytics::SendWindSettingsChangeEvent(ExportSettings.WindSettings->GetPathName());
+		}
+	}
+}
+
+EPCGChangeType UPVExportSettings::GetChangeTypeForProperty(FPropertyChangedEvent& PropertyChangedEvent) const
+{
+	const FName PropertyName = PropertyChangedEvent.GetPropertyName();
+	const FName MemberPropertyName = PropertyChangedEvent.GetMemberPropertyName();
+
+	if (MemberPropertyName == GET_MEMBER_NAME_CHECKED(UPVExportSettings, ExportSettings))
+	{
+		if (PropertyName == GET_MEMBER_NAME_CHECKED(FDirectoryPath, Path)
+			||PropertyName == GET_MEMBER_NAME_CHECKED(FPVExportParams, MeshName))
+		{
+			return EPCGChangeType::Cosmetic;
+		}
+	}
+
+	return Super::GetChangeTypeForProperty(PropertyChangedEvent);
+}
+
+#endif
+
+bool FPVExportElement::ExecuteInternal(FPCGContext* InContext) const
+{
+	const UPVExportSettings* Settings = InContext->GetInputSettings<UPVExportSettings>();
+	check(Settings);
+
+	for (const FPCGTaggedData& Input : InContext->InputData.GetInputsByPin(PCGPinConstants::DefaultInputLabel))
+	{
+		if (const UPVData* InputData = Cast<UPVData>(Input.Data))
+		{
+			FString Error;
+			if (!Settings->ExportSettings.Validate(Error))
+			{
+				PCGLog::LogErrorOnGraph(FText::FromString(Error), InContext);
+				return true;
+			}
+
+			InContext->OutputData = InContext->InputData;
+		}
+		else
+		{
+			PCGLog::InputOutput::LogInvalidInputDataError(InContext);
+			return true;
+		}
+	}
+
+	return true;
+}
+
+#undef LOCTEXT_NAMESPACE
diff --git a/Engine/Plugins/Experimental/ProceduralVegetationEditor/Source/ProceduralVegetation/Private/Nodes/PVOutputSettings.cpp b/Engine/Plugins/Experimental/ProceduralVegetationEditor/Source/ProceduralVegetation/Private/Nodes/PVOutputSettings.cpp
deleted file mode 100644
diff --git a/Engine/Plugins/Experimental/ProceduralVegetationEditor/Source/ProceduralVegetation/Private/ProceduralVegetation.cpp b/Engine/Plugins/Experimental/ProceduralVegetationEditor/Source/ProceduralVegetation/Private/ProceduralVegetation.cpp
--- a/Engine/Plugins/Experimental/ProceduralVegetationEditor/Source/ProceduralVegetation/Private/ProceduralVegetation.cpp
+++ b/Engine/Plugins/Experimental/ProceduralVegetationEditor/Source/ProceduralVegetation/Private/ProceduralVegetation.cpp
@@ -5,7 +5,7 @@
 #include "Nodes/PVGravitySettings.h"
 #include "Nodes/PVMeshBuilderSettings.h"
 #include "Nodes/PVPresetLoaderSettings.h"
-#include "Nodes/PVOutputSettings.h"
+#include "Nodes/PVExportSettings.h"
 #include "ProceduralVegetationPreset.h"
 
 namespace PV
@@ -78,7 +78,7 @@ void UProceduralVegetation::CreateGraphFromPreset(const TObjectPtr<UProceduralVe
 		const UPVGravitySettings* GravitySettings = NewObject<UPVGravitySettings>(Graph);
 		const UPVMeshBuilderSettings* MeshBuilderSettings = NewObject<UPVMeshBuilderSettings>(Graph);
 		const UPVFoliagePaletteSettings* FoliagePaletteSettings = NewObject<UPVFoliagePaletteSettings>(Graph);
-		UPVOutputSettings* OutputSettings = NewObject<UPVOutputSettings>(Graph);
+		UPVExportSettings* OutputSettings = NewObject<UPVExportSettings>(Graph);
 		FString MeshName = *VariantName;
 		PV::SanitizeFolderName(MeshName);
 		OutputSettings->ExportSettings.MeshName = *MeshName;
diff --git a/Engine/Plugins/Experimental/ProceduralVegetationEditor/Source/ProceduralVegetation/Public/Nodes/PVExportSettings.h b/Engine/Plugins/Experimental/ProceduralVegetationEditor/Source/ProceduralVegetation/Public/Nodes/PVExportSettings.h
new file mode 100644
--- /dev/null
+++ b/Engine/Plugins/Experimental/ProceduralVegetationEditor/Source/ProceduralVegetation/Public/Nodes/PVExportSettings.h
@@ -0,0 +1,47 @@
+﻿// Copyright Epic Games, Inc. All Rights Reserved.
+#pragma once
+
+#include "CoreMinimal.h"
+#include "Params/PVExportParams.h"
+#include "Nodes/PVBaseSettings.h"
+#include "PVExportSettings.generated.h"
+
+UCLASS(ClassGroup = (Procedural))
+class UPVExportSettings : public UPVBaseSettings
+{
+	GENERATED_BODY()
+
+public:
+	UPVExportSettings();
+
+#if WITH_EDITOR
+	virtual FName GetDefaultNodeName() const override { return ExportSettings.MeshName.IsNone() ? "Export" : ExportSettings.MeshName; }
+	virtual FText GetDefaultNodeTitle() const override;
+	virtual FText GetNodeTooltipText() const override;
+	virtual EPCGSettingsType GetType() const override { return EPCGSettingsType::InputOutput; }
+	virtual TArray<EPVRenderType> GetSupportedRenderTypes() const override { return TArray{EPVRenderType::PointData, EPVRenderType::Foliage, EPVRenderType::FoliageGrid, EPVRenderType::FoliageAttachments, EPVRenderType::Mesh, EPVRenderType::Bones}; }
+	virtual TArray<EPVRenderType> GetDefaultRenderType() const override { return TArray{EPVRenderType::Foliage, EPVRenderType::Mesh}; }
+#endif
+
+protected:
+	virtual FPCGDataTypeIdentifier GetInputPinTypeIdentifier() const override;
+	virtual TArray<FPCGPinProperties> OutputPinProperties() const override;
+
+	virtual FPCGElementPtr CreateElement() const override;
+
+#if WITH_EDITOR
+	virtual void PostEditChangeProperty(FPropertyChangedEvent& PropertyChangedEvent) override;
+	virtual EPCGChangeType GetChangeTypeForProperty(FPropertyChangedEvent& PropertyChangedEvent) const override;
+#endif
+	
+public:
+	UPROPERTY(EditAnywhere, Category = ExportSettings, meta = (PCG_Overridable, ShowOnlyInnerProperties))
+	FPVExportParams ExportSettings;
+};
+
+class FPVExportElement : public FPVBaseElement
+{
+	virtual bool ExecuteInternal(FPCGContext* Context) const override;
+
+	virtual bool CanExecuteOnlyOnMainThread(FPCGContext* Context) const override { return true; }
+};
\ No newline at end of file
diff --git a/Engine/Plugins/Experimental/ProceduralVegetationEditor/Source/ProceduralVegetation/Public/Nodes/PVOutputSettings.h b/Engine/Plugins/Experimental/ProceduralVegetationEditor/Source/ProceduralVegetation/Public/Nodes/PVOutputSettings.h
deleted file mode 100644
diff --git a/Engine/Plugins/Experimental/ProceduralVegetationEditor/Source/ProceduralVegetationEditor/Private/Customizations/PVOutputSettingsCustomizations.cpp b/Engine/Plugins/Experimental/ProceduralVegetationEditor/Source/ProceduralVegetationEditor/Private/Customizations/PVOutputSettingsCustomizations.cpp
--- a/Engine/Plugins/Experimental/ProceduralVegetationEditor/Source/ProceduralVegetationEditor/Private/Customizations/PVOutputSettingsCustomizations.cpp
+++ b/Engine/Plugins/Experimental/ProceduralVegetationEditor/Source/ProceduralVegetationEditor/Private/Customizations/PVOutputSettingsCustomizations.cpp
@@ -10,7 +10,7 @@
 
 #include "Framework/Docking/TabManager.h"
 
-#include "Nodes/PVOutputSettings.h"
+#include "Nodes/PVExportSettings.h"
 
 #include "Widgets/Input/SButton.h"
 #include "Widgets/Layout/SBox.h"
diff --git a/Engine/Plugins/Experimental/ProceduralVegetationEditor/Source/ProceduralVegetationEditor/Private/Exporter/PVExporter.cpp b/Engine/Plugins/Experimental/ProceduralVegetationEditor/Source/ProceduralVegetationEditor/Private/Exporter/PVExporter.cpp
--- a/Engine/Plugins/Experimental/ProceduralVegetationEditor/Source/ProceduralVegetationEditor/Private/Exporter/PVExporter.cpp
+++ b/Engine/Plugins/Experimental/ProceduralVegetationEditor/Source/ProceduralVegetationEditor/Private/Exporter/PVExporter.cpp
@@ -65,7 +65,7 @@ void UPVExportEntry::Initialize(
 	bIsNodeSelected = InIsNodeSelected;
 	Node = InNode;
 	bExportFoliage = InHasFoliage;
-	Settings = CastChecked<UPVOutputSettings>(InNode->GetSettings());
+	Settings = CastChecked<UPVExportSettings>(InNode->GetSettings());
 	OutputCollection = InCollection;
 
 	check(Settings && OutputCollection);
diff --git a/Engine/Plugins/Experimental/ProceduralVegetationEditor/Source/ProceduralVegetationEditor/Private/Exporter/PVExporter.h b/Engine/Plugins/Experimental/ProceduralVegetationEditor/Source/ProceduralVegetationEditor/Private/Exporter/PVExporter.h
--- a/Engine/Plugins/Experimental/ProceduralVegetationEditor/Source/ProceduralVegetationEditor/Private/Exporter/PVExporter.h
+++ b/Engine/Plugins/Experimental/ProceduralVegetationEditor/Source/ProceduralVegetationEditor/Private/Exporter/PVExporter.h
@@ -3,7 +3,7 @@
 #pragma once
 
 #include "Nodes/PVBaseSettings.h"
-#include "Nodes/PVOutputSettings.h"
+#include "Nodes/PVExportSettings.h"
 
 #include "PVExporter.generated.h"
 
@@ -62,7 +62,7 @@ public:
 	FPVOutputEntryStats Stats;
 
 	UPROPERTY(EditAnywhere, Category="Export Settings", meta=(ShowOnlyInnerProperties))
-	TObjectPtr<UPVOutputSettings> Settings = nullptr;
+	TObjectPtr<UPVExportSettings> Settings = nullptr;
 
 	TSharedPtr<FManagedArrayCollection> OutputCollection;
 
diff --git a/Engine/Plugins/Experimental/ProceduralVegetationEditor/Source/ProceduralVegetationEditor/Private/PVEditor.cpp b/Engine/Plugins/Experimental/ProceduralVegetationEditor/Source/ProceduralVegetationEditor/Private/PVEditor.cpp
--- a/Engine/Plugins/Experimental/ProceduralVegetationEditor/Source/ProceduralVegetationEditor/Private/PVEditor.cpp
+++ b/Engine/Plugins/Experimental/ProceduralVegetationEditor/Source/ProceduralVegetationEditor/Private/PVEditor.cpp
@@ -40,7 +40,7 @@
 #include "Misc/UObjectToken.h"
 
 #include "Nodes/PCGEditorGraphNodeBase.h"
-#include "Nodes/PVOutputSettings.h"
+#include "Nodes/PVExportSettings.h"
 
 #include "PCGEditor/Private/PCGEditorGraph.h"
 
@@ -111,7 +111,7 @@ void FPVEditor::Initialize(
 	{
 		if (const UPCGEditorGraphNodeBase* PCGEdNode = Cast<UPCGEditorGraphNodeBase>(EdNode))
 		{
-			if (PCGEdNode->GetPCGNode()->GetSettings()->IsA<UPVOutputSettings>())
+			if (PCGEdNode->GetPCGNode()->GetSettings()->IsA<UPVExportSettings>())
 			{
 				GetMainEditorGraph()->SelectNodeSet({PCGEdNode}, true);
 				break;
@@ -704,7 +704,7 @@ void FPVEditor::OnExport()
 	ProceduralVegetationGraph->ForEachNode(
 		[&](UPCGNode* Node)-> bool
 			{
-				if (!Node->GetSettings()->IsA<UPVOutputSettings>())
+				if (!Node->GetSettings()->IsA<UPVExportSettings>())
 				{
 					return true;
 				}
diff --git a/Engine/Plugins/Experimental/ProceduralVegetationEditor/Source/ProceduralVegetationEditor/Private/Widgets/SPVExportSelectionDialog.cpp b/Engine/Plugins/Experimental/ProceduralVegetationEditor/Source/ProceduralVegetationEditor/Private/Widgets/SPVExportSelectionDialog.cpp
--- a/Engine/Plugins/Experimental/ProceduralVegetationEditor/Source/ProceduralVegetationEditor/Private/Widgets/SPVExportSelectionDialog.cpp
+++ b/Engine/Plugins/Experimental/ProceduralVegetationEditor/Source/ProceduralVegetationEditor/Private/Widgets/SPVExportSelectionDialog.cpp
@@ -11,7 +11,7 @@
 #include "Algo/AnyOf.h"
 
 #include "Nodes/PCGEditorGraphNodeBase.h"
-#include "Nodes/PVOutputSettings.h"
+#include "Nodes/PVExportSettings.h"
 
 #include "Widgets/Input/SSegmentedControl.h"
 #include "Widgets/Input/STextComboBox.h"
diff --git a/Engine/Plugins/Experimental/ProceduralVegetationEditor/Source/ProceduralVegetationEditor/Private/Widgets/SPVExportSelectionDialog.h b/Engine/Plugins/Experimental/ProceduralVegetationEditor/Source/ProceduralVegetationEditor/Private/Widgets/SPVExportSelectionDialog.h
--- a/Engine/Plugins/Experimental/ProceduralVegetationEditor/Source/ProceduralVegetationEditor/Private/Widgets/SPVExportSelectionDialog.h
+++ b/Engine/Plugins/Experimental/ProceduralVegetationEditor/Source/ProceduralVegetationEditor/Private/Widgets/SPVExportSelectionDialog.h
@@ -17,7 +17,7 @@ class IStructureDetailsView;
 class ITableRow;
 class STableViewBase;
 class UPCGEditorGraphNodeBase;
-class UPVOutputSettings;
+class UPVExportSettings;
 class SScrollBox;
 class UProceduralVegetationGraph;
 class UPCGNode;
