Laziliy compute gradient attributes needed for the mesher

- Remove BranchGradient attribute (can trivially be computed on the fly where needed)
- Remove gradient attributes from the facades IsValid functions as they are optional for most nodes
- Add gradient computation functions for HullGradient, and MainTrunkGradient
- Fix incorrect computation of GroundGradient
- Re-compute required gradients on the fly (PVMeshBuilder and PVParametricDistributionHelper)
- Fix ComputeBudDevelopment computing the wrong BudDevelopment Generation with fused points

#Synced-CL 52140050

#Synced-CL 52140050

diff --git a/Engine/Plugins/Experimental/ProceduralVegetationEditor/Source/ProceduralVegetation/Private/Facades/PVPointFacade.cpp b/Engine/Plugins/Experimental/ProceduralVegetationEditor/Source/ProceduralVegetation/Private/Facades/PVPointFacade.cpp
--- a/Engine/Plugins/Experimental/ProceduralVegetationEditor/Source/ProceduralVegetation/Private/Facades/PVPointFacade.cpp
+++ b/Engine/Plugins/Experimental/ProceduralVegetationEditor/Source/ProceduralVegetation/Private/Facades/PVPointFacade.cpp
@@ -19,7 +19,6 @@ namespace PV::Facades
 		, BudDirections(InCollection, AttributeNames::BudDirection, GroupNames::PointGroup)
 		, BudStatus(InCollection, AttributeNames::BudStatus, GroupNames::PointGroup)
 		, BudLateralMeristem(InCollection, AttributeNames::BudLateralMeristem, GroupNames::PointGroup)
-		, BranchGradients(InCollection, AttributeNames::BranchGradient, GroupNames::PointGroup)
 		, BudHormoneLevels(InCollection, AttributeNames::BudHormoneLevels, GroupNames::PointGroup)
 		, PlantGradients(InCollection, AttributeNames::PlantGradient, GroupNames::PointGroup)
 		, LengthFromSeed(InCollection, AttributeNames::LengthFromSeed, GroupNames::PointGroup)
@@ -48,7 +47,6 @@ namespace PV::Facades
 		, BudDirections(InCollection, AttributeNames::BudDirection, GroupNames::PointGroup)
 		, BudStatus(InCollection, AttributeNames::BudStatus, GroupNames::PointGroup)
 		, BudLateralMeristem(InCollection, AttributeNames::BudLateralMeristem, GroupNames::PointGroup)
-		, BranchGradients(InCollection, AttributeNames::BranchGradient, GroupNames::PointGroup)
 		, BudHormoneLevels(InCollection, AttributeNames::BudHormoneLevels, GroupNames::PointGroup)
 		, PlantGradients(InCollection, AttributeNames::PlantGradient, GroupNames::PointGroup)
 		, LengthFromSeed(InCollection, AttributeNames::LengthFromSeed, GroupNames::PointGroup)
@@ -69,15 +67,9 @@ namespace PV::Facades
 		return Positions.IsValid()
 			&& LengthFromRoot.IsValid()
 			&& LengthFromSeed.IsValid()
-			&& PointScaleGradient.IsValid()
-			&& HullGradient.IsValid()
-			&& MainTrunkGradient.IsValid()
-			&& GroundGradient.IsValid()
 			&& PointScale.IsValid()
 			&& BudDirections.IsValid()
-			&& BranchGradients.IsValid()
 			&& BudHormoneLevels.IsValid()
-			&& PlantGradients.IsValid()
 			&& BudNumber.IsValid()
 			&& NjordPixelIndex.IsValid()
 			&& BudLightDetected.IsValid()
@@ -144,46 +136,6 @@ namespace PV::Facades
 		}
 	}
 
-	float FPointFacade::GetPointScaleGradient(int32 Index) const
-	{
-		if (PointScaleGradient.IsValid() && PointScaleGradient.IsValidIndex(Index))
-		{
-			return PointScaleGradient[Index];
-		}
-
-		return 0;
-	}
-
-	float FPointFacade::GetHullGradient(int32 Index) const
-	{
-		if (HullGradient.IsValid() && HullGradient.IsValidIndex(Index))
-		{
-			return HullGradient[Index];
-		}
-
-		return 0;
-	}
-
-	float FPointFacade::GetMainTrunkGradient(int32 Index) const
-	{
-		if (MainTrunkGradient.IsValid() && MainTrunkGradient.IsValidIndex(Index))
-		{
-			return MainTrunkGradient[Index];
-		}
-
-		return 0;
-	}
-
-	float FPointFacade::GetGroundGradient(int32 Index) const
-	{
-		if (GroundGradient.IsValid() && GroundGradient.IsValidIndex(Index))
-		{
-			return GroundGradient[Index];
-		}
-
-		return 0;
-	}
-
 	float FPointFacade::GetPointScale(int32 Index) const
 	{
 		if (PointScale.IsValid() && PointScale.IsValidIndex(Index))
@@ -389,42 +341,6 @@ namespace PV::Facades
 		}
 	}
 
-	float FPointFacade::GetBranchGradient(int32 Index) const
-	{
-		if (BranchGradients.IsValid() && BranchGradients.IsValidIndex(Index))
-		{
-			return BranchGradients[Index];
-		}
-
-		return 0;
-	}
-
-	void FPointFacade::SetBranchGradient(int32 Index, float InBranchGradient)
-	{
-		if (BranchGradients.IsValid() && BranchGradients.IsValidIndex(Index))
-		{
-			BranchGradients.ModifyAt(Index, InBranchGradient);
-		}
-	}
-
-	float FPointFacade::GetPlantGradient(int32 Index) const
-	{
-		if (PlantGradients.IsValid() && PlantGradients.IsValidIndex(Index))
-		{
-			return PlantGradients[Index];
-		}
-
-		return 0;
-	}
-
-	void FPointFacade::SetPlantGradient(int32 Index, float InPlantGradient)
-	{
-		if (PlantGradients.IsValid() && PlantGradients.IsValidIndex(Index))
-		{
-			PlantGradients.ModifyAt(Index, InPlantGradient);
-		}
-	}
-
 	float FPointFacade::GetNjordPixelIndex(int32 Index) const
 	{
 		if (NjordPixelIndex.IsValid() && NjordPixelIndex.IsValidIndex(Index))
@@ -566,15 +482,29 @@ namespace PV::Facades
 		{
 			Positions.ModifyAt(ToIndex, Positions[FromIndex]);
 			LengthFromRoot.ModifyAt(ToIndex, LengthFromRoot[FromIndex]);
-			PointScaleGradient.ModifyAt(ToIndex, PointScaleGradient[FromIndex]);
-			HullGradient.ModifyAt(ToIndex, HullGradient[FromIndex]);
-			MainTrunkGradient.ModifyAt(ToIndex, MainTrunkGradient[FromIndex]);
-			GroundGradient.ModifyAt(ToIndex, GroundGradient[FromIndex]);
+			if (PointScaleGradient.IsValid())
+			{
+				PointScaleGradient.ModifyAt(ToIndex, PointScaleGradient[FromIndex]);
+			}
+			if (HullGradient.IsValid())
+			{
+				HullGradient.ModifyAt(ToIndex, HullGradient[FromIndex]);
+			}
+			if (MainTrunkGradient.IsValid())
+			{
+				MainTrunkGradient.ModifyAt(ToIndex, MainTrunkGradient[FromIndex]);
+			}
+			if (GroundGradient.IsValid())
+			{
+				GroundGradient.ModifyAt(ToIndex, GroundGradient[FromIndex]);
+			}
 			PointScale.ModifyAt(ToIndex, PointScale[FromIndex]);
 			BudDirections.ModifyAt(ToIndex, BudDirections[FromIndex]);
-			BranchGradients.ModifyAt(ToIndex, BranchGradients[FromIndex]);
 			BudHormoneLevels.ModifyAt(ToIndex, BudHormoneLevels[FromIndex]);
-			PlantGradients.ModifyAt(ToIndex, PlantGradients[FromIndex]);
+			if (PlantGradients.IsValid())
+			{
+				PlantGradients.ModifyAt(ToIndex, PlantGradients[FromIndex]);
+			}
 			LengthFromSeed.ModifyAt(ToIndex, LengthFromSeed[FromIndex]);
 			BudNumber.ModifyAt(ToIndex, BudNumber[FromIndex]);
 			NjordPixelIndex.ModifyAt(ToIndex, NjordPixelIndex[FromIndex]);
@@ -643,26 +573,6 @@ namespace PV::Facades
 		FILL_COLLECTION_ATTRIBUTE(LengthFromRoot, InputArray);
 	}
 
-	void FPointFacade::FillPointScaleGradients(const TArray<float>& InputArray)
-	{
-		FILL_COLLECTION_ATTRIBUTE(PointScaleGradient, InputArray);
-	}
-
-	void FPointFacade::FillHullGradients(const TArray<float>& InputArray)
-	{
-		FILL_COLLECTION_ATTRIBUTE(HullGradient, InputArray);
-	}
-
-	void FPointFacade::FillMainTrunkGradients(const TArray<float>& InputArray)
-	{
-		FILL_COLLECTION_ATTRIBUTE(MainTrunkGradient, InputArray);
-	}
-
-	void FPointFacade::FillGroundGradients(const TArray<float>& InputArray)
-	{
-		FILL_COLLECTION_ATTRIBUTE(GroundGradient, InputArray);
-	}
-
 	void FPointFacade::FillPointScales(const TArray<float>& InputArray)
 	{
 		FILL_COLLECTION_ATTRIBUTE(PointScale, InputArray);
@@ -673,21 +583,11 @@ namespace PV::Facades
 		FILL_COLLECTION_ATTRIBUTE(BudDirections, InputArray);
 	}
 
-	void FPointFacade::FillBranchGradients(const TArray<float>& InputArray)
-	{
-		FILL_COLLECTION_ATTRIBUTE(BranchGradients, InputArray);
-	}
-
 	void FPointFacade::FillBudHormoneLevels(const  TArray<TArray<float>>& InputArray)
 	{
 		FILL_COLLECTION_ATTRIBUTE(BudHormoneLevels, InputArray);
 	}
 
-	void FPointFacade::FillPlantGradients(const TArray<float>& InputArray)
-	{
-		FILL_COLLECTION_ATTRIBUTE(PlantGradients, InputArray);
-	}
-
 	void FPointFacade::FillLengthFromSeeds(const TArray<float>& InputArray)
 	{
 		FILL_COLLECTION_ATTRIBUTE(LengthFromSeed, InputArray);
@@ -738,26 +638,6 @@ namespace PV::Facades
 		GetValueFromAccessor<float>(LengthFromRoot, OutValue, Index, -1);
 	}
 
-	void FPointFacade::GetPointScaleGradient(int32 Index, float& OutValue) const
-	{
-		GetValueFromAccessor<float>(PointScaleGradient, OutValue, Index, -1);
-	}
-
-	void FPointFacade::GetHullGradient(int32 Index, float& OutValue) const
-	{
-		GetValueFromAccessor<float>(HullGradient, OutValue, Index, -1);
-	}
-
-	void FPointFacade::GetMainTrunkGradients(int32 Index, float& OutValue) const
-	{
-		GetValueFromAccessor<float>(MainTrunkGradient, OutValue, Index, -1);
-	}
-
-	void FPointFacade::GetGroundGradients(int32 Index, float& OutValue) const
-	{
-		GetValueFromAccessor<float>(GroundGradient, OutValue, Index, -1);
-	}
-
 	void FPointFacade::GetPointScales(int32 Index, float& OutValue) const
 	{
 		GetValueFromAccessor<float>(PointScale, OutValue, Index , -1);
@@ -768,21 +648,11 @@ namespace PV::Facades
 		GetValueFromAccessor<TArray<FVector3f>>(BudDirections, OutValue, Index, TArray<FVector3f>());
 	}
 
-	void FPointFacade::GetBranchGradients(int32 Index, float& OutValue) const
-	{
-		GetValueFromAccessor<float>(BranchGradients, OutValue, Index, -1);
-	}
-
 	void FPointFacade::GetBudHormoneLevels(int32 Index, TArray<float>& OutValue) const
 	{
 		GetValueFromAccessor<TArray<float>>(BudHormoneLevels, OutValue, Index, TArray<float>());
 	}
 
-	void FPointFacade::GetPlantGradients(int32 Index, float& OutValue) const
-	{
-		GetValueFromAccessor<float>(PlantGradients, OutValue, Index, -1);
-	}
-
 	void FPointFacade::GetBudNumbers(int32 Index, int32& OutValue) const
 	{
 		GetValueFromAccessor<int32>(BudNumber, OutValue, Index, -1);
diff --git a/Engine/Plugins/Experimental/ProceduralVegetationEditor/Source/ProceduralVegetation/Private/Helpers/PVAttributesHelper.cpp b/Engine/Plugins/Experimental/ProceduralVegetationEditor/Source/ProceduralVegetation/Private/Helpers/PVAttributesHelper.cpp
--- a/Engine/Plugins/Experimental/ProceduralVegetationEditor/Source/ProceduralVegetation/Private/Helpers/PVAttributesHelper.cpp
+++ b/Engine/Plugins/Experimental/ProceduralVegetationEditor/Source/ProceduralVegetation/Private/Helpers/PVAttributesHelper.cpp
@@ -2,6 +2,45 @@
 
 #include "Helpers/PVAttributesHelper.h"
 #include "Helpers/PVPlantTraversalHelper.h"
+#include "Chaos/Convex.h"
+
+float PV::AttributesHelper::GetBranchPointBranchGradient(
+	PV::FPointLengthFromRootAttributeConstView PointLengthFromRootAttribute,
+	PV::FBranchPointsAttributeConstView BranchPointsAttribute,
+	int32 BranchIndex, 
+	int32 BranchPointIndex
+)
+{
+	if (!BranchPointsAttribute.IsValidIndex(BranchIndex) || !BranchPointsAttribute[BranchIndex].IsValidIndex(BranchPointIndex))
+	{
+		return 0;
+	}
+
+	const TArray<int32>& BranchPoints = BranchPointsAttribute[BranchIndex];
+	if (BranchPoints.Num() < 2)
+	{
+		return 0;
+	}
+
+	const int32 PointIndex = BranchPoints[BranchPointIndex];
+
+	if (!PointLengthFromRootAttribute.IsValidIndex(BranchPoints[0])
+		|| !PointLengthFromRootAttribute.IsValidIndex(BranchPoints.Last())
+		|| !PointLengthFromRootAttribute.IsValidIndex(PointIndex))
+	{
+		return 0;
+	}
+
+	const float BaseLengthFromRoot = PointLengthFromRootAttribute[BranchPoints[0]];
+	const float TipLengthFromRoot = PointLengthFromRootAttribute[BranchPoints.Last()];
+	const float LengthFromRoot = PointLengthFromRootAttribute[PointIndex];
+
+	return FMath::GetMappedRangeValueClamped(
+		FVector2f(BaseLengthFromRoot, TipLengthFromRoot),
+		FVector2f(0.0f, 1.0f),
+		LengthFromRoot
+	);
+}
 
 void PV::AttributesHelper::ComputeBudNumbers(const FComputeBudNumbersAttributes& InAttributes)
 {
@@ -169,7 +208,7 @@ void PV::AttributesHelper::ComputeBudDevelopment(const FComputeBudDevelopmentAtt
 
 			InAttributes.BudDevelopment[PointIndex].BudAge = BudAge;
 			InAttributes.BudDevelopment[PointIndex].RelativeBudAge = BudAge;
-			InAttributes.BudDevelopment[PointIndex].Generation = InAttributes.BranchParents[BranchIndex].Num() + 1;
+			InAttributes.BudDevelopment[PointIndex].Generation = InAttributes.BranchParents[BranchIndex].Num();
 
 			return EForEachResult::Continue;
 		}
@@ -296,20 +335,6 @@ void PV::AttributesHelper::ComputeBudHormoneLevels(const FComputeBudHormoneLevel
 	}
 }
 
-void PV::AttributesHelper::ComputeGradients(const FComputeGradientsAttributes& InAttributes)
-{
-	using namespace PV::PlantTraversalHelper;
-
-	if (!InAttributes.IsValid())
-	{
-		return;
-	}
-
-	ComputePointGroundGradient(InAttributes);
-	ComputePointScaleGradient(InAttributes);
-	ComputePointPlantAndBranchGradient(InAttributes);
-}
-
 void PV::AttributesHelper::ComputePointGroundGradient(const FComputePointGroundGradientAttributes& InAttributes)
 {
 	if (!InAttributes.IsValid())
@@ -317,17 +342,16 @@ void PV::AttributesHelper::ComputePointGroundGradient(const FComputePointGroundG
 		return;
 	}
 
-	const float* MaxLengthFromRoot = Algo::MaxElement(InAttributes.PointLengthFromRoot);
-	if (MaxLengthFromRoot == nullptr || *MaxLengthFromRoot == 0)
+	const FVector3f* MaxZVertex = Algo::MaxElementBy(InAttributes.PointPosition, [&](const auto& A) { return A.Z; });
+	if (MaxZVertex == nullptr || MaxZVertex->Z == 0)
 	{
 		return;
 	}
 
+	const float MaxZ = MaxZVertex->Z;
 	for (int32 PointIndex = 0; PointIndex < InAttributes.PointGroundGradient.Num(); PointIndex++)
 	{
-		const float LengthFromRoot = InAttributes.PointLengthFromRoot[PointIndex];
-		const float LengthFromRootAlpha = LengthFromRoot / *MaxLengthFromRoot;
-		InAttributes.PointGroundGradient[PointIndex] = 1.0 - LengthFromRootAlpha;
+		InAttributes.PointGroundGradient[PointIndex] = 1.f - (InAttributes.PointPosition[PointIndex].Z / MaxZ);
 	}
 }
 
@@ -352,7 +376,7 @@ void PV::AttributesHelper::ComputePointScaleGradient(const FComputePointScaleGra
 	}
 }
 
-void PV::AttributesHelper::ComputePointPlantAndBranchGradient(const FComputePointPlantAndBranchGradientAttributes& InAttributes)
+void PV::AttributesHelper::ComputePointPlantGradient(const FComputePointPlantGradientAttributes& InAttributes)
 {
 	using namespace PV::PlantTraversalHelper;
 
@@ -362,7 +386,6 @@ void PV::AttributesHelper::ComputePointPlantAndBranchGradient(const FComputePoin
 	}
 
 	// PlantGradient - A gradient from 1-0 where the tip of each branch is 0, and the base of each branch is whatever the gradient of the parent branch is at that point.
-	// BranchGradient - A gradient from 1-0 where 1 = base of the branch and 0 = tip of the branch
 	RecursiveWalkBranches(
 		InAttributes.BranchParentNumber,
 		InAttributes.BranchChildren,
@@ -375,16 +398,14 @@ void PV::AttributesHelper::ComputePointPlantAndBranchGradient(const FComputePoin
 				return EForEachResult::Continue;
 			}
 
-			const float BasePointPlantGradient = GetBranchParentIndex(InAttributes.BranchParentNumber, InAttributes.BranchNumber, BranchIndex) == INDEX_NONE
-				? 1.0
-				: InAttributes.PointPlantGradient[BranchPoints[0]];
+			const bool bIsTrunk = IsTrunk(InAttributes.BranchParentNumber, BranchIndex);
+			const float BasePointPlantGradient = bIsTrunk ? 1.0 : InAttributes.PointPlantGradient[BranchPoints[0]];
 
-			for (int32 i = 0; i < BranchPoints.Num(); ++i)
+			for (int32 i = bIsTrunk ? 0 : 1; i < BranchPoints.Num(); ++i)
 			{
 				const float BranchAlpha = i / float(BranchPoints.Num() - 1);
 				const float PointPlantGradient = FMath::Lerp(BasePointPlantGradient, 0, BranchAlpha);
 				InAttributes.PointPlantGradient[BranchPoints[i]] = PointPlantGradient;
-				InAttributes.PointBranchGradient[BranchPoints[i]] = 1.0 - BranchAlpha;
 			}
 
 			return EForEachResult::Continue;
@@ -392,6 +413,180 @@ void PV::AttributesHelper::ComputePointPlantAndBranchGradient(const FComputePoin
 	);
 }
 
+void PV::AttributesHelper::ComputePointHullGradient(const FComputePointHullGradientAttributes& InAttributes)
+{
+	using namespace PV::PlantTraversalHelper;
+
+	if (!InAttributes.IsValid())
+	{
+		return;
+	}
+
+	for (int32 i = 0; i < InAttributes.PointHullGradient.Num(); ++i)
+	{
+		InAttributes.PointHullGradient[i] = 1.0f;
+	}
+
+	TArray<float> HullDistanceArray;
+	HullDistanceArray.SetNumZeroed(InAttributes.PointHullGradient.Num());
+
+	TArray<Chaos::FConvex::FVec3Type> HullPoints;
+	for (int32 BranchIndex = 0; BranchIndex < InAttributes.BranchPoints.Num(); ++BranchIndex)
+	{
+		const TArray<int32>& BranchPoints = InAttributes.BranchPoints[BranchIndex];
+		if (BranchPoints.IsEmpty())
+		{
+			continue;
+		}
+
+		// Tip of each branch is a hull candidate
+		const FVector3f& TipPosition = InAttributes.PointPosition[BranchPoints.Last()];
+		HullPoints.Add(Chaos::FConvex::FVec3Type(TipPosition.X, TipPosition.Y, TipPosition.Z));
+
+		// Root of the trunk is also a hull candidate
+		const bool bIsTrunk = IsTrunk(InAttributes.BranchParentNumber, BranchIndex);
+		if (bIsTrunk)
+		{
+			const FVector3f& BasePosition = InAttributes.PointPosition[BranchPoints[0]];
+			HullPoints.Add(Chaos::FConvex::FVec3Type(BasePosition.X, BasePosition.Y, BasePosition.Z));
+		}
+	}
+
+	const Chaos::FConvex ConvexShape(HullPoints, 0.f);
+	if (ConvexShape.NumPlanes() == 0)
+	{
+		return;
+	}
+
+	// For each plant point compute the distance to the hull surface
+	for (int32 PointIndex = 0; PointIndex < InAttributes.PointHullGradient.Num(); ++PointIndex)
+	{
+		const FVector3f& PointPos = InAttributes.PointPosition[PointIndex];
+		Chaos::FVec3 Normal;
+		const double Dist = FMath::Abs(ConvexShape.PhiWithNormal(Chaos::FVec3(PointPos.X, PointPos.Y, PointPos.Z), Normal));
+		HullDistanceArray[PointIndex] = static_cast<float>(Dist);
+	}
+
+	// Normalize globally: map [MinDist, MaxDist] -> [1, 0]
+	// Points near the hull surface (small distance) get gradient 1; deep interior points get 0.
+	const float MinDist = *Algo::MinElement(HullDistanceArray);
+	const float MaxDist = *Algo::MaxElement(HullDistanceArray);
+
+	if (FMath::IsNearlyEqual(MinDist, MaxDist))
+	{
+		return;
+	}
+
+	for (int32 i = 0; i < InAttributes.PointHullGradient.Num(); ++i)
+	{
+		InAttributes.PointHullGradient[i] = FMath::GetMappedRangeValueClamped(
+			FVector2f(MinDist, MaxDist),
+			FVector2f(1.0f, 0.0f),
+			HullDistanceArray[i]
+		);
+	}
+}
+
+void PV::AttributesHelper::ComputePointMainTrunkGradient(const FComputePointMainTrunkGradientAttributes& InAttributes)
+{
+	using namespace PV::PlantTraversalHelper;
+
+	if (!InAttributes.IsValid())
+	{
+		return;
+	}
+
+	TArray<float> DistanceArray;
+	DistanceArray.SetNumZeroed(InAttributes.PointMainTrunkGradient.Num());
+
+	ForEachPlant(InAttributes.BranchPlantNumber, [&](int32 PlantNumber)
+	{
+		TArray<int32> TrunkPoints;
+		TArray<int32> NonTrunkPoints;
+
+		ForEachPlantPoint(
+			InAttributes.BranchPlantNumber, 
+			InAttributes.BranchPoints, 
+			InAttributes.BranchParentNumber, 
+			InAttributes.BranchNumber,
+			PlantNumber,
+			[&](int32 BranchIndex, int32 PointIndex)
+			{
+				// Note that we check the Generation instead of using IsTrunk to also include bifurcated branches in the TrunkPoints
+				const bool bIsTrunk = InAttributes.BudDevelopment[PointIndex].Generation == 1;
+				if (bIsTrunk)
+				{
+					TrunkPoints.Add(PointIndex);
+				}
+				else
+				{
+					NonTrunkPoints.Add(PointIndex);
+				}
+
+				return EForEachResult::Continue;
+			}
+		);
+
+		if (TrunkPoints.IsEmpty() || NonTrunkPoints.IsEmpty())
+		{
+			return EForEachResult::Continue;
+		}
+
+		for (int32 i = 0; i < NonTrunkPoints.Num(); ++i)
+		{
+			const int32 PointIndex = NonTrunkPoints[i];
+
+			const FVector3f& PointPos = InAttributes.PointPosition[PointIndex];
+			float MinDist = TNumericLimits<float>::Max();
+
+			if (TrunkPoints.Num() == 1)
+			{
+				MinDist = FMath::Min(MinDist, (PointPos - InAttributes.PointPosition[TrunkPoints[0]]).Size());
+			}
+			else
+			{
+				for (int32 j = 0; j < TrunkPoints.Num() - 1; ++j)
+				{
+					const FVector3f& P0 = InAttributes.PointPosition[TrunkPoints[j]];
+					const FVector3f& P1 = InAttributes.PointPosition[TrunkPoints[j + 1]];
+
+					// We're currently using the physical distance to the trunk branch. It may suffice to use
+					// the difference in LengthFromRoot from the connection point to the trunk to the current
+					// point.
+					const FVector3f ClosestPoint = FMath::ClosestPointOnSegment(PointPos, P0, P1);
+					MinDist = FMath::Min(MinDist, (PointPos - ClosestPoint).Size());
+				}
+			}
+
+			DistanceArray[PointIndex] = MinDist;
+		}
+
+		return EForEachResult::Continue;
+	});
+
+	const float MinDist = *Algo::MinElement(DistanceArray);
+	const float MaxDist = *Algo::MaxElement(DistanceArray);
+
+	if (FMath::IsNearlyEqual(MinDist, MaxDist))
+	{
+		for (int32 i = 0; i < InAttributes.PointMainTrunkGradient.Num(); ++i)
+		{
+			InAttributes.PointMainTrunkGradient[i] = 1.0f;
+		}
+	}
+	else
+	{
+		for (int32 PointIndex = 0; PointIndex < InAttributes.PointMainTrunkGradient.Num(); ++PointIndex)
+		{
+			InAttributes.PointMainTrunkGradient[PointIndex] = FMath::GetMappedRangeValueClamped(
+				FVector2f(MinDist, MaxDist),
+				FVector2f(1.0f, 0.0f),
+				DistanceArray[PointIndex]
+			);
+		}
+	}
+}
+
 void PV::AttributesHelper::ComputeBudStatus(const FComputeBudStatusAttributes& InAttributes)
 {
 	using namespace PV::PlantTraversalHelper;
@@ -702,17 +897,31 @@ bool PV::AttributesHelper::SplitFusedPoints(FSplitFusedPointsAttributes InAttrib
 		InAttributes.PointPosition.ModifyAt(NewPointIndex) = InAttributes.PointPosition[PrevPointIndex];
 		InAttributes.PointLengthFromRoot.ModifyAt(NewPointIndex) = InAttributes.PointLengthFromRoot[PrevPointIndex];
 		InAttributes.PointLengthFromSeed.ModifyAt(NewPointIndex) = InAttributes.PointLengthFromSeed[PrevPointIndex];
-		InAttributes.PointScaleGradient.ModifyAt(NewPointIndex) = InAttributes.PointScaleGradient[PrevPointIndex];
-		InAttributes.PointHullGradient.ModifyAt(NewPointIndex) = InAttributes.PointHullGradient[PrevPointIndex];
-		InAttributes.PointMainTrunkGradient.ModifyAt(NewPointIndex) = InAttributes.PointMainTrunkGradient[PrevPointIndex];
-		InAttributes.PointGroundGradient.ModifyAt(NewPointIndex) = InAttributes.PointGroundGradient[PrevPointIndex];
+		if (InAttributes.PointScaleGradient.IsValid())
+		{
+			InAttributes.PointScaleGradient.ModifyAt(NewPointIndex) = InAttributes.PointScaleGradient[PrevPointIndex];
+		}
+		if (InAttributes.PointHullGradient.IsValid())
+		{
+			InAttributes.PointHullGradient.ModifyAt(NewPointIndex) = InAttributes.PointHullGradient[PrevPointIndex];
+		}
+		if (InAttributes.PointMainTrunkGradient.IsValid())
+		{
+			InAttributes.PointMainTrunkGradient.ModifyAt(NewPointIndex) = InAttributes.PointMainTrunkGradient[PrevPointIndex];
+		}
+		if (InAttributes.PointGroundGradient.IsValid())
+		{
+			InAttributes.PointGroundGradient.ModifyAt(NewPointIndex) = InAttributes.PointGroundGradient[PrevPointIndex];
+		}
 		InAttributes.PointScale.ModifyAt(NewPointIndex) = InAttributes.PointScale[PrevPointIndex];
 		InAttributes.BudDirection.ModifyAt(NewPointIndex).Array = InAttributes.BudDirection[PrevPointIndex].Array;
 		InAttributes.BudStatus.ModifyAt(NewPointIndex).Array = InAttributes.BudStatus[PrevPointIndex].Array;
 		InAttributes.BudLateralMeristem.ModifyAt(NewPointIndex).Array = InAttributes.BudLateralMeristem[PrevPointIndex].Array;
-		InAttributes.PointBranchGradient.ModifyAt(NewPointIndex) = InAttributes.PointBranchGradient[PrevPointIndex];
 		InAttributes.BudHormoneLevels.ModifyAt(NewPointIndex).Array = InAttributes.BudHormoneLevels[PrevPointIndex].Array;
-		InAttributes.PointPlantGradient.ModifyAt(NewPointIndex) = InAttributes.PointPlantGradient[PrevPointIndex];
+		if (InAttributes.PointPlantGradient.IsValid())
+		{
+			InAttributes.PointPlantGradient.ModifyAt(NewPointIndex) = InAttributes.PointPlantGradient[PrevPointIndex];
+		}
 		InAttributes.PointBudNumber.ModifyAt(NewPointIndex) = NewBudNumber;
 		InAttributes.PointNjordPixelIndex.ModifyAt(NewPointIndex) = InAttributes.PointNjordPixelIndex[PrevPointIndex];
 		InAttributes.BudLightDetected.ModifyAt(NewPointIndex).Array = InAttributes.BudLightDetected[PrevPointIndex].Array;
@@ -788,19 +997,57 @@ bool PV::AttributesHelper::RecomputeAllGrowthDataAttributes(const FRecomputeGrow
 		InAttributes.BranchNumber
 	});
 
-	ComputeGradients({
-		InAttributes.PointGroundGradient,
-		InAttributes.PointScaleGradient,
+	ComputePointPlantGradient({
 		InAttributes.PointPlantGradient,
-		InAttributes.PointBranchGradient,
-		InAttributes.PointLengthFromRoot,
-		InAttributes.PointScale,
 		InAttributes.BranchPoints,
 		InAttributes.BranchParentNumber,
 		InAttributes.BranchNumber,
 		InAttributes.BranchChildren
 	});
 
+
+	if (InAttributes.PointGroundGradient.IsValid() 
+		&& InAttributes.PointGroundGradient.Num() == InAttributes.PointPosition.Num())
+	{
+		ComputePointGroundGradient({ 
+			InAttributes.PointGroundGradient,
+			InAttributes.PointPosition
+		});
+	}
+
+	if (InAttributes.PointScaleGradient.IsValid()
+		&& InAttributes.PointScaleGradient.Num() == InAttributes.PointScale.Num())
+	{
+		ComputePointScaleGradient({ 
+			InAttributes.PointScaleGradient,
+			InAttributes.PointScale
+		});
+	}
+		
+	if (InAttributes.PointHullGradient.IsValid())
+	{
+		ComputePointHullGradient({
+			InAttributes.PointHullGradient,
+			InAttributes.PointPosition,
+			InAttributes.BranchPoints,
+			InAttributes.BranchParentNumber
+		});
+	}
+
+	if (InAttributes.PointMainTrunkGradient.IsValid()
+		&& InAttributes.PointMainTrunkGradient.Num() == InAttributes.PointPosition.Num())
+	{
+		ComputePointMainTrunkGradient({
+			InAttributes.PointMainTrunkGradient,
+			InAttributes.PointPosition,
+			InAttributes.BranchPoints,
+			InAttributes.BranchParentNumber,
+			InAttributes.BranchPlantNumber,
+			InAttributes.BranchNumber,
+			InAttributes.BudDevelopment
+		});
+	}
+
 	ComputeBudDevelopment({
 		InAttributes.BudDevelopment,
 		InAttributes.PointPlantGradient,
@@ -825,4 +1072,4 @@ bool PV::AttributesHelper::RecomputeAllGrowthDataAttributes(const FRecomputeGrow
 	});
 
 	return true;
-}
\ No newline at end of file
+}
diff --git a/Engine/Plugins/Experimental/ProceduralVegetationEditor/Source/ProceduralVegetation/Private/Helpers/PVAttributesHelper.h b/Engine/Plugins/Experimental/ProceduralVegetationEditor/Source/ProceduralVegetation/Private/Helpers/PVAttributesHelper.h
--- a/Engine/Plugins/Experimental/ProceduralVegetationEditor/Source/ProceduralVegetation/Private/Helpers/PVAttributesHelper.h
+++ b/Engine/Plugins/Experimental/ProceduralVegetationEditor/Source/ProceduralVegetation/Private/Helpers/PVAttributesHelper.h
@@ -5,6 +5,13 @@
 
 namespace PV::AttributesHelper
 {
+	float GetBranchPointBranchGradient(
+		PV::FPointLengthFromRootAttributeConstView PointLengthFromRootAttribute,
+		PV::FBranchPointsAttributeConstView BranchPointsAttribute,
+		int32 BranchIndex, 
+		int32 BranchPointIndex
+	);
+
 	struct FComputeBudNumbersAttributes
 	{
 		FPointBudNumberAttributeView PointBudNumber;
@@ -33,7 +40,7 @@ namespace PV::AttributesHelper
 	};
 
 	// Re-computes the point bud numbers and writes them to the BudNumber and BranchSourceBudNumber attributes.
-	PROCEDURALVEGETATION_API extern void ComputeBudNumbers(const FComputeBudNumbersAttributes& InAttributes);
+	void ComputeBudNumbers(const FComputeBudNumbersAttributes& InAttributes);
 
 	struct FComputeNjordPixelIndexAttributes
 	{
@@ -73,7 +80,7 @@ namespace PV::AttributesHelper
 	};
 
 	// Re-computes all njord pixel indices and writes them to the NjordPixelIndexAttribute.
-	PROCEDURALVEGETATION_API extern void ComputeNjordPixelIndex(const FComputeNjordPixelIndexAttributes& InAttributes);
+	void ComputeNjordPixelIndex(const FComputeNjordPixelIndexAttributes& InAttributes);
 
 	struct FComputeLengthFromRootAttributes
 	{
@@ -121,7 +128,7 @@ namespace PV::AttributesHelper
 	};
 
 	// Re-computes the length from root and length from seed for all points.
-	PROCEDURALVEGETATION_API extern void ComputeLengthFromRoot(const FComputeLengthFromRootAttributes& InAttributes);
+	void ComputeLengthFromRoot(const FComputeLengthFromRootAttributes& InAttributes);
 
 	struct FComputeBudDevelopmentAttributes
 	{
@@ -163,7 +170,7 @@ namespace PV::AttributesHelper
 	};
 
 	// Re-computes bud development for all points using the existing branch hierarchy and plant gradient. 
-	PROCEDURALVEGETATION_API extern void ComputeBudDevelopment(const FComputeBudDevelopmentAttributes& InAttributes);
+	void ComputeBudDevelopment(const FComputeBudDevelopmentAttributes& InAttributes);
 
 	struct FComputeBudDirectionsAttributes
 	{
@@ -209,7 +216,7 @@ namespace PV::AttributesHelper
 	};
 
 	// Re-computes the bud directions for all points.
-	PROCEDURALVEGETATION_API extern void ComputeBudDirections(const FComputeBudDirectionsAttributes& InAttributes);
+	void ComputeBudDirections(const FComputeBudDirectionsAttributes& InAttributes);
 
 	struct FComputeBudHormoneLevelsAttributes
 	{
@@ -223,108 +230,37 @@ namespace PV::AttributesHelper
 	};
 
 	// Re-computes the bud hormone levels using the plant gradient for all points.
-	PROCEDURALVEGETATION_API extern void ComputeBudHormoneLevels(const FComputeBudHormoneLevelsAttributes& InAttributes);
-
-	struct FComputeGradientsAttributes
-	{
-		PV::FPointGroundGradientAttributeView PointGroundGradient;
-		PV::FPointScaleGradientAttributeView PointScaleGradient;
-		PV::FPointPlantGradientAttributeView PointPlantGradient;
-		PV::FPointBranchGradientAttributeView PointBranchGradient;
-		PV::FPointLengthFromRootAttributeConstView PointLengthFromRoot;
-		PV::FPointScaleAttributeConstView PointScale;
-		PV::FBranchPointsAttributeConstView BranchPoints;
-		PV::FBranchParentNumberAttributeConstView BranchParentNumber;
-		PV::FBranchNumberAttributeConstView BranchNumber;
-		PV::FBranchChildrenAttributeConstView BranchChildren;
-
-		FComputeGradientsAttributes(
-			PV::FPointGroundGradientAttributeView InPointGroundGradient,
-			PV::FPointScaleGradientAttributeView InPointScaleGradient,
-			PV::FPointPlantGradientAttributeView InPointPlantGradient,
-			PV::FPointBranchGradientAttributeView InPointBranchGradient,
-			PV::FPointLengthFromRootAttributeConstView InPointLengthFromRoot,
-			PV::FPointScaleAttributeConstView InPointScale,
-			PV::FBranchPointsAttributeConstView InBranchPoints,
-			PV::FBranchParentNumberAttributeConstView InBranchParentNumber,
-			PV::FBranchNumberAttributeConstView InBranchNumber,
-			PV::FBranchChildrenAttributeConstView InBranchChildren)
-			: PointGroundGradient(InPointGroundGradient)
-			, PointScaleGradient(InPointScaleGradient)
-			, PointPlantGradient(InPointPlantGradient)
-			, PointBranchGradient(InPointBranchGradient)
-			, PointLengthFromRoot(InPointLengthFromRoot)
-			, PointScale(InPointScale)
-			, BranchPoints(InBranchPoints)
-			, BranchParentNumber(InBranchParentNumber)
-			, BranchNumber(InBranchNumber)
-			, BranchChildren(InBranchChildren)
-		{
-		}
-
-		FComputeGradientsAttributes(FManagedArrayCollection& InCollection)
-			: PointGroundGradient(PV::FPointGroundGradientAttribute::FindAttribute(InCollection))
-			, PointScaleGradient(PV::FPointScaleGradientAttribute::FindAttribute(InCollection))
-			, PointPlantGradient(PV::FPointPlantGradientAttribute::FindAttribute(InCollection))
-			, PointBranchGradient(PV::FPointBranchGradientAttribute::FindAttribute(InCollection))
-			, PointLengthFromRoot(PV::FPointLengthFromRootAttribute::FindAttribute(InCollection))
-			, PointScale(PV::FPointScaleAttribute::FindAttribute(InCollection))
-			, BranchPoints(PV::FBranchPointsAttribute::FindAttribute(InCollection))
-			, BranchParentNumber(PV::FBranchParentNumberAttribute::FindAttribute(InCollection))
-			, BranchNumber(PV::FBranchNumberAttribute::FindAttribute(InCollection))
-			, BranchChildren(PV::FBranchChildrenAttribute::FindAttribute(InCollection))
-		{}
-
-		bool IsValid() const
-		{
-			return ValidateAttributeCollection(PointGroundGradient, PointScaleGradient, PointPlantGradient, PointBranchGradient, PointLengthFromRoot, PointScale, BranchPoints, BranchParentNumber, BranchNumber, BranchChildren);
-		}
-	};
-
-	// Re-computes the point- ground gradient, scale gradient, plant gradient, and branch gradient.
-	PROCEDURALVEGETATION_API extern void ComputeGradients(const FComputeGradientsAttributes& InAttributes);
+	void ComputeBudHormoneLevels(const FComputeBudHormoneLevelsAttributes& InAttributes);
 
 	struct FComputePointGroundGradientAttributes
 	{
 		PV::FPointGroundGradientAttributeView PointGroundGradient;
-		PV::FPointLengthFromRootAttributeConstView PointLengthFromRoot;
-
-		FComputePointGroundGradientAttributes(const FComputeGradientsAttributes& InComputeGradientsAttributes)
-			: PointGroundGradient(InComputeGradientsAttributes.PointGroundGradient)
-			, PointLengthFromRoot(InComputeGradientsAttributes.PointLengthFromRoot)
-		{
-		}
+		PV::FPointPositionAttributeConstView PointPosition;
 
 		FComputePointGroundGradientAttributes(
 			PV::FPointGroundGradientAttributeView InPointGroundGradient,
-			PV::FPointLengthFromRootAttributeConstView InPointLengthFromRoot)
+			PV::FPointPositionAttributeConstView InPointLengthFromRoot)
 			: PointGroundGradient(InPointGroundGradient)
-			, PointLengthFromRoot(InPointLengthFromRoot)
+			, PointPosition(InPointLengthFromRoot)
 		{
 		}
 
 		FComputePointGroundGradientAttributes(FManagedArrayCollection& InCollection)
 			: PointGroundGradient(PV::FPointGroundGradientAttribute::FindAttribute(InCollection))
-			, PointLengthFromRoot(PV::FPointLengthFromRootAttribute::FindAttribute(InCollection))
+			, PointPosition(PV::FPointPositionAttribute::FindAttribute(InCollection))
 		{
 		}
 
-		bool IsValid() const { return ValidateAttributeCollection(PointGroundGradient, PointLengthFromRoot); }
+		bool IsValid() const { return ValidateAttributeCollection(PointGroundGradient, PointPosition); }
 	};
 
-	PROCEDURALVEGETATION_API extern void ComputePointGroundGradient(const FComputePointGroundGradientAttributes& InAttributes);
+	void ComputePointGroundGradient(const FComputePointGroundGradientAttributes& InAttributes);
 
 	struct FComputePointScaleGradientAttributes
 	{
 		PV::FPointScaleGradientAttributeView PointScaleGradient;
 		PV::FPointScaleAttributeConstView PointScale;
 
-		FComputePointScaleGradientAttributes(const FComputeGradientsAttributes& InComputeGradientsAttributes)
-			: PointScaleGradient(InComputeGradientsAttributes.PointScaleGradient)
-			, PointScale(InComputeGradientsAttributes.PointScale)
-		{
-		}
-
 		FComputePointScaleGradientAttributes(
 			PV::FPointScaleGradientAttributeView InPointScaleGradient,
 			PV::FPointScaleAttributeConstView InPointScale)
@@ -342,36 +278,23 @@ namespace PV::AttributesHelper
 		bool IsValid() const { return ValidateAttributeCollection(PointScaleGradient, PointScale); }
 	};
 
-	PROCEDURALVEGETATION_API extern void ComputePointScaleGradient(const FComputePointScaleGradientAttributes& InAttributes);
+	void ComputePointScaleGradient(const FComputePointScaleGradientAttributes& InAttributes);
 
-	struct FComputePointPlantAndBranchGradientAttributes
+	struct FComputePointPlantGradientAttributes
 	{
 		PV::FPointPlantGradientAttributeView PointPlantGradient;
-		PV::FPointBranchGradientAttributeView PointBranchGradient;
 		PV::FBranchPointsAttributeConstView BranchPoints;
 		PV::FBranchParentNumberAttributeConstView BranchParentNumber;
 		PV::FBranchNumberAttributeConstView BranchNumber;
 		PV::FBranchChildrenAttributeConstView BranchChildren;
 
-		FComputePointPlantAndBranchGradientAttributes(const FComputeGradientsAttributes& InComputeGradientsAttributes)
-			: PointPlantGradient(InComputeGradientsAttributes.PointPlantGradient)
-			, PointBranchGradient(InComputeGradientsAttributes.PointBranchGradient)
-			, BranchPoints(InComputeGradientsAttributes.BranchPoints)
-			, BranchParentNumber(InComputeGradientsAttributes.BranchParentNumber)
-			, BranchNumber(InComputeGradientsAttributes.BranchNumber)
-			, BranchChildren(InComputeGradientsAttributes.BranchChildren)
-		{
-		}
-
-		FComputePointPlantAndBranchGradientAttributes(
+		FComputePointPlantGradientAttributes(
 			PV::FPointPlantGradientAttributeView InPointPlantGradient,
-			PV::FPointBranchGradientAttributeView InPointBranchGradient,
 			PV::FBranchPointsAttributeConstView InBranchPoints,
 			PV::FBranchParentNumberAttributeConstView InBranchParentNumber,
 			PV::FBranchNumberAttributeConstView InBranchNumber,
 			PV::FBranchChildrenAttributeConstView InBranchChildren)
 			: PointPlantGradient(InPointPlantGradient)
-			, PointBranchGradient(InPointBranchGradient)
 			, BranchPoints(InBranchPoints)
 			, BranchParentNumber(InBranchParentNumber)
 			, BranchNumber(InBranchNumber)
@@ -379,9 +302,8 @@ namespace PV::AttributesHelper
 		{
 		}
 
-		FComputePointPlantAndBranchGradientAttributes(FManagedArrayCollection& InCollection)
+		FComputePointPlantGradientAttributes(FManagedArrayCollection& InCollection)
 			: PointPlantGradient(PV::FPointPlantGradientAttribute::FindAttribute(InCollection))
-			, PointBranchGradient(PV::FPointBranchGradientAttribute::FindAttribute(InCollection))
 			, BranchPoints(PV::FBranchPointsAttribute::FindAttribute(InCollection))
 			, BranchParentNumber(PV::FBranchParentNumberAttribute::FindAttribute(InCollection))
 			, BranchNumber(PV::FBranchNumberAttribute::FindAttribute(InCollection))
@@ -391,11 +313,93 @@ namespace PV::AttributesHelper
 
 		bool IsValid() const
 		{
-			return ValidateAttributeCollection(PointPlantGradient, PointBranchGradient, BranchPoints, BranchParentNumber, BranchNumber, BranchChildren);
+			return ValidateAttributeCollection(PointPlantGradient, BranchPoints, BranchParentNumber, BranchNumber, BranchChildren);
+		}
+	};
+
+	void ComputePointPlantGradient(const FComputePointPlantGradientAttributes& InAttributes);
+
+	struct FComputePointHullGradientAttributes
+	{
+		PV::FPointHullGradientAttributeView PointHullGradient;
+		PV::FPointPositionAttributeConstView PointPosition;
+		PV::FBranchPointsAttributeConstView BranchPoints;
+		PV::FBranchParentNumberAttributeConstView BranchParentNumber;
+
+		FComputePointHullGradientAttributes(
+			PV::FPointHullGradientAttributeView InPointHullGradient,
+			PV::FPointPositionAttributeConstView InPointPosition,
+			PV::FBranchPointsAttributeConstView InBranchPoints,
+			PV::FBranchParentNumberAttributeConstView InBranchParentNumber)
+			: PointHullGradient(InPointHullGradient)
+			, PointPosition(InPointPosition)
+			, BranchPoints(InBranchPoints)
+			, BranchParentNumber(InBranchParentNumber)
+		{
+		}
+
+		FComputePointHullGradientAttributes(FManagedArrayCollection& InCollection)
+			: PointHullGradient(PV::FPointHullGradientAttribute::FindAttribute(InCollection))
+			, PointPosition(PV::FPointPositionAttribute::FindAttribute(InCollection))
+			, BranchPoints(PV::FBranchPointsAttribute::FindAttribute(InCollection))
+			, BranchParentNumber(PV::FBranchParentNumberAttribute::FindAttribute(InCollection))
+		{
+		}
+
+		bool IsValid() const
+		{
+			return ValidateAttributeCollection(PointHullGradient, PointPosition, BranchPoints, BranchParentNumber);
+		}
+	};
+
+	void ComputePointHullGradient(const FComputePointHullGradientAttributes& InAttributes);
+
+	struct FComputePointMainTrunkGradientAttributes
+	{
+		PV::FPointMainTrunkGradientAttributeView PointMainTrunkGradient;
+		PV::FPointPositionAttributeConstView PointPosition;
+		PV::FBranchPointsAttributeConstView BranchPoints;
+		PV::FBranchParentNumberAttributeConstView BranchParentNumber;
+		PV::FBranchPlantNumberAttributeConstView BranchPlantNumber;
+		PV::FBranchNumberAttributeConstView BranchNumber;
+		PV::FBudDevelopmentAttributeConstView BudDevelopment;
+
+		FComputePointMainTrunkGradientAttributes(
+			PV::FPointMainTrunkGradientAttributeView InPointMainTrunkGradient,
+			PV::FPointPositionAttributeConstView InPointPosition,
+			PV::FBranchPointsAttributeConstView InBranchPoints,
+			PV::FBranchParentNumberAttributeConstView InBranchParentNumber,
+			PV::FBranchPlantNumberAttributeConstView InBranchPlantNumber,
+			PV::FBranchNumberAttributeConstView InBranchNumber,
+			PV::FBudDevelopmentAttributeConstView InBudDevelopment)
+			: PointMainTrunkGradient(InPointMainTrunkGradient)
+			, PointPosition(InPointPosition)
+			, BranchPoints(InBranchPoints)
+			, BranchParentNumber(InBranchParentNumber)
+			, BranchPlantNumber(InBranchPlantNumber)
+			, BranchNumber(InBranchNumber)
+			, BudDevelopment(InBudDevelopment)
+		{
+		}
+
+		FComputePointMainTrunkGradientAttributes(FManagedArrayCollection& InCollection)
+			: PointMainTrunkGradient(PV::FPointMainTrunkGradientAttribute::FindAttribute(InCollection))
+			, PointPosition(PV::FPointPositionAttribute::FindAttribute(InCollection))
+			, BranchPoints(PV::FBranchPointsAttribute::FindAttribute(InCollection))
+			, BranchParentNumber(PV::FBranchParentNumberAttribute::FindAttribute(InCollection))
+			, BranchPlantNumber(PV::FBranchPlantNumberAttribute::FindAttribute(InCollection))
+			, BranchNumber(PV::FBranchNumberAttribute::FindAttribute(InCollection))
+			, BudDevelopment(PV::FBudDevelopmentAttribute::FindAttribute(InCollection))
+		{
+		}
+
+		bool IsValid() const
+		{
+			return ValidateAttributeCollection(PointMainTrunkGradient, PointPosition, BranchPoints, BranchParentNumber, BranchPlantNumber, BranchNumber, BudDevelopment);
 		}
 	};
 
-	PROCEDURALVEGETATION_API extern void ComputePointPlantAndBranchGradient(const FComputePointPlantAndBranchGradientAttributes& InAttributes);
+	void ComputePointMainTrunkGradient(const FComputePointMainTrunkGradientAttributes& InAttributes);
 
 	struct FComputeBudStatusAttributes
 	{ 
@@ -438,7 +442,7 @@ namespace PV::AttributesHelper
 	};
 
 	// Re-computes the bud status for all points.
-	PROCEDURALVEGETATION_API extern void ComputeBudStatus(const FComputeBudStatusAttributes& InAttributes);
+	void ComputeBudStatus(const FComputeBudStatusAttributes& InAttributes);
 
 	struct FComputeBudLateralMeristemAttributes
 	{
@@ -494,10 +498,10 @@ namespace PV::AttributesHelper
 	};
 
 	// Re-computes the bud lateral meristem for all points.
-	PROCEDURALVEGETATION_API extern void ComputeBudLateralMeristem(const FComputeBudLateralMeristemAttributes& InAttributes);
+	void ComputeBudLateralMeristem(const FComputeBudLateralMeristemAttributes& InAttributes);
 
 	// Re-computes the bud light detected for all points.
-	PROCEDURALVEGETATION_API extern void ComputeBudLightDetected(PV::FBudLightDetectedAttributeView BudLightDetectedAttribute);
+	void ComputeBudLightDetected(PV::FBudLightDetectedAttributeView BudLightDetectedAttribute);
 
 	struct FComputeBranchHierarchyNumberAttributes
 	{
@@ -533,7 +537,7 @@ namespace PV::AttributesHelper
 	};
 
 	// Re-computes the BranchHierarchyNumber attribute given an already existing branch hierarchy.
-	PROCEDURALVEGETATION_API extern void ComputeBranchHierarchyNumbers(const FComputeBranchHierarchyNumberAttributes& InAttributes);
+	void ComputeBranchHierarchyNumbers(const FComputeBranchHierarchyNumberAttributes& InAttributes);
 
 	struct FComputeSmoothedPointScalesAttributes
 	{
@@ -578,7 +582,7 @@ namespace PV::AttributesHelper
 
 	// Smooths the existing point scales across the branch hierarchy and writes the results back to the PointScale attribute.
 	// Expects the PointScale attribute to already contain valid scale values prior to calling this function.
-	PROCEDURALVEGETATION_API extern void ComputeSmoothedPointScales(const FComputeSmoothedPointScalesAttributes& InAttributes);
+	void ComputeSmoothedPointScales(const FComputeSmoothedPointScalesAttributes& InAttributes);
 
 	struct FSplitFusedPointsAttributes
 	{
@@ -597,7 +601,6 @@ namespace PV::AttributesHelper
 		FBudDirectionAttributeAccessor BudDirection;
 		FBudStatusAttributeAccessor BudStatus;
 		FBudLateralMeristemAttributeAccessor BudLateralMeristem;
-		FPointBranchGradientAttributeAccessor PointBranchGradient;
 		FBudHormoneLevelsAttributeAccessor BudHormoneLevels;
 		FPointPlantGradientAttributeAccessor PointPlantGradient;
 		FPointBudNumberAttributeAccessor PointBudNumber;
@@ -621,7 +624,6 @@ namespace PV::AttributesHelper
 			, BudDirection(InCollection)
 			, BudStatus(InCollection)
 			, BudLateralMeristem(InCollection)
-			, PointBranchGradient(InCollection)
 			, BudHormoneLevels(InCollection)
 			, PointPlantGradient(InCollection)
 			, PointBudNumber(InCollection)
@@ -639,17 +641,11 @@ namespace PV::AttributesHelper
 				PointPosition.GetConstView(),
 				PointLengthFromRoot.GetConstView(),
 				PointLengthFromSeed.GetConstView(),
-				PointScaleGradient.GetConstView(),
-				PointHullGradient.GetConstView(),
-				PointMainTrunkGradient.GetConstView(),
-				PointGroundGradient.GetConstView(),
 				PointScale.GetConstView(),
 				BudDirection.GetConstView(),
 				BudStatus.GetConstView(),
 				BudLateralMeristem.GetConstView(),
-				PointBranchGradient.GetConstView(),
 				BudHormoneLevels.GetConstView(),
-				PointPlantGradient.GetConstView(),
 				PointBudNumber.GetConstView(),
 				PointNjordPixelIndex.GetConstView(),
 				BudLightDetected.GetConstView(),
@@ -661,7 +657,7 @@ namespace PV::AttributesHelper
 	// Converts a growth data collection with fused points to a split points collection.
 	// Will not validate that the input collection has fused points.
 	// NOTE: Assumes input collection is valid growth data.
-	PROCEDURALVEGETATION_API extern bool SplitFusedPoints(FSplitFusedPointsAttributes InAttributes);
+	bool SplitFusedPoints(FSplitFusedPointsAttributes InAttributes);
 
 	struct FRecomputeGrowthDataAttributes
 	{
@@ -685,7 +681,8 @@ namespace PV::AttributesHelper
 		FPointGroundGradientAttributeView PointGroundGradient;
 		FPointScaleGradientAttributeView PointScaleGradient;
 		FPointPlantGradientAttributeView PointPlantGradient;
-		FPointBranchGradientAttributeView PointBranchGradient;
+		FPointHullGradientAttributeView PointHullGradient;
+		FPointMainTrunkGradientAttributeView PointMainTrunkGradient;
 		FBudDevelopmentAttributeView BudDevelopment;
 		FBudHormoneLevelsAttributeView BudHormoneLevels;
 		FBudStatusAttributeView BudStatus;
@@ -712,7 +709,8 @@ namespace PV::AttributesHelper
 			, PointGroundGradient(FPointGroundGradientAttribute::FindAttribute(InCollection))
 			, PointScaleGradient(FPointScaleGradientAttribute::FindAttribute(InCollection))
 			, PointPlantGradient(FPointPlantGradientAttribute::FindAttribute(InCollection))
-			, PointBranchGradient(FPointBranchGradientAttribute::FindAttribute(InCollection))
+			, PointHullGradient(FPointHullGradientAttribute::FindAttribute(InCollection))
+			, PointMainTrunkGradient(FPointMainTrunkGradientAttribute::FindAttribute(InCollection))
 			, BudDevelopment(FBudDevelopmentAttribute::FindAttribute(InCollection))
 			, BudHormoneLevels(FBudHormoneLevelsAttribute::FindAttribute(InCollection))
 			, BudStatus(FBudStatusAttribute::FindAttribute(InCollection))
@@ -738,10 +736,7 @@ namespace PV::AttributesHelper
 				BudDirection,
 				BudLightDetected,
 				BudLateralMeristem,
-				PointGroundGradient,
-				PointScaleGradient,
 				PointPlantGradient,
-				PointBranchGradient,
 				BudDevelopment,
 				BudHormoneLevels,
 				BudStatus
@@ -749,5 +744,5 @@ namespace PV::AttributesHelper
 		}
 	};
 
-	PROCEDURALVEGETATION_API extern bool RecomputeAllGrowthDataAttributes(const FRecomputeGrowthDataAttributes& InAttributes);
+	bool RecomputeAllGrowthDataAttributes(const FRecomputeGrowthDataAttributes& InAttributes);
 }
diff --git a/Engine/Plugins/Experimental/ProceduralVegetationEditor/Source/ProceduralVegetation/Private/Helpers/PVHormoneDistributionHelper.cpp b/Engine/Plugins/Experimental/ProceduralVegetationEditor/Source/ProceduralVegetation/Private/Helpers/PVHormoneDistributionHelper.cpp
--- a/Engine/Plugins/Experimental/ProceduralVegetationEditor/Source/ProceduralVegetation/Private/Helpers/PVHormoneDistributionHelper.cpp
+++ b/Engine/Plugins/Experimental/ProceduralVegetationEditor/Source/ProceduralVegetation/Private/Helpers/PVHormoneDistributionHelper.cpp
@@ -1,6 +1,7 @@
 // Copyright Epic Games, Inc. All Rights Reserved.
 
 #include "PVHormoneDistributionHelper.h"
+#include "Helpers/PVAttributesHelper.h"
 
 void PV::HormoneDistributionHelper::ComputeHormoneBasedAttachmentPoints(
 	const FManagedArrayCollection& OutCollection,
@@ -13,12 +14,23 @@ void PV::HormoneDistributionHelper::ComputeHormoneBasedAttachmentPoints(
 {
 	auto BranchPointsAttribute = PV::FBranchPointsAttribute::GetAttribute(OutCollection);
 	auto BudDirectionAttribute = PV::FBudDirectionAttribute::GetAttribute(OutCollection);
-	auto BranchGradientsAttribute = PV::FPointBranchGradientAttribute::GetAttribute(OutCollection);
 	auto BudHormoneLevels = PV::FBudHormoneLevelsAttribute::GetAttribute(OutCollection);
 	auto PointLengthFromRootAttribute = PV::FPointLengthFromRootAttribute::GetAttribute(OutCollection);
 	auto PointScaleAttribute = PV::FPointScaleAttribute::GetAttribute(OutCollection);
 	auto PointPositionAttribute = PV::FPointPositionAttribute::GetAttribute(OutCollection);
 
+	if (!ValidateAttributeCollection(
+		BranchPointsAttribute,
+		BudDirectionAttribute,
+		BudHormoneLevels,
+		PointLengthFromRootAttribute,
+		PointScaleAttribute,
+		PointPositionAttribute
+	))
+	{
+		return;
+	}
+
 	const int32 NumBranches = BranchPointsAttribute.Num();
 
 	const auto [
@@ -118,8 +130,8 @@ void PV::HormoneDistributionHelper::ComputeHormoneBasedAttachmentPoints(
 						FVector2f(0.0f, 1.0f), FVector2f(FirstPointLengthFromRoot, LastPointLengthFromRoot), NormalizedAttachmentPoints[Idx]);
 
 					// Scale Calculation
-					const float CurrentPointBranchGradient = BranchGradientsAttribute[CurrentPointIndex];
-					const float NextPointBranchGradient = BranchGradientsAttribute[NextPointIndex];
+					const float CurrentPointBranchGradient = PV::AttributesHelper::GetBranchPointBranchGradient(PointLengthFromRootAttribute, BranchPointsAttribute, BranchIndex, BranchPointIndex);
+					const float NextPointBranchGradient = PV::AttributesHelper::GetBranchPointBranchGradient(PointLengthFromRootAttribute, BranchPointsAttribute, BranchIndex, BranchPointIndex + 1);
 					const float BranchGradient = FMath::Clamp(FMath::Lerp(CurrentPointBranchGradient, NextPointBranchGradient, DistanceAlpha), 0.0f,
 						1.0f);
 					float BranchGradientRampValue = ScaleSettings.ScaleRamp
diff --git a/Engine/Plugins/Experimental/ProceduralVegetationEditor/Source/ProceduralVegetation/Private/Helpers/PVImportHelpers.cpp b/Engine/Plugins/Experimental/ProceduralVegetationEditor/Source/ProceduralVegetation/Private/Helpers/PVImportHelpers.cpp
--- a/Engine/Plugins/Experimental/ProceduralVegetationEditor/Source/ProceduralVegetation/Private/Helpers/PVImportHelpers.cpp
+++ b/Engine/Plugins/Experimental/ProceduralVegetationEditor/Source/ProceduralVegetation/Private/Helpers/PVImportHelpers.cpp
@@ -814,15 +814,12 @@ void PV::ImportHelper::CreateEmptyGrowthData(FManagedArrayCollection& OutCollect
 	FPointPositionAttribute::AddAttribute(OutCollection, FVector3f::ZeroVector);
 	FPointLengthFromRootAttribute::AddAttribute(OutCollection, 0);
 	FPointLengthFromSeedAttribute::AddAttribute(OutCollection, 0);
-	FPointScaleGradientAttribute::AddAttribute(OutCollection, 0);
 	FPointHullGradientAttribute::AddAttribute(OutCollection, 0);
 	FPointMainTrunkGradientAttribute::AddAttribute(OutCollection, 0);
-	FPointGroundGradientAttribute::AddAttribute(OutCollection, 0);
 	FPointScaleAttribute::AddAttribute(OutCollection, 0.f);
 	FBudDirectionAttribute::AddAttribute(OutCollection);
 	FBudStatusAttribute::AddAttribute(OutCollection);
 	FBudLateralMeristemAttribute::AddAttribute(OutCollection);
-	FPointBranchGradientAttribute::AddAttribute(OutCollection, 0);
 	FBudHormoneLevelsAttribute::AddAttribute(OutCollection);
 	FPointPlantGradientAttribute::AddAttribute(OutCollection, 0);
 	FPointBudNumberAttribute::AddAttribute(OutCollection, 0);
diff --git a/Engine/Plugins/Experimental/ProceduralVegetationEditor/Source/ProceduralVegetation/Private/Helpers/PVJSONHelper.cpp b/Engine/Plugins/Experimental/ProceduralVegetationEditor/Source/ProceduralVegetation/Private/Helpers/PVJSONHelper.cpp
--- a/Engine/Plugins/Experimental/ProceduralVegetationEditor/Source/ProceduralVegetation/Private/Helpers/PVJSONHelper.cpp
+++ b/Engine/Plugins/Experimental/ProceduralVegetationEditor/Source/ProceduralVegetation/Private/Helpers/PVJSONHelper.cpp
@@ -557,36 +557,6 @@ namespace PV::JSON
 			}
 		}
 
-		if (!PV::FPointGroundGradientAttribute::HasAttribute(Collection))
-		{
-			PV::FPointGroundGradientAttribute::AddAttribute(Collection);
-			PV::AttributesHelper::ComputePointGroundGradient(Collection);
-		}
-
-		if (!PV::FPointScaleGradientAttribute::HasAttribute(Collection))
-		{
-			PV::FPointScaleGradientAttribute::AddAttribute(Collection);
-			PV::AttributesHelper::ComputePointScaleGradient(Collection);
-		}
-
-		if (!PV::FPointPlantGradientAttribute::HasAttribute(Collection) 
-			|| !PV::FPointBranchGradientAttribute::HasAttribute(Collection))
-		{
-			PV::FPointPlantGradientAttribute::AddAttribute(Collection);
-			PV::FPointBranchGradientAttribute::AddAttribute(Collection);
-			PV::AttributesHelper::ComputePointPlantAndBranchGradient(Collection);
-		}
-
-		if (!PV::FPointHullGradientAttribute::HasAttribute(Collection))
-		{
-			PV::FPointHullGradientAttribute::AddAttribute(Collection, 0);
-		}
-
-		if (!PV::FPointMainTrunkGradientAttribute::HasAttribute(Collection))
-		{
-			PV::FPointMainTrunkGradientAttribute::AddAttribute(Collection, 0);
-		}
-
 		if (!PV::FPointNjordPixelIndexAttribute::HasAttribute(Collection))
 		{
 			PV::FPointNjordPixelIndexAttribute::AddAttribute(Collection);
diff --git a/Engine/Plugins/Experimental/ProceduralVegetationEditor/Source/ProceduralVegetation/Private/Helpers/PVParametricDistributionHelper.cpp b/Engine/Plugins/Experimental/ProceduralVegetationEditor/Source/ProceduralVegetation/Private/Helpers/PVParametricDistributionHelper.cpp
--- a/Engine/Plugins/Experimental/ProceduralVegetationEditor/Source/ProceduralVegetation/Private/Helpers/PVParametricDistributionHelper.cpp
+++ b/Engine/Plugins/Experimental/ProceduralVegetationEditor/Source/ProceduralVegetation/Private/Helpers/PVParametricDistributionHelper.cpp
@@ -75,8 +75,6 @@ void PV::ParametricDistributionHelper::ComputeAttachmentPoints(
 	const int32 RandomSeed,
 	TArray<DistributionHelper::FAttachmentPoint>& OutAttachmentPoints)
 {
-	auto PointPlantGradientAttribute = FPointPlantGradientAttribute::GetAttribute(OutCollection);
-	auto PointBranchGradientAttribute = FPointBranchGradientAttribute::GetAttribute(OutCollection);
 	auto BranchParentNumberAttribute = FBranchParentNumberAttribute::GetAttribute(OutCollection);
 	auto BranchChildrenAttribute = FBranchChildrenAttribute::GetAttribute(OutCollection);
 	auto BranchNumberAttribute = FBranchNumberAttribute::GetAttribute(OutCollection);
@@ -85,13 +83,21 @@ void PV::ParametricDistributionHelper::ComputeAttachmentPoints(
 	auto BudDirectionAttribute = FBudDirectionAttribute::GetAttribute(OutCollection);
 	auto PointLengthFromRootAttribute = FPointLengthFromRootAttribute::GetAttribute(OutCollection);
 
-	const int32 NumOfBranches = BranchPointsAttribute.Num();
-	check(NumOfBranches > 0);
-	if (NumOfBranches == 0)
+	if (!ValidateAttributeCollection(
+		BranchParentNumberAttribute,
+		BranchChildrenAttribute,
+		BranchNumberAttribute,
+		BranchPointsAttribute,
+		PointPositionAttribute,
+		BudDirectionAttribute,
+		PointLengthFromRootAttribute
+	))
 	{
 		return;
 	}
 
+	const int32 NumOfBranches = BranchPointsAttribute.Num();
+
 	const auto [
 			PhyllotaxyFormationLeaf,
 			MinBudsLeaf,
@@ -106,10 +112,10 @@ void PV::ParametricDistributionHelper::ComputeAttachmentPoints(
 	BranchFacadeOut.RecomputeBranchChildren();
 	BranchFacadeOut.RecomputeBranchParents();
 
-	AttributesHelper::ComputePointPlantAndBranchGradient(
+	auto PointPlantGradientAttribute = FPointPlantGradientAttribute::AddAttribute(OutCollection);
+	AttributesHelper::ComputePointPlantGradient(
 		{
 			PointPlantGradientAttribute,
-			PointBranchGradientAttribute,
 			BranchPointsAttribute,
 			BranchParentNumberAttribute,
 			BranchNumberAttribute,
diff --git a/Engine/Plugins/Experimental/ProceduralVegetationEditor/Source/ProceduralVegetation/Private/Implementations/PVCarve.cpp b/Engine/Plugins/Experimental/ProceduralVegetationEditor/Source/ProceduralVegetation/Private/Implementations/PVCarve.cpp
--- a/Engine/Plugins/Experimental/ProceduralVegetationEditor/Source/ProceduralVegetation/Private/Implementations/PVCarve.cpp
+++ b/Engine/Plugins/Experimental/ProceduralVegetationEditor/Source/ProceduralVegetation/Private/Implementations/PVCarve.cpp
@@ -6,6 +6,7 @@
 #include "Facades/PVPlantFacade.h"
 #include "Facades/PVPointFacade.h"
 #include "Facades/PVTreeFacade.h"
+#include "Helpers/PVAttributesHelper.h"
 
 namespace
 {
@@ -95,12 +96,7 @@ void FPVCarve::RecomputeAttributes(FManagedArrayCollection& OutCollection, const
 
 	for (const int32& PlantNumber : PlantFacadeOut.GetPlantNumbers())
 	{
-		// Remap Ethylene levels while gathering minimum and maximum PScales for each branch and the whole plant
-		float MinimumPScale = TNumericLimits<float>::Max();
-		float MaximumPScale = TNumericLimits<float>::Min();
-		TArray<TTuple<float, float>> BranchMinMaxPScales;
-		BranchMinMaxPScales.Init(TTuple<float, float>(TNumericLimits<float>::Max(), TNumericLimits<float>::Min()), BranchFacadeOut.GetElementCount());
-
+		// Remap Ethylene levels for each branch and the whole plant
 		for (const int32& BranchIndex : PlantFacadeOut.GetBranchIndices(PlantNumber))
 		{
 			const TArray<int32>& BranchPoints = BranchFacadeOut.GetPoints(BranchIndex);
@@ -134,34 +130,6 @@ void FPVCarve::RecomputeAttributes(FManagedArrayCollection& OutCollection, const
 				check(CurrentBudHormoneLevels.Num() >= 4);
 				CurrentBudHormoneLevels[4] = NewEthyleneLevel;
 				PointFacadeOut.SetBudHormoneLevels(BranchPoints[Idx], CurrentBudHormoneLevels);
-
-				const float PointScaleAtIndex = PointFacadeOut.GetPointScale(BranchPoints[Idx]);
-				MinimumPScale = FMath::Min(MinimumPScale, PointScaleAtIndex);
-				MaximumPScale = FMath::Max(MaximumPScale, PointScaleAtIndex);
-
-				BranchMinMaxPScales[BranchIndex].Get<0>() = FMath::Min(BranchMinMaxPScales[BranchIndex].Get<0>(), PointScaleAtIndex);
-				BranchMinMaxPScales[BranchIndex].Get<1>() = FMath::Max(BranchMinMaxPScales[BranchIndex].Get<1>(), PointScaleAtIndex);
-			}
-		}
-
-		// Remap BranchGradient and PlantGradient
-		const FVector2f NormalRange(0.0f, 1.0f);
-		const FVector2f PlantPScaleRange(MinimumPScale, MaximumPScale);
-
-		for (const int32& BranchIndex : PlantFacadeOut.GetBranchIndices(PlantNumber))
-		{
-			const TArray<int32>& BranchPoints = BranchFacadeOut.GetPoints(BranchIndex);
-			const FVector2f BranchPScaleRange(BranchMinMaxPScales[BranchIndex].Get<0>(), BranchMinMaxPScales[BranchIndex].Get<1>());
-
-			for (const int32& PointIndex : BranchPoints)
-			{
-				const float PlantGradientAtPointIndex = FMath::GetMappedRangeValueClamped(PlantPScaleRange, NormalRange,
-					PointFacadeOut.GetPointScale(PointIndex));
-				PointFacadeOut.SetPlantGradient(PointIndex, PlantGradientAtPointIndex);
-
-				const float BranchGradientAtPointIndex = FMath::GetMappedRangeValueClamped(BranchPScaleRange, NormalRange,
-					PointFacadeOut.GetPointScale(PointIndex));
-				PointFacadeOut.SetBranchGradient(PointIndex, BranchGradientAtPointIndex);
 			}
 		}
 	}
diff --git a/Engine/Plugins/Experimental/ProceduralVegetationEditor/Source/ProceduralVegetation/Private/Implementations/PVGrafter.cpp b/Engine/Plugins/Experimental/ProceduralVegetationEditor/Source/ProceduralVegetation/Private/Implementations/PVGrafter.cpp
--- a/Engine/Plugins/Experimental/ProceduralVegetationEditor/Source/ProceduralVegetation/Private/Implementations/PVGrafter.cpp
+++ b/Engine/Plugins/Experimental/ProceduralVegetationEditor/Source/ProceduralVegetation/Private/Implementations/PVGrafter.cpp
@@ -15,6 +15,7 @@
 #include "Helpers/PVPlantTraversalHelper.h"
 #include "Helpers/PVFoliageJSONHelper.h"
 #include "Helpers/PVUtilities.h"
+#include "Helpers/PVAttributesHelper.h"
 
 using namespace PV::DistributionConditionUtils;
 using namespace PV::ParametricDistributionHelper;
@@ -59,36 +60,6 @@ FTransform FPVGrafter::GetPointTransform(const FVector3f AttachmentPoint, const
 	return FTransform(RotationQuat, FVector(AttachmentPoint), FVector(Scale));
 }
 
-
-void FPVGrafter::ComputeBranchGradientAttribute(
-	PV::FPointBranchGradientAttributeView PointBranchGradientAttribute,
-	PV::FBranchParentNumberAttributeConstView BranchParentNumberAttribute,
-	PV::FBranchChildrenAttributeConstView BranchChildrenAttribute,
-	PV::FBranchNumberAttributeConstView BranchNumberAttribute,
-	PV::FBranchPointsAttributeConstView BranchPointsAttribute)
-{
-	using namespace PV::PlantTraversalHelper;
-
-	RecursiveWalkBranches(
-		BranchParentNumberAttribute,
-		BranchChildrenAttribute,
-		BranchNumberAttribute,
-		[&](int32 BranchIndex)-> EForEachResult
-			{
-				const TArray<int32>& BranchPoints = BranchPointsAttribute[BranchIndex];
-				check(BranchPoints.Num() > 1);
-
-				for (int32 i = 0; i < BranchPoints.Num(); ++i)
-				{
-					const float BranchAlpha = i / static_cast<float>(BranchPoints.Num() - 1);
-					PointBranchGradientAttribute[BranchPoints[i]] = 1.0 - BranchAlpha;
-				}
-
-				return EForEachResult::Continue;
-			}
-	);
-}
-
 void FPVGrafter::RescaleBranchForGraftingOnTip(FManagedArrayCollection& OutCollection, const int32 BranchIndex, const int32 BranchTipPointIndex,
                                                const float GraftTrunkLength)
 {
@@ -767,11 +738,13 @@ void FPVGrafter::DistributeGraftWithHormoneBasedSettings(
 	FRandomStream RandomStream;
 	RandomStream.Initialize(RandomSeed);
 
-	auto PointBranchGradientAttribute = PV::FPointBranchGradientAttribute::GetAttribute(OutCollection);
-	auto BranchParentNumberAttribute = PV::FBranchParentNumberAttribute::GetAttribute(OutCollection);
-	auto BranchChildrenAttribute = PV::FBranchChildrenAttribute::GetAttribute(OutCollection);
-	auto BranchNumberAttribute = PV::FBranchNumberAttribute::GetAttribute(OutCollection);
 	auto BranchPointsAttribute = PV::FBranchPointsAttribute::GetAttribute(OutCollection);
+	auto PointLengthFromRootAttribute = PV::FPointLengthFromRootAttribute::GetAttribute(OutCollection);
+
+	if (!ValidateAttributeCollection(BranchPointsAttribute, PointLengthFromRootAttribute))
+	{
+		return;
+	}
 
 	// Only need to do this temporarily since legacy skeleton has inconsistent data in children
 	// TODO: Remove this once support for legacy skeletons is dropped
@@ -779,13 +752,6 @@ void FPVGrafter::DistributeGraftWithHormoneBasedSettings(
 	BranchFacadeOut.RecomputeBranchChildren();
 	BranchFacadeOut.RecomputeBranchParents();
 
-	ComputeBranchGradientAttribute(
-		PointBranchGradientAttribute,
-		BranchParentNumberAttribute,
-		BranchChildrenAttribute,
-		BranchNumberAttribute,
-		BranchPointsAttribute);
-
 	FNormalizedPointCaches NormalizedPointCaches;
 	FBuildOptions BuildOptions;
 	BuildOptions.bComputePScales = DistributionConditions.IsActiveCondition(EFoliageDistributionCondition::Scale);
@@ -858,7 +824,6 @@ void FPVGrafter::DistributeGraftWithHormoneBasedSettings(
 			}
 		};
 
-	auto PointLengthFromRootAttribute = PV::FPointLengthFromRootAttribute::GetAttribute(OutCollection);
 	for (const auto& AttachmentPoint : AttachmentPoints)
 	{
 		const int32 MaxBudNumber = *Algo::MaxElement(BudNumberAttribute);
diff --git a/Engine/Plugins/Experimental/ProceduralVegetationEditor/Source/ProceduralVegetation/Private/Implementations/PVGrafter.h b/Engine/Plugins/Experimental/ProceduralVegetationEditor/Source/ProceduralVegetation/Private/Implementations/PVGrafter.h
--- a/Engine/Plugins/Experimental/ProceduralVegetationEditor/Source/ProceduralVegetation/Private/Implementations/PVGrafter.h
+++ b/Engine/Plugins/Experimental/ProceduralVegetationEditor/Source/ProceduralVegetation/Private/Implementations/PVGrafter.h
@@ -32,13 +32,6 @@ struct FPVGrafter
 	static FTransform GetPointTransform(const FVector3f AttachmentPoint, const FVector3f UpVector, const FVector3f NormalVector,
 	                                    const float Scale);
 
-	static void ComputeBranchGradientAttribute(
-		PV::FPointBranchGradientAttributeView PointBranchGradientAttribute,
-		PV::FBranchParentNumberAttributeConstView BranchParentNumberAttribute,
-		PV::FBranchChildrenAttributeConstView BranchChildrenAttribute,
-		PV::FBranchNumberAttributeConstView BranchNumberAttribute,
-		PV::FBranchPointsAttributeConstView BranchPointsAttribute);
-
 	static void RescaleBranchForGraftingOnTip(FManagedArrayCollection& OutCollection, const int32 BranchIndex, const int32 BranchTipPointIndex,
 	                                          const float GraftTrunkLength);
 
diff --git a/Engine/Plugins/Experimental/ProceduralVegetationEditor/Source/ProceduralVegetation/Private/Implementations/PVGrower.cpp b/Engine/Plugins/Experimental/ProceduralVegetationEditor/Source/ProceduralVegetation/Private/Implementations/PVGrower.cpp
--- a/Engine/Plugins/Experimental/ProceduralVegetationEditor/Source/ProceduralVegetation/Private/Implementations/PVGrower.cpp
+++ b/Engine/Plugins/Experimental/ProceduralVegetationEditor/Source/ProceduralVegetation/Private/Implementations/PVGrower.cpp
@@ -8,6 +8,7 @@
 #include "Facades/PVAttributesNames.h"
 #include "Facades/PVBranchFacade.h"
 #include "Facades/PVPointFacade.h"
+#include "Helpers/PVAttributesHelper.h"
 #include "GeometryCollection/ManagedArrayCollection.h"
 #include "Implementations/PVSeedGenerator.h"
 #include "Implementations/PVLightDetection.h"
@@ -115,17 +116,11 @@ void FPVGrower::FillCollection(FManagedArrayCollection& Collection, UPVGrowerDat
 
 	TArray<FVector3f> Points;
 	TArray<float> LengthFromRoots;
-	TArray<float> PointScaleGradients;
-	TArray<float> HullGradients;
-	TArray<float> MainTrunkGradients;
-	TArray<float> GroundGradients;
 	TArray<float> PointScales;
 	TArray<float> SeedPScales;
 	TArray<float> SeedPScaleRatios;
 	TArray<TArray<FVector3f>> BudDirections;
-	TArray<float> BranchGradients;
 	TArray<TArray<float>> BudHormoneLevels;
-	TArray<float> PlantGradients;
 	TArray<float> LengthFromSeeds;
 	TArray<int32> BudNumbers;
 	TArray<TArray<float>> BudLightDetected;
@@ -158,17 +153,11 @@ void FPVGrower::FillCollection(FManagedArrayCollection& Collection, UPVGrowerDat
 	PointFacade.AddElements(Points.Num());
 	PointFacade.FillPositions(Points);
 	PointFacade.FillLFR(LengthFromRoots);
-	PointFacade.FillPointScaleGradients(PointScaleGradients);
-	PointFacade.FillHullGradients(HullGradients);
-	PointFacade.FillMainTrunkGradients(MainTrunkGradients);
-	PointFacade.FillGroundGradients(GroundGradients);
 	PointFacade.FillPointScales(PointScales);
 	PointFacade.FillBudDirections(BudDirections);
 	PointFacade.FillBudStatus(BudStatus);
 	PointFacade.FillBudLateralMeristem(BudLateralMeristem);
-	PointFacade.FillBranchGradients(BranchGradients);
 	PointFacade.FillBudHormoneLevels(BudHormoneLevels);
-	PointFacade.FillPlantGradients(PlantGradients);
 	PointFacade.FillLengthFromSeeds(LengthFromSeeds);
 	PointFacade.FillBudNumbers(BudNumbers);
 	PointFacade.FillBudLightDetected(BudLightDetected);
diff --git a/Engine/Plugins/Experimental/ProceduralVegetationEditor/Source/ProceduralVegetation/Private/Implementations/PVMeshBuilder.cpp b/Engine/Plugins/Experimental/ProceduralVegetationEditor/Source/ProceduralVegetation/Private/Implementations/PVMeshBuilder.cpp
--- a/Engine/Plugins/Experimental/ProceduralVegetationEditor/Source/ProceduralVegetation/Private/Implementations/PVMeshBuilder.cpp
+++ b/Engine/Plugins/Experimental/ProceduralVegetationEditor/Source/ProceduralVegetation/Private/Implementations/PVMeshBuilder.cpp
@@ -72,7 +72,11 @@ void FPVMeshBuilder::GenerateGeometryCollection(const FManagedArrayCollection& I
 	
 	AddPointsToSkeleton(OutGeometryCollection);
 	
-	PV::AttributesHelper::ComputeGradients(OutGeometryCollection);
+	PV::AttributesHelper::ComputePointPlantGradient(OutGeometryCollection);
+	PV::AttributesHelper::ComputePointGroundGradient(OutGeometryCollection);
+	PV::AttributesHelper::ComputePointScaleGradient(OutGeometryCollection);
+	PV::AttributesHelper::ComputePointHullGradient(OutGeometryCollection);
+	PV::AttributesHelper::ComputePointMainTrunkGradient(OutGeometryCollection);
 
 	MeshBuilderParams.MaterialSettings.ApplyMaterialSettings(OutGeometryCollection);
 	GetLocalDynamicMeshData(OutGeometryCollection, InPlantProfileCollection, MeshBuilderParams, MeshesDataCollection);
@@ -464,21 +468,24 @@ TSet<int32> FPVMeshBuilder::CollectHardPoints(const PV::Facades::FBranchFacade&
 	return HardPoints;
 }
 
-TSet<int32> FPVMeshBuilder::ComputePointGradients(const PV::Facades::FPointFacade& PointFacade,
+TSet<int32> FPVMeshBuilder::ComputePointGradients(PV::FPointHullGradientAttributeConstView HullGradientAttribute,
+                                                  PV::FPointMainTrunkGradientAttributeConstView MainTrunkGradientAttribute,
+                                                  PV::FPointGroundGradientAttributeConstView GroundGradientAttribute,
+                                                  PV::FPointScaleGradientAttributeConstView ScaleGradientAttribute,
                                                   const FPVMeshBuilderParams& MeshBuilderParams, const TSet<int32>& HardPoints,
                                                   const float MaxPointScale, TArray<float>& OutMeshDivisionsGradients,
                                                   TArray<float>& OutDeltaModifiers)
 {
 	TSet<int32> PointsToRemove;
-	OutMeshDivisionsGradients.Reserve(PointFacade.GetElementCount());
-	OutDeltaModifiers.Reserve(PointFacade.GetElementCount());
+	OutMeshDivisionsGradients.Reserve(HullGradientAttribute.Num());
+	OutDeltaModifiers.Reserve(HullGradientAttribute.Num());
 
-	for (int32 PointIndex = 0; PointIndex < PointFacade.GetElementCount(); ++PointIndex)
+	for (int32 PointIndex = 0; PointIndex < HullGradientAttribute.Num(); ++PointIndex)
 	{
-		const float HullGradient = PointFacade.GetHullGradient(PointIndex);
-		const float MainTrunkGradient = PointFacade.GetMainTrunkGradient(PointIndex);
-		const float GroundGradient = PointFacade.GetGroundGradient(PointIndex);
-		const float ScaleGradient = PointFacade.GetPointScaleGradient(PointIndex);
+		const float HullGradient = HullGradientAttribute[PointIndex];
+		const float MainTrunkGradient = MainTrunkGradientAttribute[PointIndex];
+		const float GroundGradient = GroundGradientAttribute[PointIndex];
+		const float ScaleGradient = ScaleGradientAttribute[PointIndex];
 
 		// Compute scale removal gradient
 		const float PointHullGradient =
@@ -866,7 +873,7 @@ void FPVMeshBuilder::GenerateBranchMeshData(const bool bPrimitiveIsTrunk, const
 	const TManagedArray<FVector3f>& PointPositions = PointFacade.GetPositions();
 	const TManagedArray<TArray<FVector3f>>& PointBudDirections = BudVectorFacade.GetBudDirections();
 	const TManagedArray<float>& PointScales = PointFacade.GetPointScales();
-	
+
 	const float MaxPointScale = *Algo::MaxElement(PointScales);
 	const float MaxScaleRatio = 1.0f / (MaxPointScale * UE_TWO_PI);
 	int32 PolyGroupIndex = 0;
@@ -888,13 +895,16 @@ void FPVMeshBuilder::GenerateBranchMeshData(const bool bPrimitiveIsTrunk, const
 	{
 		return FMath::IsNearlyZero(InDenominator) ? InNumerator : (InNumerator / InDenominator);
 	};
-	
+
+	const float BaseLengthFromRoot = PointFacade.GetLengthFromRoot(PrimitivePoints[0]);
+	const float TipLengthFromRoot = PointFacade.GetLengthFromRoot(PrimitivePoints.Last());
+
 	FVector3f TangentVector = FVector3f::UpVector;
 	for (int32 Index = 0; Index < PrimitivePoints.Num(); ++Index)
 	{
 		const int32& PointIndex = PrimitivePoints[Index];
 		const float NjordPixelId = PointIndex;
-		
+
 		const float LengthFromRoot = PointFacade.GetLengthFromRoot(PointIndex);
 
 		float TextureCoordV = PointFacade.GetTextureCoordV(PointIndex);
@@ -984,7 +994,11 @@ void FPVMeshBuilder::GenerateBranchMeshData(const bool bPrimitiveIsTrunk, const
 		const FVector2f UV2 = FVector2f(LengthFromRoot * 0.01, SafeDivide(PointScale, MaxPointScale));
 		const float UVOGenerationOffset = GenerationNumber - 1;
 
-		const float BranchGradient = PointFacade.GetBranchGradient(PointIndex);
+		const float BranchGradient = FMath::GetMappedRangeValueClamped(
+			FVector2f(BaseLengthFromRoot, TipLengthFromRoot),
+			FVector2f(0.0f, 1.0f),
+			LengthFromRoot);
+
 		const float ProfileFallOff = MeshBuilderParams.PlantProfileFallOff.GetRichCurveConst()->Eval(BranchGradient);
 		if (PointScale == 0.0f)
 		[[unlikely]]
@@ -1285,21 +1299,59 @@ void FPVMeshBuilder::GetLocalDynamicMeshData(FManagedArrayCollection& Collection
 	const PV::Facades::FPointFacade PointFacade(Collection);
 	const PV::Facades::FPlantFacade PlantFacade(Collection);
 
-	if (!BranchFacade.IsValid())
+	if (!BranchFacade.IsValid() || !PointFacade.IsValid() || !PlantFacade.IsValid())
 	{
 		UE_LOG(LogProceduralVegetation, Warning, TEXT("Fail to evaluate, InCollection dose not match the schema."));
 		return;
 	}
 
-	const TManagedArray<float>& PointScales = PointFacade.GetPointScales();
-	const float MaxPointScale = *Algo::MaxElement(PointScales);
+	PV::FPointPositionAttributeConstView PointPositionAttribute = PV::FPointPositionAttribute::FindAttribute(Collection);
+	PV::FPointScaleAttributeConstView PointScaleAttribute = PV::FPointScaleAttribute::FindAttribute(Collection);
+	PV::FBranchPointsAttributeConstView BranchPointsAttribute = PV::FBranchPointsAttribute::FindAttribute(Collection);
+	PV::FBranchParentNumberAttributeConstView BranchParentNumberAttribute = PV::FBranchParentNumberAttribute::FindAttribute(Collection);
+	PV::FBranchPlantNumberAttributeConstView BranchPlantNumberAttribute = PV::FBranchPlantNumberAttribute::FindAttribute(Collection);
+	PV::FBranchNumberAttributeConstView BranchNumberAttribute = PV::FBranchNumberAttribute::FindAttribute(Collection);
+	PV::FBudDevelopmentAttributeConstView BudDevelopmentAttribute = PV::FBudDevelopmentAttribute::FindAttribute(Collection);
+	if (!PV::ValidateAttributeCollection(
+		PointPositionAttribute,
+		PointScaleAttribute,
+		BranchPointsAttribute,
+		BranchParentNumberAttribute,
+		BranchPlantNumberAttribute,
+		BranchNumberAttribute,
+		BudDevelopmentAttribute
+	))
+	{
+		return;
+	}
+
+	PV::FPointGroundGradientAttributeView GroundGradientAttribute = PV::FPointGroundGradientAttribute::AddAttribute(Collection);
+	PV::FPointScaleGradientAttributeView PointScaleGradientAttribute = PV::FPointScaleGradientAttribute::AddAttribute(Collection);
+	PV::FPointHullGradientAttributeView HullGradientAttribute = PV::FPointHullGradientAttribute::AddAttribute(Collection);
+	PV::FPointMainTrunkGradientAttributeView MainTrunkGradientAttribute = PV::FPointMainTrunkGradientAttribute::AddAttribute(Collection);
+
+	PV::AttributesHelper::ComputePointGroundGradient({ GroundGradientAttribute, PointPositionAttribute });
+	PV::AttributesHelper::ComputePointScaleGradient({ PointScaleGradientAttribute, PointScaleAttribute });
+	PV::AttributesHelper::ComputePointHullGradient({ HullGradientAttribute, PointPositionAttribute, BranchPointsAttribute, BranchParentNumberAttribute });
+	PV::AttributesHelper::ComputePointMainTrunkGradient({ MainTrunkGradientAttribute, PointPositionAttribute, BranchPointsAttribute, BranchParentNumberAttribute, BranchPlantNumberAttribute, BranchNumberAttribute, BudDevelopmentAttribute });
+
+	const float MaxPointScale = *Algo::MaxElement(PointScaleAttribute);
 
 	const TSet<int32> HardPoints = CollectHardPoints(BranchFacade, PointFacade, PlantFacade);
 
 	TArray<float> MeshDivisionsGradientsForPoints;
 	TArray<float> DeltaModifiersForPoints;
-	TSet<int32> PointsToRemove = ComputePointGradients(PointFacade, MeshBuilderParams, HardPoints, MaxPointScale, MeshDivisionsGradientsForPoints,
-		DeltaModifiersForPoints);
+	TSet<int32> PointsToRemove = ComputePointGradients(
+		HullGradientAttribute, 
+		MainTrunkGradientAttribute, 
+		GroundGradientAttribute, 
+		PointScaleGradientAttribute, 
+		MeshBuilderParams, 
+		HardPoints, 
+		MaxPointScale, 
+		MeshDivisionsGradientsForPoints,
+		DeltaModifiersForPoints
+	);
 
 	PerformPathSimplification(BranchFacade, PointFacade, MeshBuilderParams, MaxPointScale, HardPoints, DeltaModifiersForPoints, PointsToRemove);
 
diff --git a/Engine/Plugins/Experimental/ProceduralVegetationEditor/Source/ProceduralVegetation/Private/Implementations/PVMeshBuilder.h b/Engine/Plugins/Experimental/ProceduralVegetationEditor/Source/ProceduralVegetation/Private/Implementations/PVMeshBuilder.h
--- a/Engine/Plugins/Experimental/ProceduralVegetationEditor/Source/ProceduralVegetation/Private/Implementations/PVMeshBuilder.h
+++ b/Engine/Plugins/Experimental/ProceduralVegetationEditor/Source/ProceduralVegetation/Private/Implementations/PVMeshBuilder.h
@@ -307,7 +307,11 @@ private:
 	static TSet<int32> CollectHardPoints(const PV::Facades::FBranchFacade& BranchFacade, const PV::Facades::FPointFacade& PointFacade,
 	                                     const PV::Facades::FPlantFacade& PlantFacade);
 
-	static TSet<int32> ComputePointGradients(const PV::Facades::FPointFacade& PointFacade, const FPVMeshBuilderParams& MeshBuilderParams,
+	static TSet<int32> ComputePointGradients(PV::FPointHullGradientAttributeConstView HullGradientAttribute,
+		                                     PV::FPointMainTrunkGradientAttributeConstView MainTrunkGradientAttribute,
+		                                     PV::FPointGroundGradientAttributeConstView GroundGradientAttribute,
+		                                     PV::FPointScaleGradientAttributeConstView ScaleGradientAttribute,
+		                                     const FPVMeshBuilderParams& MeshBuilderParams,
 	                                         const TSet<int32>& HardPoints, const float MaxPointScale, TArray<float>& OutMeshDivisionsGradients,
 	                                         TArray<float>& OutDeltaModifiers);
 
diff --git a/Engine/Plugins/Experimental/ProceduralVegetationEditor/Source/ProceduralVegetation/Private/Implementations/PVSeedGenerator.cpp b/Engine/Plugins/Experimental/ProceduralVegetationEditor/Source/ProceduralVegetation/Private/Implementations/PVSeedGenerator.cpp
--- a/Engine/Plugins/Experimental/ProceduralVegetationEditor/Source/ProceduralVegetation/Private/Implementations/PVSeedGenerator.cpp
+++ b/Engine/Plugins/Experimental/ProceduralVegetationEditor/Source/ProceduralVegetation/Private/Implementations/PVSeedGenerator.cpp
@@ -108,15 +108,10 @@ void FPVConvertToSeedPointImplementation::FillCollection(FManagedArrayCollection
 
 	TArray<FVector3f> Points;
 	TArray<float> LengthFromRoots;
-	TArray<float> PointScaleGradients;
-	TArray<float> HullGradients;
-	TArray<float> MainTrunkGradients;
-	TArray<float> GroundGradients;
 	TArray<float> PointScales;
 	TArray<float> SeedPScales;
 	TArray<float> SeedPScaleRatios;
 	TArray<TArray<FVector3f>> BudDirections;
-	TArray<float> BranchGradients;
 	TArray<TArray<float>> BudHormoneLevels;
 	TArray<float> PlantGradients;
 	TArray<float> LengthFromSeeds;
@@ -199,17 +194,11 @@ void FPVConvertToSeedPointImplementation::FillCollection(FManagedArrayCollection
 	PointFacade.AddElements(Points.Num());
 	PointFacade.FillPositions(Points);
 	PointFacade.FillLFR(LengthFromRoots);
-	PointFacade.FillPointScaleGradients(PointScaleGradients);
-	PointFacade.FillHullGradients(HullGradients);
-	PointFacade.FillMainTrunkGradients(MainTrunkGradients);
-	PointFacade.FillGroundGradients(GroundGradients);
 	PointFacade.FillPointScales(PointScales);
 	PointFacade.FillBudDirections(BudDirections);
 	PointFacade.FillBudStatus(BudStatus);
 	PointFacade.FillBudLateralMeristem(BudLateralMeristem);
-	PointFacade.FillBranchGradients(BranchGradients);
 	PointFacade.FillBudHormoneLevels(BudHormoneLevels);
-	PointFacade.FillPlantGradients(PlantGradients);
 	PointFacade.FillLengthFromSeeds(LengthFromSeeds);
 	PointFacade.FillBudNumbers(BudNumbers);
 	PointFacade.FillBudLightDetected(BudLightDetected);
diff --git a/Engine/Plugins/Experimental/ProceduralVegetationEditor/Source/ProceduralVegetation/Private/Nodes/Development/PVGrowthDataDiffSettings.cpp b/Engine/Plugins/Experimental/ProceduralVegetationEditor/Source/ProceduralVegetation/Private/Nodes/Development/PVGrowthDataDiffSettings.cpp
--- a/Engine/Plugins/Experimental/ProceduralVegetationEditor/Source/ProceduralVegetation/Private/Nodes/Development/PVGrowthDataDiffSettings.cpp
+++ b/Engine/Plugins/Experimental/ProceduralVegetationEditor/Source/ProceduralVegetation/Private/Nodes/Development/PVGrowthDataDiffSettings.cpp
@@ -435,7 +435,6 @@ bool FPVGrowthDataDiffElement::ExecuteInternal(FPCGContext* InContext) const
 		PV::FBudDirectionAttribute::ManagedType,
 		PV::FBudStatusAttribute::ManagedType,
 		PV::FBudLateralMeristemAttribute::ManagedType,
-		PV::FPointBranchGradientAttribute::ManagedType,
 		PV::FBudHormoneLevelsAttribute::ManagedType,
 		PV::FPointPlantGradientAttribute::ManagedType,
 		PV::FPointBudNumberAttribute::ManagedType,
diff --git a/Engine/Plugins/Experimental/ProceduralVegetationEditor/Source/ProceduralVegetation/Private/Utils/PVAttributes.h b/Engine/Plugins/Experimental/ProceduralVegetationEditor/Source/ProceduralVegetation/Private/Utils/PVAttributes.h
--- a/Engine/Plugins/Experimental/ProceduralVegetationEditor/Source/ProceduralVegetation/Private/Utils/PVAttributes.h
+++ b/Engine/Plugins/Experimental/ProceduralVegetationEditor/Source/ProceduralVegetation/Private/Utils/PVAttributes.h
@@ -426,7 +426,6 @@ PV_DEFINE_ATTRIBUTE(PointHullGradient, float, AttributeNames::HullGradientStr, G
 PV_DEFINE_ATTRIBUTE(PointMainTrunkGradient, float, AttributeNames::MainTrunkGradientStr, GroupNames::PointGroupStr);
 PV_DEFINE_ATTRIBUTE(PointGroundGradient, float, AttributeNames::GroundGradientStr, GroupNames::PointGroupStr);
 PV_DEFINE_ATTRIBUTE(PointScaleGradient, float, AttributeNames::PointScaleGradientStr, GroupNames::PointGroupStr);
-PV_DEFINE_ATTRIBUTE(PointBranchGradient, float, AttributeNames::BranchGradientStr, GroupNames::PointGroupStr);
 PV_DEFINE_ATTRIBUTE(PointPlantGradient, float, AttributeNames::PlantGradientStr, GroupNames::PointGroupStr);
 PV_DEFINE_ATTRIBUTE(PointBudNumber, int32, AttributeNames::BudNumberStr, GroupNames::PointGroupStr);
 PV_DEFINE_ATTRIBUTE(PointNjordPixelIndex, float, AttributeNames::NjordPixelIndexStr, GroupNames::PointGroupStr);
diff --git a/Engine/Plugins/Experimental/ProceduralVegetationEditor/Source/ProceduralVegetation/ProceduralVegetation.Build.cs b/Engine/Plugins/Experimental/ProceduralVegetationEditor/Source/ProceduralVegetation/ProceduralVegetation.Build.cs
--- a/Engine/Plugins/Experimental/ProceduralVegetationEditor/Source/ProceduralVegetation/ProceduralVegetation.Build.cs
+++ b/Engine/Plugins/Experimental/ProceduralVegetationEditor/Source/ProceduralVegetation/ProceduralVegetation.Build.cs
@@ -18,6 +18,7 @@ namespace UnrealBuildTool.Rules
 				new string[]
 				{
 					"Chaos",
+					"ChaosCore",
 					"Core",
 					"CoreUObject",
 					"Engine",
diff --git a/Engine/Plugins/Experimental/ProceduralVegetationEditor/Source/ProceduralVegetation/Public/Facades/PVAttributesNames.h b/Engine/Plugins/Experimental/ProceduralVegetationEditor/Source/ProceduralVegetation/Public/Facades/PVAttributesNames.h
--- a/Engine/Plugins/Experimental/ProceduralVegetationEditor/Source/ProceduralVegetation/Public/Facades/PVAttributesNames.h
+++ b/Engine/Plugins/Experimental/ProceduralVegetationEditor/Source/ProceduralVegetation/Public/Facades/PVAttributesNames.h
@@ -77,7 +77,6 @@ namespace PV::AttributeNames
 	inline constexpr TCHAR BudHormoneLevelsStr[] = TEXT("budHormoneLevels");
 	inline constexpr TCHAR BudDirectionStr[] = TEXT("budDirection");
 	inline constexpr TCHAR BudStatusStr[] = TEXT("budStatus");
-	inline constexpr TCHAR BranchGradientStr[] = TEXT("branchGradient");
 	inline constexpr TCHAR PlantGradientStr[] = TEXT("plantGradient");
 	inline constexpr TCHAR NjordPixelIndexStr[] = TEXT("njord_pixelIdx");
 
@@ -99,7 +98,6 @@ namespace PV::AttributeNames
 	inline static const FName BudHormoneLevels = FName(BudHormoneLevelsStr);
 	inline static const FName BudDirection = FName(BudDirectionStr);
 	inline static const FName BudStatus = FName(BudStatusStr);
-	inline static const FName BranchGradient = FName(BranchGradientStr);
 	inline static const FName PlantGradient = FName(PlantGradientStr);
 	inline static const FName NjordPixelIndex = FName(NjordPixelIndexStr);
 
diff --git a/Engine/Plugins/Experimental/ProceduralVegetationEditor/Source/ProceduralVegetation/Public/Facades/PVPointFacade.h b/Engine/Plugins/Experimental/ProceduralVegetationEditor/Source/ProceduralVegetation/Public/Facades/PVPointFacade.h
--- a/Engine/Plugins/Experimental/ProceduralVegetationEditor/Source/ProceduralVegetation/Public/Facades/PVPointFacade.h
+++ b/Engine/Plugins/Experimental/ProceduralVegetationEditor/Source/ProceduralVegetation/Public/Facades/PVPointFacade.h
@@ -51,14 +51,6 @@ Accessor.ModifyAt(i, InputArray[i]); \
 		
 		void SetLengthFromSeed(int32 Index, float InLength);
 
-		float GetPointScaleGradient(int32 Index) const;
-
-		float GetHullGradient(int32 Index) const;
-
-		float GetMainTrunkGradient(int32 Index) const;
-
-		float GetGroundGradient(int32 Index) const;
-		
 		float GetPointScale(int32 Index) const;
 
 		void SetPointScale(int32 Index, float InPointScale);
@@ -101,14 +93,6 @@ Accessor.ModifyAt(i, InputArray[i]); \
 		
 		void SetBudLateralMeristem(int32 Index, const TArray<float>& InBudLateralMeristem);
 
-		float GetBranchGradient(int32 Index) const;
-
-		void SetBranchGradient(int32 Index, float InBranchGradient);
-
-		float GetPlantGradient(int32 Index) const;
-
-		void SetPlantGradient(int32 Index, float InPlantGradient);
-	
 		float GetNjordPixelIndex(int32 Index) const;
 
 		void SetNjordPixelIndex(int32 Index, float InNjordPixelIndex);
@@ -152,25 +136,13 @@ Accessor.ModifyAt(i, InputArray[i]); \
 		void FillPositions(const TArray<FVector3f>& InputArray);
 
 		void FillLFR(const TArray<float>& InputArray);
-
-		void FillPointScaleGradients(const TArray<float>& InputArray);
-		
-		void FillHullGradients(const TArray<float>& InputArray);
-		
-		void FillMainTrunkGradients(const TArray<float>& InputArray);
-		
-		void FillGroundGradients(const TArray<float>& InputArray);
 		
 		void FillPointScales(const TArray<float>& InputArray);
 
 		void FillBudDirections(const TArray<TArray<FVector3f>>& InputArray);
 		
-		void FillBranchGradients(const TArray<float>& InputArray);
-		
 		void FillBudHormoneLevels(const TArray<TArray<float>>& InputArray);
 		
-		void FillPlantGradients(const TArray<float>& InputArray);
-		
 		void FillLengthFromSeeds(const TArray<float>& InputArray);
 		
 		void FillBudNumbers(const TArray<int32>& InputArray);
@@ -218,24 +190,12 @@ Accessor.ModifyAt(i, InputArray[i]); \
 		
 		void GetLFR(int32 Index, float& LFR) const;
 
-		void GetPointScaleGradient(int32 Index, float& OutPointScaleGradient) const;
-		
-		void GetHullGradient(int32 Index, float& OutHullGradient) const;
-
-		void GetMainTrunkGradients(int32 Index, float& OutValue) const;
-
-		void GetGroundGradients(int32 Index, float& OutValue) const;
-
 		void GetPointScales(int32 Index, float& OutValue) const;
 
 		void GetBudDirections(int32 Index, TArray<FVector3f>& OutValue) const;
 
-		void GetBranchGradients(int32 Index, float& OutValue) const;
-
 		void GetBudHormoneLevels(int32 Index, TArray<float>& OutValue) const;
 
-		void GetPlantGradients(int32 Index, float& OutValue) const;
-
 		void GetBudNumbers(int32 Index, int32& OutValue) const;
 		
 		const TManagedArray<int32>& GetBudNumbersAttribute() const;
@@ -277,7 +237,6 @@ Accessor.ModifyAt(i, InputArray[i]); \
 		TManagedArrayAccessor<TArray<FVector3f>> BudDirections;
 		TManagedArrayAccessor<TArray<int>> BudStatus;
 		TManagedArrayAccessor<TArray<float>> BudLateralMeristem;
-		TManagedArrayAccessor<float> BranchGradients;
 		TManagedArrayAccessor<TArray<float>> BudHormoneLevels;
 		TManagedArrayAccessor<float> PlantGradients;
 		TManagedArrayAccessor<float> LengthFromSeed;
