Fix out-of-date procedural vegetation import unit test
Update expected branch parent attributes to include the implicit root entry (0) and correct the expected children of the root branch
Replace order-sensitive equality checks for branch children with set-based comparison to avoid ordering assumptions

#Synced-CL 52140050

diff --git a/Engine/Plugins/Experimental/ProceduralVegetationEditor/Source/ProceduralVegetation/Private/Tests/PVImportTests.cpp b/Engine/Plugins/Experimental/ProceduralVegetationEditor/Source/ProceduralVegetation/Private/Tests/PVImportTests.cpp
--- a/Engine/Plugins/Experimental/ProceduralVegetationEditor/Source/ProceduralVegetation/Private/Tests/PVImportTests.cpp
+++ b/Engine/Plugins/Experimental/ProceduralVegetationEditor/Source/ProceduralVegetation/Private/Tests/PVImportTests.cpp
@@ -62,12 +62,12 @@ PV_IMPORTER_TEST(FillAttributesFromBranchHierarchy)
 		ExpectedBranchParentNumberAttribute[2] = PVMockTree::Branch1_BranchNumber;
 		ExpectedBranchParentNumberAttribute[3] = PVMockTree::Branch0_BranchNumber;
 
-		ExpectedBranchParentsAttribute[0] = { };
-		ExpectedBranchParentsAttribute[1] = { PVMockTree::Branch0_BranchNumber };
-		ExpectedBranchParentsAttribute[2] = { PVMockTree::Branch0_BranchNumber, PVMockTree::Branch1_BranchNumber };
-		ExpectedBranchParentsAttribute[3] = { PVMockTree::Branch0_BranchNumber };
+		ExpectedBranchParentsAttribute[0] = { 0 }; // Root branch always has a branch parent number of "0" (0 is an invalid branch parent number in this case)
+		ExpectedBranchParentsAttribute[1] = { 0, PVMockTree::Branch0_BranchNumber };
+		ExpectedBranchParentsAttribute[2] = { 0, PVMockTree::Branch0_BranchNumber, PVMockTree::Branch1_BranchNumber };
+		ExpectedBranchParentsAttribute[3] = { 0, PVMockTree::Branch0_BranchNumber };
 
-		ExpectedBranchChildrenAttribute[0] = { PVMockTree::Branch1_BranchNumber, PVMockTree::Branch3_BranchNumber };
+		ExpectedBranchChildrenAttribute[0] = { PVMockTree::Branch1_BranchNumber, PVMockTree::Branch2_BranchNumber, PVMockTree::Branch3_BranchNumber };
 		ExpectedBranchChildrenAttribute[1] = { PVMockTree::Branch2_BranchNumber };
 		ExpectedBranchChildrenAttribute[2] = { };
 		ExpectedBranchChildrenAttribute[3] = { };
@@ -117,11 +117,18 @@ PV_IMPORTER_TEST(FillAttributesFromBranchHierarchy)
 	UTEST_EQUAL("BranchParents2", ActualBranchParentsAttribute[2], ExpectedBranchParentsAttribute[2]);
 	UTEST_EQUAL("BranchParents3", ActualBranchParentsAttribute[3], ExpectedBranchParentsAttribute[3]);
 
+	const static auto AreBranchChildrenAttributesEqual = [](const TArray<int32>& BranchChildrenA, const TArray<int32>& BranchChildrenB)
+	{
+		const auto SetA = TSet<int32>(BranchChildrenA);
+		const auto SetB = TSet<int32>(BranchChildrenB);
+		return SetA.Num() == SetB.Num() && SetA.Includes(SetB);
+	};
+
 	// Validate branch children
-	UTEST_EQUAL("BranchChildren0", ActualBranchChildrenAttribute[0], ExpectedBranchChildrenAttribute[0]);
-	UTEST_EQUAL("BranchChildren1", ActualBranchChildrenAttribute[1], ExpectedBranchChildrenAttribute[1]);
-	UTEST_EQUAL("BranchChildren2", ActualBranchChildrenAttribute[2], ExpectedBranchChildrenAttribute[2]);
-	UTEST_EQUAL("BranchChildren3", ActualBranchChildrenAttribute[3], ExpectedBranchChildrenAttribute[3]);
+	UTEST_TRUE("BranchChildren0", AreBranchChildrenAttributesEqual(ActualBranchChildrenAttribute[0], ExpectedBranchChildrenAttribute[0]));
+	UTEST_TRUE("BranchChildren1", AreBranchChildrenAttributesEqual(ActualBranchChildrenAttribute[1], ExpectedBranchChildrenAttribute[1]));
+	UTEST_TRUE("BranchChildren2", AreBranchChildrenAttributesEqual(ActualBranchChildrenAttribute[2], ExpectedBranchChildrenAttribute[2]));
+	UTEST_TRUE("BranchChildren3", AreBranchChildrenAttributesEqual(ActualBranchChildrenAttribute[3], ExpectedBranchChildrenAttribute[3]));
 
 	return true;
 }
