[PVE] Skeleton Edit update for fused points

#Synced-CL 52128411

diff --git a/Engine/Plugins/Experimental/ProceduralVegetationEditor/Source/ProceduralVegetation/Private/Nodes/PVSkeletonEditSettings.cpp b/Engine/Plugins/Experimental/ProceduralVegetationEditor/Source/ProceduralVegetation/Private/Nodes/PVSkeletonEditSettings.cpp
--- a/Engine/Plugins/Experimental/ProceduralVegetationEditor/Source/ProceduralVegetation/Private/Nodes/PVSkeletonEditSettings.cpp
+++ b/Engine/Plugins/Experimental/ProceduralVegetationEditor/Source/ProceduralVegetation/Private/Nodes/PVSkeletonEditSettings.cpp
@@ -7,6 +7,7 @@
 #include "Facades/PVBranchFacade.h"
 #include "Facades/PVBudVectorsFacade.h"
 #include "Facades/PVPointFacade.h"
+#include "Facades/PVTreeFacade.h"
 
 #include "Implementations/PVCarve.h"
 
@@ -93,7 +94,6 @@ bool FPVSkeletonEditElement::ExecuteInternal(FPCGContext* InContext) const
 			if (MaxBudNumber && Offsets.Num() >= *MaxBudNumber + 1)
 			{
 				TManagedArray<FVector3f>& Positions = PointFacade.ModifyPositions();
-				TManagedArray<FVector3f>& FoliagePositions = FoliageFacade.ModifyPivotPositions();
 				TManagedArray<TArray<FVector3f>>& BudDirections = BudVectorsFacade.ModifyBudDirections();
 
 				const TArray<FVector3f> OriginalPositions = Positions.GetConstArray();
@@ -101,7 +101,7 @@ bool FPVSkeletonEditElement::ExecuteInternal(FPCGContext* InContext) const
 				for (int32 BranchIndex = 0; BranchIndex < BranchFacade.GetElementCount(); BranchIndex++)
 				{
 					const TArray<int32>& BranchPoints = BranchFacade.GetPoints(BranchIndex);
-					for (int32 BranchPointIndex = 0; BranchPointIndex < BranchPoints.Num(); BranchPointIndex++)
+					for (int32 BranchPointIndex = !BranchFacade.IsTrunk(BranchIndex); BranchPointIndex < BranchPoints.Num(); BranchPointIndex++)
 					{
 						const int32 PointIndex = BranchPoints[BranchPointIndex];
 						const int32 PointBudNumber = PointFacade.GetBudNumber(PointIndex);
@@ -118,43 +118,25 @@ bool FPVSkeletonEditElement::ExecuteInternal(FPCGContext* InContext) const
 
 							TArray<FVector3f>& BudDirection = BudDirections[PointIndex];
 							BudDirection[PV::Facades::BudDirectionsApical] = Rotation.RotateVector(BudDirection[PV::Facades::BudDirectionsApical]);
-							BudDirection[PV::Facades::BudDirectionsAxillary] = Rotation.RotateVector(BudDirection[PV::Facades::BudDirectionsAxillary]);
-							BudDirection[PV::Facades::BudDirectionsGuideCurve] = Rotation.RotateVector(BudDirection[PV::Facades::BudDirectionsGuideCurve]);
-							BudDirection[PV::Facades::BudDirectionsUpVector] = Rotation.RotateVector(BudDirection[PV::Facades::BudDirectionsUpVector]);
-
-							const float PointLFR = PointFacade.GetLengthFromRoot(PointIndex);
-							const float PreviousPointLFR = PointFacade.GetLengthFromRoot(PreviousPointIndex);
-							for (const int32 FoliageEntryId : FoliageFacade.GetFoliageEntryIdsForBranch(BranchIndex))
-							{
-								const float FoliageLFR = FoliageFacade.GetLengthFromRoot(FoliageEntryId);
-								if (FoliageLFR <= PointLFR && FoliageLFR > PreviousPointLFR)
-								{
-									FoliagePositions[FoliageEntryId] += Offsets[PointBudNumber];
-									FoliageFacade.SetUpVector(
-										FoliageEntryId,
-										Rotation.RotateVector(FoliageFacade.GetUpVector(FoliageEntryId))
-									);
-									FoliageFacade.SetNormalVector(
-										FoliageEntryId,
-										Rotation.RotateVector(FoliageFacade.GetNormalVector(FoliageEntryId))
-									);
-								}
-							}
+							BudDirection[PV::Facades::BudDirectionsAxillary] = Rotation.
+								RotateVector(BudDirection[PV::Facades::BudDirectionsAxillary]);
+							BudDirection[PV::Facades::BudDirectionsGuideCurve] = Rotation.RotateVector(
+								BudDirection[PV::Facades::BudDirectionsGuideCurve]);
+							BudDirection[PV::Facades::BudDirectionsUpVector] = Rotation.
+								RotateVector(BudDirection[PV::Facades::BudDirectionsUpVector]);
 						}
 					}
 				}
 			}
-			
+
 			TArray<bool> RemovedPoints = InputSettings->SkeletonEditSettings.RemovedPoints;
 
 			TArray<bool> PointsToRemove;
 			TArray<bool> BranchesToRemove;
-			TArray<bool> FoliageInstancesToRemove;
 			if (MaxBudNumber && RemovedPoints.Num() >= *MaxBudNumber + 1)
 			{
 				PointsToRemove.Init(false, PointFacade.GetElementCount());
 				BranchesToRemove.Init(false, BranchFacade.GetElementCount());
-				FoliageInstancesToRemove.Init(false, FoliageFacade.NumFoliageEntries());
 
 				for (int32 BranchIndex = 0; BranchIndex < BranchFacade.GetElementCount(); BranchIndex++)
 				{
@@ -165,37 +147,46 @@ bool FPVSkeletonEditElement::ExecuteInternal(FPCGContext* InContext) const
 					{
 						const int32 PointIndex = BranchPoints[BranchPointIndex];
 						const int32 PointBudNumber = PointFacade.GetBudNumber(PointIndex);
-						PointsToRemove[PointIndex] = RemovedPoints[PointBudNumber];
-						
-						bAllPointsRemoved &= PointsToRemove[PointIndex];
+						if (!PointsToRemove[PointIndex])
+						{
+							PointsToRemove[PointIndex] = RemovedPoints[PointBudNumber];
+						}
+
+						if (BranchPointIndex != 0)
+						{
+							bAllPointsRemoved &= PointsToRemove[PointIndex];
+						}
+
 						if (PointsToRemove[PointIndex])
 						{
 							if (FirstPointRemovedIndex == INDEX_NONE)
 							{
 								FirstPointRemovedIndex = BranchPointIndex;
 							}
-							const int32 PreviousPointIndex = BranchPointIndex > 0
-								? BranchPoints[BranchPointIndex - 1]
-								: GetParentPoint(BranchIndex);
-							const float PointLFR = PointFacade.GetLengthFromRoot(PointIndex);
-							const float PreviousPointLFR = PointFacade.GetLengthFromRoot(PreviousPointIndex);
-							for (const int32 FoliageEntryId : FoliageFacade.GetFoliageEntryIdsForBranch(BranchIndex))
+						}
+					}
+					BranchesToRemove[BranchIndex] = bAllPointsRemoved;
+					if (bAllPointsRemoved)
+					{
+						const TManagedArray<int32>& BranchNumbers = BranchFacade.GetBranchNumbers();
+						for (int32 BranchChildNumber : BranchFacade.GetChildren(BranchIndex))
+						{
+							if (const int32 ChildBranchIndex = BranchNumbers.Find(BranchChildNumber); ChildBranchIndex != INDEX_NONE)
 							{
-								const float FoliageLFR = FoliageFacade.GetLengthFromRoot(FoliageEntryId);
-								if (FoliageLFR <= PointLFR && FoliageLFR > PreviousPointLFR)
+								for (int32 ChildPoint : BranchFacade.GetPoints(ChildBranchIndex))
 								{
-									FoliageInstancesToRemove[FoliageEntryId] = true;
+									PointsToRemove[ChildPoint] = true;
 								}
+								BranchesToRemove[ChildBranchIndex] = true;
 							}
 						}
 					}
-					BranchesToRemove[BranchIndex] = bAllPointsRemoved;
 
 					if (FirstPointRemovedIndex != INDEX_NONE)
 					{
 						const float Ratio = static_cast<float>(FirstPointRemovedIndex) / static_cast<float>(BranchPoints.Num());
 						// Rescale the branches based on the cut point
-						for (int32 Index = 0; Index <= FirstPointRemovedIndex; ++Index)
+						for (int32 Index = !BranchFacade.IsTrunk(BranchIndex); Index <= FirstPointRemovedIndex; ++Index)
 						{
 							PointFacade.SetPointScale(BranchPoints[Index], PointFacade.GetPointScale(BranchPoints[Index]) * Ratio);
 						}
@@ -203,8 +194,10 @@ bool FPVSkeletonEditElement::ExecuteInternal(FPCGContext* InContext) const
 						{
 							if (const int32 ChildBranchIndex = BranchFacade.GetBranchNumbers().Find(ChildBranch); ChildBranchIndex != INDEX_NONE)
 							{
-								for (int32 ChildBranchPoint : BranchFacade.GetPoints(ChildBranchIndex))
+								const TArray<int32>& ChildBranchPoints = BranchFacade.GetPoints(ChildBranchIndex);
+								for (int32 ChildBranchPointIndex = 1; ChildBranchPointIndex < ChildBranchPoints.Num(); ChildBranchPointIndex++)
 								{
+									const int32 ChildBranchPoint = ChildBranchPoints[ChildBranchPointIndex];
 									PointFacade.SetPointScale(ChildBranchPoint, PointFacade.GetPointScale(ChildBranchPoint) * Ratio);
 								}
 							}
@@ -212,10 +205,10 @@ bool FPVSkeletonEditElement::ExecuteInternal(FPCGContext* InContext) const
 					}
 				}
 
-				const FManagedArrayCollection InCollection = OutCollection;
-				FPVCarve::RemoveEntriesAndRecomputeAttributes(
+				// Foliage Instances are not being used now with the new Grower
+				TArray<bool> FoliageInstancesToRemove;
+				PV::Facades::FTreeFacade::RemoveEntriesAndReIndexAttributes(
 					OutCollection,
-					InCollection,
 					PointsToRemove,
 					BranchesToRemove,
 					FoliageInstancesToRemove
diff --git a/Engine/Plugins/Experimental/ProceduralVegetationEditor/Source/ProceduralVegetation/Public/Nodes/PVSkeletonEditSettings.h b/Engine/Plugins/Experimental/ProceduralVegetationEditor/Source/ProceduralVegetation/Public/Nodes/PVSkeletonEditSettings.h
--- a/Engine/Plugins/Experimental/ProceduralVegetationEditor/Source/ProceduralVegetation/Public/Nodes/PVSkeletonEditSettings.h
+++ b/Engine/Plugins/Experimental/ProceduralVegetationEditor/Source/ProceduralVegetation/Public/Nodes/PVSkeletonEditSettings.h
@@ -21,8 +21,9 @@ enum class EPointSelectionSmoothnessMethod: uint8
 UENUM()
 enum class ESkeletonSelectionMode: uint8
 {
-	SelectByDistance,
+	SelectByTreeDistance,
 	SelectByNeighbours,
+	SelectByEuclideanDistance,
 };
 
 UENUM()
@@ -55,11 +56,8 @@ struct FPVSkeletonSelectionParams
 	int32 NumPointsSelection = 0;
 
 	UPROPERTY(EditAnywhere, Category = "Selection",
-		meta = (UIMin=0, ClampMin=0, EditCondition="SkeletonSelectionMode == ESkeletonSelectionMode::SelectByDistance", EditConditionHides))
+		meta = (UIMin=0, ClampMin=0, EditCondition="SkeletonSelectionMode != ESkeletonSelectionMode::SelectByNeighbours", EditConditionHides))
 	float SelectionDistance = 0.0f;
-
-	UPROPERTY(EditAnywhere, Category = "Selection")
-	bool bRootSelectsBranch = false;
 };
 
 USTRUCT()
@@ -87,10 +85,10 @@ struct FPVSkeletonEditParams
 		meta=(ShowOnlyInnerProperties, EditCondition="SkeletonEditMode == ESkeletonEditMode::BranchRemoval", EditConditionHides))
 	FPVBranchRemovalParams BranchRemovalParams;
 
-	UPROPERTY(EditAnywhere, EditFixedSize, Category = "Edits")
+	UPROPERTY()
 	TArray<FVector3f> RelativeOffsets;
 
-	UPROPERTY(EditAnywhere, EditFixedSize, Category = "Edits")
+	UPROPERTY()
 	TArray<bool> RemovedPoints;
 };
 
diff --git a/Engine/Plugins/Experimental/ProceduralVegetationEditor/Source/ProceduralVegetationEditor/Private/Tools/NodeTools/PVSkeletonEditTool.cpp b/Engine/Plugins/Experimental/ProceduralVegetationEditor/Source/ProceduralVegetationEditor/Private/Tools/NodeTools/PVSkeletonEditTool.cpp
--- a/Engine/Plugins/Experimental/ProceduralVegetationEditor/Source/ProceduralVegetationEditor/Private/Tools/NodeTools/PVSkeletonEditTool.cpp
+++ b/Engine/Plugins/Experimental/ProceduralVegetationEditor/Source/ProceduralVegetationEditor/Private/Tools/NodeTools/PVSkeletonEditTool.cpp
@@ -34,7 +34,12 @@ void FSkeletonPointsSelection::SelectPoint(const int32 InBranchIndex, const int3
 	SelectedPoints.Empty();
 }
 
-void FSkeletonPointsSelection::ExtendSelection(const int32 NumNeighbours, const EPointSelectionSmoothnessMethod SmoothnessMethod)
+void FSkeletonPointsSelection::ExtendSelection(
+	const ESkeletonSelectionMode SelectionMode,
+	const float ExtensionLimit,
+	const TArray<FVector3f>& Offsets,
+	const EPointSelectionSmoothnessMethod SmoothnessMethod
+)
 {
 	check(Collection);
 
@@ -102,220 +107,214 @@ void FSkeletonPointsSelection::ExtendSelection(const int32 NumNeighbours, const
 		};
 
 	SelectedPoints.Empty();
-	PrimarySelectedPoint.Weight = 1.0f;
+	PrimarySelectedPoint.Weight = 0.0f;
 	SelectedPoints.Add(PrimarySelectedPoint);
 
-	int32 CurrentNeighbour = 0;
-	while (CurrentNeighbour <= NumNeighbours)
+	switch (SelectionMode)
 	{
-		for (int32 SelectedPointIndex = 0; SelectedPointIndex < SelectedPoints.Num(); SelectedPointIndex++)
+	case ESkeletonSelectionMode::SelectByNeighbours:
 		{
-			const FSelectedPoint& SelectedPoint = SelectedPoints[SelectedPointIndex];
-			if (SelectedPoint.NeighbourIndex != CurrentNeighbour)
+			int32 CurrentNeighbour = 0;
+			const int32 MaxNeighbours = static_cast<int32>(ExtensionLimit);
+			while (CurrentNeighbour <= MaxNeighbours)
 			{
-				continue;
-			}
-
-			if (SelectedPoint.BranchPointIndex == 0)
-			{
-				const auto [ParentBranchIndex, ParentBranchPointIndex] = GetParentPoint(SelectedPoint.BranchIndex);
-				if (ParentBranchIndex != INDEX_NONE && ParentBranchPointIndex != INDEX_NONE)
+				for (int32 SelectedPointIndex = 0; SelectedPointIndex < SelectedPoints.Num(); SelectedPointIndex++)
 				{
-					SelectedPoints.AddUnique({ParentBranchIndex, ParentBranchPointIndex, 0.0f, CurrentNeighbour});
-				}
-			}
-			const TArray<TPair<int32, int32>> ChildPoints = GetChildPoints(
-				SelectedPoint.BranchIndex,
-				SelectedPoint.BranchPointIndex,
-				BranchFacade.GetImmediateChildren(SelectedPoint.BranchIndex)
-			);
-			for (const auto [ChildBranchIndex, ChildBranchPointIndex] : ChildPoints)
-			{
-				SelectedPoints.AddUnique({ChildBranchIndex, ChildBranchPointIndex, 0.0f, CurrentNeighbour});
-			}
+					const FSelectedPoint& SelectedPoint = SelectedPoints[SelectedPointIndex];
+					if (SelectedPoint.NeighbourIndex != CurrentNeighbour)
+					{
+						continue;
+					}
 
-			if (CurrentNeighbour < NumNeighbours)
-			{
-				const TArray<int32>& BranchPoints = BranchFacade.GetPoints(SelectedPoint.BranchIndex);
-				if (BranchPoints.IsValidIndex(SelectedPoint.BranchPointIndex - 1))
-				{
-					SelectedPoints.AddUnique({SelectedPoint.BranchIndex, SelectedPoint.BranchPointIndex - 1, 0.0f, CurrentNeighbour + 1});
-				}
-				if (BranchPoints.IsValidIndex(SelectedPoint.BranchPointIndex + 1))
-				{
-					SelectedPoints.AddUnique({SelectedPoint.BranchIndex, SelectedPoint.BranchPointIndex + 1, 0.0f, CurrentNeighbour + 1});
+					if (SelectedPoint.BranchPointIndex == 0)
+					{
+						const auto [ParentBranchIndex, ParentBranchPointIndex] = GetParentPoint(SelectedPoint.BranchIndex);
+						if (ParentBranchIndex != INDEX_NONE && ParentBranchPointIndex != INDEX_NONE)
+						{
+							SelectedPoints.AddUnique({ParentBranchIndex, ParentBranchPointIndex, 0.0f, CurrentNeighbour});
+						}
+					}
+					const TArray<TPair<int32, int32>> ChildPoints = GetChildPoints(
+						SelectedPoint.BranchIndex,
+						SelectedPoint.BranchPointIndex,
+						BranchFacade.GetImmediateChildren(SelectedPoint.BranchIndex)
+					);
+					for (const auto [ChildBranchIndex, ChildBranchPointIndex] : ChildPoints)
+					{
+						SelectedPoints.AddUnique({ChildBranchIndex, ChildBranchPointIndex, 0.0f, CurrentNeighbour});
+					}
+
+					if (CurrentNeighbour < static_cast<int32>(ExtensionLimit))
+					{
+						const TArray<int32>& BranchPoints = BranchFacade.GetPoints(SelectedPoint.BranchIndex);
+						if (BranchPoints.IsValidIndex(SelectedPoint.BranchPointIndex - 1))
+						{
+							SelectedPoints.AddUnique({SelectedPoint.BranchIndex, SelectedPoint.BranchPointIndex - 1, 0.0f, CurrentNeighbour + 1});
+						}
+						if (BranchPoints.IsValidIndex(SelectedPoint.BranchPointIndex + 1))
+						{
+							SelectedPoints.AddUnique({SelectedPoint.BranchIndex, SelectedPoint.BranchPointIndex + 1, 0.0f, CurrentNeighbour + 1});
+						}
+					}
 				}
+
+				++CurrentNeighbour;
 			}
 		}
+		break;
 
-		++CurrentNeighbour;
-	}
-
-	const float MaxNeighbours = NumNeighbours + 1;
-	for (FSelectedPoint& SelectedPoint : SelectedPoints)
-	{
-		SelectedPoint.Weight = GetWeightBySelectedSmoothness(
-			(MaxNeighbours - SelectedPoint.NeighbourIndex) / MaxNeighbours,
-			SmoothnessMethod
-		);
-	}
-}
-
-void FSkeletonPointsSelection::ExtendSelectionByDistance(
-	const float MaxDistance,
-	const TArray<FVector3f>& Offsets,
-	const EPointSelectionSmoothnessMethod SmoothnessMethod
-)
-{
-	check(Collection);
-
-	const PV::Facades::FPointFacade PointFacade(*Collection);
-	const PV::Facades::FBranchFacade BranchFacade(*Collection);
-	if (!PointFacade.IsValid() || !BranchFacade.IsValid())
-	{
-		return;
-	}
-
-	const auto GetParentPoint = [&](const int32 BranchIndex) -> TPair<int32, int32>
+	case ESkeletonSelectionMode::SelectByTreeDistance:
 		{
-			const int32 RootPointIndex = BranchFacade.GetRootPoint(BranchIndex);
-			if (RootPointIndex != INDEX_NONE)
+			int32 CurrentNeighbour = 0;
+			bool bPointAdded = true;
+			while (bPointAdded)
 			{
-				const float PointLFR = PointFacade.GetLengthFromRoot(RootPointIndex);
-				const int32 ParentBranchIndex = BranchFacade.GetParentBranchIndex(BranchIndex);
-				if (ParentBranchIndex != INDEX_NONE)
+				bPointAdded = false;
+				for (int32 SelectedPointIndex = 0; SelectedPointIndex < SelectedPoints.Num(); SelectedPointIndex++)
 				{
-					const TArray<int32>& ParentBranchPoints = BranchFacade.GetPoints(ParentBranchIndex);
-					for (int32 BranchPointIndex = 0; BranchPointIndex < ParentBranchPoints.Num(); BranchPointIndex++)
+					const FSelectedPoint& SelectedPoint = SelectedPoints[SelectedPointIndex];
+					if (SelectedPoint.NeighbourIndex != CurrentNeighbour)
 					{
-						if (PointFacade.GetLengthFromRoot(ParentBranchPoints[BranchPointIndex]) >= PointLFR)
+						continue;
+					}
+
+					if (SelectedPoint.BranchPointIndex == 0)
+					{
+						const auto [ParentBranchIndex, ParentBranchPointIndex] = GetParentPoint(SelectedPoint.BranchIndex);
+						if (ParentBranchIndex != INDEX_NONE && ParentBranchPointIndex != INDEX_NONE)
 						{
-							return {ParentBranchIndex, BranchPointIndex};
+							SelectedPoints.AddUnique({ParentBranchIndex, ParentBranchPointIndex, SelectedPoint.Weight, CurrentNeighbour});
 						}
 					}
-				}
-			}
-			return {INDEX_NONE, INDEX_NONE};
-		};
-
-	const auto GetChildPoints = [&](
-		const int32 BranchIndex,
-		const int32 BranchPointIndex,
-		const TArray<int32>& ImmediateChildBranches
-	) -> TArray<TPair<int32, int32>>
-		{
-			const TArray<int32>& BranchPoints = BranchFacade.GetPoints(BranchIndex);
+					const TArray<TPair<int32, int32>> ChildPoints = GetChildPoints(
+						SelectedPoint.BranchIndex,
+						SelectedPoint.BranchPointIndex,
+						BranchFacade.GetImmediateChildren(SelectedPoint.BranchIndex)
+					);
+					for (const auto [ChildBranchIndex, ChildBranchPointIndex] : ChildPoints)
+					{
+						SelectedPoints.AddUnique({ChildBranchIndex, ChildBranchPointIndex, SelectedPoint.Weight, CurrentNeighbour});
+					}
 
-			const float PointLFR = PointFacade.GetLengthFromRoot(BranchPoints[BranchPointIndex]);
-			const float NextPointLFR = BranchPoints.IsValidIndex(BranchPointIndex + 1)
-				? PointFacade.GetLengthFromRoot(BranchPoints[BranchPointIndex + 1])
-				: UE_MAX_FLT;
+					const TArray<int32>& BranchPoints = BranchFacade.GetPoints(SelectedPoint.BranchIndex);
+					const int32 SelectedIndex = BranchPoints[SelectedPoint.BranchPointIndex];
+					const int32 SelectedBudNumber = PointFacade.GetBudNumber(SelectedIndex);
+					const FVector3f SelectedPointPosition = PointFacade.GetPosition(SelectedIndex) + Offsets[SelectedBudNumber];
 
-			TArray<TPair<int32, int32>> ChildPoints;
-			for (int32 ChildBranch : ImmediateChildBranches)
-			{
-				const int32 ChildBranchIndex = BranchFacade.GetBranchNumbers().Find(ChildBranch);
-				if (ChildBranchIndex != INDEX_NONE)
-				{
-					const int32 ChildBranchRootPointIndex = BranchFacade.GetRootPoint(ChildBranchIndex);
-					if (ChildBranchRootPointIndex != INDEX_NONE)
+					if (BranchPoints.IsValidIndex(SelectedPoint.BranchPointIndex - 1))
 					{
-						const float ChildLFR = PointFacade.GetLengthFromRoot(ChildBranchRootPointIndex);
-						if (ChildLFR >= PointLFR && ChildLFR < NextPointLFR)
+						const int32 PointIndex = BranchPoints[SelectedPoint.BranchPointIndex - 1];
+						const int32 PointBudNumber = PointFacade.GetBudNumber(PointIndex);
+						const FVector3f PointPosition = PointFacade.GetPosition(PointIndex) + Offsets[PointBudNumber];
+						const float PointDistance = SelectedPoint.Weight + FVector3f::Distance(SelectedPointPosition, PointPosition);
+						if (PointDistance < ExtensionLimit)
 						{
-							ChildPoints.Emplace(ChildBranchIndex, 0);
+							SelectedPoints.AddUnique({
+								SelectedPoint.BranchIndex, SelectedPoint.BranchPointIndex - 1, PointDistance, CurrentNeighbour + 1
+							});
+							bPointAdded = true;
+						}
+					}
+					if (BranchPoints.IsValidIndex(SelectedPoint.BranchPointIndex + 1))
+					{
+						const int32 PointIndex = BranchPoints[SelectedPoint.BranchPointIndex + 1];
+						const int32 PointBudNumber = PointFacade.GetBudNumber(PointIndex);
+						const FVector3f PointPosition = PointFacade.GetPosition(PointIndex) + Offsets[PointBudNumber];
+						const float PointDistance = SelectedPoint.Weight + FVector3f::Distance(SelectedPointPosition, PointPosition);
+						if (PointDistance < ExtensionLimit)
+						{
+							SelectedPoints.AddUnique({
+								SelectedPoint.BranchIndex, SelectedPoint.BranchPointIndex + 1, PointDistance, CurrentNeighbour + 1
+							});
+							bPointAdded = true;
 						}
 					}
 				}
-			}
 
-			return ChildPoints;
-		};
-
-	SelectedPoints.Empty();
-	PrimarySelectedPoint.Weight = 0.0f;
-	SelectedPoints.Add(PrimarySelectedPoint);
-
-	int32 CurrentNeighbour = 0;
-	bool bPointAdded = true;
-	while (bPointAdded)
-	{
-		bPointAdded = false;
-		for (int32 SelectedPointIndex = 0; SelectedPointIndex < SelectedPoints.Num(); SelectedPointIndex++)
-		{
-			const FSelectedPoint& SelectedPoint = SelectedPoints[SelectedPointIndex];
-			if (SelectedPoint.NeighbourIndex != CurrentNeighbour)
-			{
-				continue;
+				++CurrentNeighbour;
 			}
+		}
+		break;
 
-			if (SelectedPoint.BranchPointIndex == 0)
+	case ESkeletonSelectionMode::SelectByEuclideanDistance:
+		{
+			// Get the primary selected point's position
+			const TArray<int32>& PrimaryBranchPoints = BranchFacade.GetPoints(PrimarySelectedPoint.BranchIndex);
+			const int32 PrimaryPointIndex = PrimaryBranchPoints[PrimarySelectedPoint.BranchPointIndex];
+			const int32 PrimaryBudNumber = PointFacade.GetBudNumber(PrimaryPointIndex);
+			const FVector3f PrimaryPointPosition = PointFacade.GetPosition(PrimaryPointIndex) + Offsets[PrimaryBudNumber];
+
+			// Loop through all branches and their points
+			for (int32 BranchIndex = 0; BranchIndex < BranchFacade.GetElementCount(); BranchIndex++)
 			{
-				const auto [ParentBranchIndex, ParentBranchPointIndex] = GetParentPoint(SelectedPoint.BranchIndex);
-				if (ParentBranchIndex != INDEX_NONE && ParentBranchPointIndex != INDEX_NONE)
+				const TArray<int32>& BranchPoints = BranchFacade.GetPoints(BranchIndex);
+				for (int32 BranchPointIndex = 0; BranchPointIndex < BranchPoints.Num(); BranchPointIndex++)
 				{
-					SelectedPoints.AddUnique({ParentBranchIndex, ParentBranchPointIndex, SelectedPoint.Weight, CurrentNeighbour});
+					// Skip the primary selected point as it's already added
+					if (BranchIndex == PrimarySelectedPoint.BranchIndex && BranchPointIndex == PrimarySelectedPoint.BranchPointIndex)
+					{
+						continue;
+					}
+
+					const int32 PointIndex = BranchPoints[BranchPointIndex];
+					const int32 PointBudNumber = PointFacade.GetBudNumber(PointIndex);
+					const FVector3f PointPosition = PointFacade.GetPosition(PointIndex) + Offsets[PointBudNumber];
+
+					// Calculate Euclidean distance from primary selected point
+					const float PointDistance = FVector3f::Distance(PrimaryPointPosition, PointPosition);
+
+					if (PointDistance <= ExtensionLimit)
+					{
+						SelectedPoints.Add({BranchIndex, BranchPointIndex, PointDistance, 0});
+					}
 				}
 			}
-			const TArray<TPair<int32, int32>> ChildPoints = GetChildPoints(
-				SelectedPoint.BranchIndex,
-				SelectedPoint.BranchPointIndex,
-				BranchFacade.GetImmediateChildren(SelectedPoint.BranchIndex)
-			);
-			for (const auto [ChildBranchIndex, ChildBranchPointIndex] : ChildPoints)
-			{
-				SelectedPoints.AddUnique({ChildBranchIndex, ChildBranchPointIndex, SelectedPoint.Weight, CurrentNeighbour});
-			}
+		}
+		break;
+
+	default:
+		break;
+	}
 
-			const TArray<int32>& BranchPoints = BranchFacade.GetPoints(SelectedPoint.BranchIndex);
-			const int32 SelectedIndex = BranchPoints[SelectedPoint.BranchPointIndex];
-			const int32 SelectedBudNumber = PointFacade.GetBudNumber(SelectedIndex);
-			const FVector3f SelectedPointPosition = PointFacade.GetPosition(SelectedIndex) + Offsets[SelectedBudNumber];
+	ApplySelectionWeights(SelectionMode, ExtensionLimit, SmoothnessMethod);
+}
 
-			if (BranchPoints.IsValidIndex(SelectedPoint.BranchPointIndex - 1))
+void FSkeletonPointsSelection::ApplySelectionWeights(
+	const ESkeletonSelectionMode SelectionMode,
+	const float ExtensionLimit,
+	const EPointSelectionSmoothnessMethod SmoothnessMethod
+)
+{
+	switch (SelectionMode)
+	{
+	case ESkeletonSelectionMode::SelectByNeighbours:
+		{
+			const float MaxNeighbours = static_cast<int32>(ExtensionLimit) + 1;
+			for (FSelectedPoint& SelectedPoint : SelectedPoints)
 			{
-				const int32 PointIndex = BranchPoints[SelectedPoint.BranchPointIndex - 1];
-				const int32 PointBudNumber = PointFacade.GetBudNumber(PointIndex);
-				const FVector3f PointPosition = PointFacade.GetPosition(PointIndex) + Offsets[PointBudNumber];
-				const float PointDistance = SelectedPoint.Weight + FVector3f::Distance(SelectedPointPosition, PointPosition);
-				if (PointDistance < MaxDistance)
-				{
-					SelectedPoints.AddUnique({
-						SelectedPoint.BranchIndex, SelectedPoint.BranchPointIndex - 1, PointDistance, CurrentNeighbour + 1
-					});
-					bPointAdded = true;
-				}
+				SelectedPoint.Weight = GetWeightBySelectedSmoothness(
+					(MaxNeighbours - SelectedPoint.NeighbourIndex) / MaxNeighbours,
+					SmoothnessMethod
+				);
 			}
-			if (BranchPoints.IsValidIndex(SelectedPoint.BranchPointIndex + 1))
+		}
+		break;
+
+	case ESkeletonSelectionMode::SelectByTreeDistance:
+	case ESkeletonSelectionMode::SelectByEuclideanDistance:
+		{
+			for (FSelectedPoint& SelectedPoint : SelectedPoints)
 			{
-				const int32 PointIndex = BranchPoints[SelectedPoint.BranchPointIndex + 1];
-				const int32 PointBudNumber = PointFacade.GetBudNumber(PointIndex);
-				const FVector3f PointPosition = PointFacade.GetPosition(PointIndex) + Offsets[PointBudNumber];
-				const float PointDistance = SelectedPoint.Weight + FVector3f::Distance(SelectedPointPosition, PointPosition);
-				if (PointDistance < MaxDistance)
-				{
-					SelectedPoints.AddUnique({
-						SelectedPoint.BranchIndex, SelectedPoint.BranchPointIndex + 1, PointDistance, CurrentNeighbour + 1
-					});
-					bPointAdded = true;
-				}
+				SelectedPoint.Weight = GetWeightBySelectedSmoothness(
+					1.0f - (ExtensionLimit != 0.0f ? SelectedPoint.Weight / ExtensionLimit : 0.0f),
+					SmoothnessMethod
+				);
 			}
 		}
+		break;
 
-		++CurrentNeighbour;
-	}
-
-	const float MaxNeighbours = Algo::MaxElementBy(
-		SelectedPoints, [](const FSelectedPoint& Point)
-			{
-				return Point.NeighbourIndex;
-			})->NeighbourIndex + 1;
-	for (FSelectedPoint& SelectedPoint : SelectedPoints)
-	{
-		SelectedPoint.Weight = GetWeightBySelectedSmoothness(
-			(MaxNeighbours - SelectedPoint.NeighbourIndex) / MaxNeighbours,
-			SmoothnessMethod
-		);
+	default:
+		break;
 	}
 }
 
@@ -434,21 +433,14 @@ void UPVSkeletonEditTool::OnNodeSettingsUpdated()
 		const TObjectPtr<UPVSkeletonEditSettings> SkeletonEditSettings = GetNodeSettings<UPVSkeletonEditSettings>();
 		const FPVSkeletonSelectionParams& SelectionParams = SkeletonEditSettings->SkeletonEditSettings.SkeletonSelectionParams;
 
-		if (SelectionParams.SkeletonSelectionMode == ESkeletonSelectionMode::SelectByDistance)
-		{
-			Selection->ExtendSelectionByDistance(
-				SelectionParams.SelectionDistance,
-				SkeletonEditSettings->SkeletonEditSettings.RelativeOffsets,
-				SelectionParams.SelectionFalloff
-			);
-		}
-		else
-		{
-			Selection->ExtendSelection(
-				SelectionParams.NumPointsSelection,
-				SelectionParams.SelectionFalloff
-			);
-		}
+		Selection->ExtendSelection(
+			SelectionParams.SkeletonSelectionMode,
+			SelectionParams.SkeletonSelectionMode == ESkeletonSelectionMode::SelectByNeighbours
+			? static_cast<float>(SelectionParams.NumPointsSelection)
+			: SelectionParams.SelectionDistance,
+			SkeletonEditSettings->SkeletonEditSettings.RelativeOffsets,
+			SelectionParams.SelectionFalloff
+		);
 
 		const PV::Facades::FPointFacade PointFacade(GetCollection());
 		const PV::Facades::FBranchFacade BranchFacade(GetCollection());
@@ -470,168 +462,143 @@ void UPVSkeletonEditTool::OnNodeSettingsUpdated()
 
 FInputRayHit UPVSkeletonEditTool::IsHitByClick(const FInputDeviceRay& ClickPos)
 {
-	return FInputRayHit(1.0f);
+	if (FViewport* FocusedViewport = GetToolManager()->GetContextQueriesAPI()->GetFocusedViewport())
+	{
+		if (FocusedViewport->GetHitProxy(ClickPos.ScreenPosition.X, ClickPos.ScreenPosition.Y))
+		{
+			return FInputRayHit(1.0f);
+		}
+	}
+
+	Selection->ClearSelection();
+	CombinedGizmo->SetVisibility(false);
+	return FInputRayHit();
 }
 
 void UPVSkeletonEditTool::OnClicked(const FInputDeviceRay& ClickPos)
 {
 	FViewport* FocusedViewport = GetToolManager()->GetContextQueriesAPI()->GetFocusedViewport();
-	const HHitProxy* HitProxy = FocusedViewport ? FocusedViewport->GetHitProxy(ClickPos.ScreenPosition.X, ClickPos.ScreenPosition.Y) : nullptr;
-	if (!HitProxy || !HitProxy->IsA(HPVSkeletonPointProxy::StaticGetType()))
+	if (!FocusedViewport)
+	{
+		return;
+	}
+	const HHitProxy* HitProxy = FocusedViewport->GetHitProxy(ClickPos.ScreenPosition.X, ClickPos.ScreenPosition.Y);
+	if (!HitProxy)
 	{
-		PreviousRotation = FQuat4f::Identity;
-		Selection->ClearSelection();
-		CombinedGizmo->SetVisibility(false);
 		return;
 	}
 
-	const PV::Facades::FPointFacade PointFacade(GetCollection());
-	const PV::Facades::FBranchFacade BranchFacade(GetCollection());
-
-	const auto GetParentPoint = [&](const int32 BranchIndex) -> TPair<int32, int32>
-		{
-			const int32 RootPointIndex = BranchFacade.GetRootPoint(BranchIndex);
-			if (RootPointIndex != INDEX_NONE)
-			{
-				const float PointLFR = PointFacade.GetLengthFromRoot(RootPointIndex);
-				const int32 ParentBranchIndex = BranchFacade.GetParentBranchIndex(BranchIndex);
-				if (ParentBranchIndex != INDEX_NONE)
-				{
-					const TArray<int32>& ParentBranchPoints = BranchFacade.GetPoints(ParentBranchIndex);
-					for (int32 BranchPointIndex = 0; BranchPointIndex < ParentBranchPoints.Num(); BranchPointIndex++)
-					{
-						if (PointFacade.GetLengthFromRoot(ParentBranchPoints[BranchPointIndex]) >= PointLFR)
-						{
-							return {ParentBranchIndex, BranchPointIndex};
-						}
-					}
-				}
-			}
-			return {INDEX_NONE, INDEX_NONE};
-		};
-
-	const HPVSkeletonPointProxy* SkeletonPointProxy = static_cast<const HPVSkeletonPointProxy*>(HitProxy);
+	if (HitProxy->IsA(HPVSkeletonPointProxy::StaticGetType()))
+	{
+		const PV::Facades::FPointFacade PointFacade(GetCollection());
+		const PV::Facades::FBranchFacade BranchFacade(GetCollection());
 
-	const TObjectPtr<UPVSkeletonEditSettings> SkeletonEditSettings = GetNodeSettings<UPVSkeletonEditSettings>();
+		const HPVSkeletonPointProxy* SkeletonPointProxy = static_cast<const HPVSkeletonPointProxy*>(HitProxy);
 
-	if (SkeletonEditSettings->SkeletonEditSettings.SkeletonEditMode == ESkeletonEditMode::TRS)
-	{
-		const FPVSkeletonSelectionParams& SelectionParams = SkeletonEditSettings->SkeletonEditSettings.SkeletonSelectionParams;
+		const TObjectPtr<UPVSkeletonEditSettings> SkeletonEditSettings = GetNodeSettings<UPVSkeletonEditSettings>();
 
-		if (!SelectionParams.bRootSelectsBranch && SkeletonPointProxy->BranchPointIndex == 0)
-		{
-			const auto [ParentBranchIndex, ParentBranchPointIndex] = GetParentPoint(SkeletonPointProxy->BranchIndex);
-			if (ParentBranchIndex == INDEX_NONE || ParentBranchPointIndex == INDEX_NONE)
-			{
-				Selection->SelectPoint(SkeletonPointProxy->BranchIndex, SkeletonPointProxy->BranchPointIndex);
-			}
-			else
-			{
-				Selection->SelectPoint(ParentBranchIndex, ParentBranchPointIndex);
-			}
-		}
-		else
+		if (SkeletonEditSettings->SkeletonEditSettings.SkeletonEditMode == ESkeletonEditMode::TRS)
 		{
+			const FPVSkeletonSelectionParams& SelectionParams = SkeletonEditSettings->SkeletonEditSettings.SkeletonSelectionParams;
+
 			Selection->SelectPoint(SkeletonPointProxy->BranchIndex, SkeletonPointProxy->BranchPointIndex);
-		}
 
-		if (SelectionParams.SkeletonSelectionMode == ESkeletonSelectionMode::SelectByDistance)
-		{
-			Selection->ExtendSelectionByDistance(
-				SelectionParams.SelectionDistance,
+			Selection->ExtendSelection(
+				SelectionParams.SkeletonSelectionMode,
+				SelectionParams.SkeletonSelectionMode == ESkeletonSelectionMode::SelectByNeighbours
+				? static_cast<float>(SelectionParams.NumPointsSelection)
+				: SelectionParams.SelectionDistance,
 				SkeletonEditSettings->SkeletonEditSettings.RelativeOffsets,
 				SelectionParams.SelectionFalloff
 			);
+
+			TArray<FVector3f>& Offsets = SkeletonEditSettings->SkeletonEditSettings.RelativeOffsets;
+
+			const int32 SelectedPointIndex = BranchFacade.GetPoints(SkeletonPointProxy->BranchIndex)[SkeletonPointProxy->BranchPointIndex];
+			const int32 SelectedPointBudNumber = PointFacade.GetBudNumber(SelectedPointIndex);
+			const FVector Position = static_cast<FVector>(PointFacade.GetPosition(SelectedPointIndex) + Offsets[SelectedPointBudNumber]);
+			PreviousRotation = FQuat4f::Identity;
+			CombinedGizmo->SetVisibility(true);
+			CombinedGizmo->ReinitializeGizmoTransform(FTransform(Position));
 		}
-		else
+		else if (SkeletonEditSettings->SkeletonEditSettings.SkeletonEditMode == ESkeletonEditMode::BranchRemoval)
 		{
-			Selection->ExtendSelection(
-				SelectionParams.NumPointsSelection,
-				SelectionParams.SelectionFalloff
-			);
-		}
-
-		TArray<FVector3f>& Offsets = SkeletonEditSettings->SkeletonEditSettings.RelativeOffsets;
+			PreviousRotation = FQuat4f::Identity;
+			Selection->ClearSelection();
+			CombinedGizmo->SetVisibility(false);
 
-		const int32 SelectedPointIndex = BranchFacade.GetPoints(SkeletonPointProxy->BranchIndex)[SkeletonPointProxy->BranchPointIndex];
-		const int32 SelectedPointBudNumber = PointFacade.GetBudNumber(SelectedPointIndex);
-		const FVector Position = static_cast<FVector>(PointFacade.GetPosition(SelectedPointIndex) + Offsets[SelectedPointBudNumber]);
-		PreviousRotation = FQuat4f::Identity;
-		CombinedGizmo->SetVisibility(true);
-		CombinedGizmo->ReinitializeGizmoTransform(FTransform(Position));
-	}
-	else if (SkeletonEditSettings->SkeletonEditSettings.SkeletonEditMode == ESkeletonEditMode::BranchRemoval)
-	{
-		PreviousRotation = FQuat4f::Identity;
-		Selection->ClearSelection();
-		CombinedGizmo->SetVisibility(false);
+			TArray<bool>& RemovedPoints = SkeletonEditSettings->SkeletonEditSettings.RemovedPoints;
 
-		TArray<bool>& RemovedPoints = SkeletonEditSettings->SkeletonEditSettings.RemovedPoints;
+			const TArray<int32>& BranchPoints = BranchFacade.GetPoints(SkeletonPointProxy->BranchIndex);
 
-		const TArray<int32>& BranchPoints = BranchFacade.GetPoints(SkeletonPointProxy->BranchIndex);
+			int32 StartIndex =
+				SkeletonEditSettings->SkeletonEditSettings.BranchRemovalParams.BranchRemovalMode == EBranchRemovalMode::TrimBranch
+				? SkeletonPointProxy->BranchPointIndex
+				: 0;
 
-		int32 StartIndex =
-			SkeletonEditSettings->SkeletonEditSettings.BranchRemovalParams.BranchRemovalMode == EBranchRemovalMode::TrimBranch
-			? SkeletonPointProxy->BranchPointIndex
-			: 0;
+			bool bNextPointRemoved = false;
+			if (BranchPoints.IsValidIndex(StartIndex + 1))
+			{
+				const int32 NextPointBudNumber = PointFacade.GetBudNumber(BranchPoints[StartIndex + 1]);
+				bNextPointRemoved = RemovedPoints[NextPointBudNumber];
+			}
+			if (StartIndex > 0 && StartIndex < BranchPoints.Num() - 1 && !bNextPointRemoved)
+			{
+				// Remove from the next point onwards from currently selected point
+				// unless it's the first or the last point, we want to remove the first or last point if selected
+				StartIndex += 1;
+			}
+			if (StartIndex <= 1)
+			{
+				StartIndex = 1;
+			}
 
-		bool bNextPointRemoved = false;
-		if (BranchPoints.IsValidIndex(StartIndex + 1))
-		{
-			const int32 NextPointBudNumber = PointFacade.GetBudNumber(BranchPoints[StartIndex + 1]);
-			bNextPointRemoved = RemovedPoints[NextPointBudNumber];
-		}
-		if (StartIndex > 0 && StartIndex < BranchPoints.Num() - 1 && !bNextPointRemoved)
-		{
-			// Remove from the next point onwards from currently selected point
-			// unless it's the first or the last point, we want to remove the first or last point if selected
-			StartIndex += 1;
-		}
-		if (StartIndex <= 1)
-		{
-			// Remove the complete branch if first point was selected (can't have a branch with 1 point)
-			StartIndex = 0;
-		}
-			
-		for (int32 BranchPointIndex = StartIndex; BranchPointIndex < BranchPoints.Num(); BranchPointIndex++)
-		{
-			const int32 PointIndex = BranchPoints[BranchPointIndex];
-			const int32 PointBudNumber = PointFacade.GetBudNumber(PointIndex);
-			RemovedPoints[PointBudNumber] = true;
-		}
+			for (int32 BranchPointIndex = StartIndex; BranchPointIndex < BranchPoints.Num(); BranchPointIndex++)
+			{
+				const int32 PointIndex = BranchPoints[BranchPointIndex];
+				const int32 PointBudNumber = PointFacade.GetBudNumber(PointIndex);
+				RemovedPoints[PointBudNumber] = true;
+			}
 
-		const int32 RemovedPointIndex = BranchPoints[StartIndex];
-		const float PointLFR = PointFacade.GetLengthFromRoot(RemovedPointIndex);
-		const TManagedArray<int32>& BranchNumbers = BranchFacade.GetBranchNumbers();
-		for (int32 BranchChildNumber : BranchFacade.GetImmediateChildren(SkeletonPointProxy->BranchIndex))
-		{
-			if (const int32 BranchIndex = BranchNumbers.Find(BranchChildNumber); BranchIndex != INDEX_NONE)
+			const int32 RemovedPointIndex = BranchPoints[StartIndex];
+			const float PointLFR = PointFacade.GetLengthFromRoot(RemovedPointIndex);
+			const TManagedArray<int32>& BranchNumbers = BranchFacade.GetBranchNumbers();
+			for (int32 BranchChildNumber : BranchFacade.GetImmediateChildren(SkeletonPointProxy->BranchIndex))
 			{
-				const TArray<int32>& ChildBranchPoints = BranchFacade.GetPoints(BranchIndex);
-				if (PointFacade.GetLengthFromRoot(ChildBranchPoints[0]) >= PointLFR)
+				if (const int32 BranchIndex = BranchNumbers.Find(BranchChildNumber); BranchIndex != INDEX_NONE)
 				{
-					for (const int32 ChildBranchPointIndex : ChildBranchPoints)
+					const TArray<int32>& ChildBranchPoints = BranchFacade.GetPoints(BranchIndex);
+					if (PointFacade.GetLengthFromRoot(ChildBranchPoints[0]) >= PointLFR)
 					{
-						const int32 ChildPointBudNumber = PointFacade.GetBudNumber(ChildBranchPointIndex);
-						RemovedPoints[ChildPointBudNumber] = true;
-					}
+						for (const int32 ChildBranchPointIndex : ChildBranchPoints)
+						{
+							const int32 ChildPointBudNumber = PointFacade.GetBudNumber(ChildBranchPointIndex);
+							RemovedPoints[ChildPointBudNumber] = true;
+						}
 
-					for (int32 SubBranchChildNumber : BranchFacade.GetChildren(BranchIndex))
-					{
-						if (const int32 SubBranchIndex = BranchNumbers.Find(SubBranchChildNumber); SubBranchIndex != INDEX_NONE)
+						for (int32 SubBranchChildNumber : BranchFacade.GetChildren(BranchIndex))
 						{
-							for (const int32 ChildSubBranchPointIndex : BranchFacade.GetPoints(SubBranchIndex))
+							if (const int32 SubBranchIndex = BranchNumbers.Find(SubBranchChildNumber); SubBranchIndex != INDEX_NONE)
 							{
-								const int32 SubBranchChildPointBudNumber = PointFacade.GetBudNumber(ChildSubBranchPointIndex);
-								RemovedPoints[SubBranchChildPointBudNumber] = true;
+								for (const int32 ChildSubBranchPointIndex : BranchFacade.GetPoints(SubBranchIndex))
+								{
+									const int32 SubBranchChildPointBudNumber = PointFacade.GetBudNumber(ChildSubBranchPointIndex);
+									RemovedPoints[SubBranchChildPointBudNumber] = true;
+								}
 							}
 						}
 					}
 				}
 			}
-		}
 
-		GetNodeSettings()->PostEditChange();
+			GetNodeSettings()->PostEditChange();
+		}
+	}
+	else
+	{
+		PreviousRotation = FQuat4f::Identity;
+		Selection->ClearSelection();
+		CombinedGizmo->SetVisibility(false);
 	}
 }
 
@@ -663,6 +630,13 @@ void UPVSkeletonEditTool::Render(IToolsContextRenderAPI* RenderAPI)
 			{
 				break;
 			}
+			if (BranchPointIndex == 0 && BranchPoints.IsValidIndex(1))
+			{
+				if (RemovedPoints[PointFacade.GetBudNumber(BranchPoints[1])])
+				{
+					break;
+				}
+			}
 			const int32 BudNumber = PointFacade.GetBudNumber(PointIndex);
 			const FVector CurrentPosition = static_cast<FVector>(Positions[PointIndex] + Offsets[BudNumber]);
 			FLinearColor DrawColor = FLinearColor::White;
@@ -709,6 +683,40 @@ void UPVSkeletonEditTool::Render(IToolsContextRenderAPI* RenderAPI)
 			PreviousSelectedPointIndex = SelectedPointIndex;
 		}
 	}
+	
+	if (SkeletonEditSettings->SkeletonEditSettings.SkeletonSelectionParams.SkeletonSelectionMode == ESkeletonSelectionMode::SelectByEuclideanDistance)
+	{
+		const int32 Radius = SkeletonEditSettings->SkeletonEditSettings.SkeletonSelectionParams.SelectionDistance;
+		if (Selection->HasSelection())
+		{
+			// Get the primary selected point position
+			const int32 SelectedBranchIndex = Selection->GetSelectedBranchIndex();
+			const int32 SelectedBranchPointIndex = Selection->GetSelectedBranchPointIndex();
+
+			if (SelectedBranchIndex != INDEX_NONE && SelectedBranchPointIndex != INDEX_NONE)
+			{
+				const TArray<int32>& BranchPoints = BranchFacade.GetPoints(SelectedBranchIndex);
+				if (BranchPoints.IsValidIndex(SelectedBranchPointIndex))
+				{
+					const int32 PointIndex = BranchPoints[SelectedBranchPointIndex];
+					const int32 BudNumber = PointFacade.GetBudNumber(PointIndex);
+					const FVector SelectedPosition = static_cast<FVector>(Positions[PointIndex] + Offsets[BudNumber]);
+
+					constexpr int32 NumSides = 32;
+					constexpr float Thickness = 0.0f;
+					const FLinearColor DrawColor = FLinearColor::Yellow;
+
+					FTransform Transform = FTransform(FRotator(), SelectedPosition);
+					DrawCircle(PDI, Transform.GetLocation(), Transform.GetScaledAxis(EAxis::X), Transform.GetScaledAxis(EAxis::Y), DrawColor, Radius, NumSides, 0, Thickness, 0.001f, false);
+					DrawCircle(PDI, Transform.GetLocation(), Transform.GetScaledAxis(EAxis::X), Transform.GetScaledAxis(EAxis::Z), DrawColor, Radius, NumSides, 0, Thickness, 0.001f, false);
+					DrawCircle(PDI, Transform.GetLocation(), Transform.GetScaledAxis(EAxis::Y), Transform.GetScaledAxis(EAxis::Z), DrawColor, Radius, NumSides, 0, Thickness, 0.001f, false);
+					Transform = FTransform(FRotator(0, 45, 0), SelectedPosition);
+					DrawCircle(PDI, Transform.GetLocation(), Transform.GetScaledAxis(EAxis::X), Transform.GetScaledAxis(EAxis::Z), DrawColor, Radius, NumSides, 0, Thickness, 0.001f, false);
+					DrawCircle(PDI, Transform.GetLocation(), Transform.GetScaledAxis(EAxis::Y), Transform.GetScaledAxis(EAxis::Z), DrawColor, Radius, NumSides, 0, Thickness, 0.001f, false);
+				}
+			}
+		}
+	}
 }
 
 void UPVSkeletonEditTool::OnGizmoTransformBegin(UTransformProxy* InTransformProxy)
@@ -741,6 +749,13 @@ void UPVSkeletonEditTool::OnGizmoTransformChanged(UTransformProxy* InTransformPr
 		{
 			for (const FSkeletonPointsSelection::FSelectedPoint& SelectedPoint : Selection->SelectedPoints)
 			{
+				// Skip the root point of a branch (index 0) as it's shared with the parent branch's point
+				// and will be transformed via the parent branch
+				if (!BranchFacade.IsTrunk(SelectedPoint.BranchIndex) && SelectedPoint.BranchPointIndex == 0)
+				{
+					continue;
+				}
+
 				const TArray<int32>& BranchPoints = BranchFacade.GetPoints(SelectedPoint.BranchIndex);
 				const int32 SelectedPointIndex = BranchPoints[SelectedPoint.BranchPointIndex];
 				const int32 SelectedPointBudNumber = PointFacade.GetBudNumber(SelectedPointIndex);
@@ -772,8 +787,12 @@ void UPVSkeletonEditTool::OnGizmoTransformChanged(UTransformProxy* InTransformPr
 					const TArray<int32>& ChildBranchPoints = BranchFacade.GetPoints(BranchIndex);
 					if (PointFacade.GetLengthFromRoot(ChildBranchPoints[0]) > PointLFR)
 					{
-						for (int32 ChildBranchPointIndex : ChildBranchPoints)
+						// Start from index 1 to skip the root point (index 0) of the child branch,
+						// as it's shared with the parent branch's point
+						for (int32 ChildBranchPointIndexInArray = 1; ChildBranchPointIndexInArray < ChildBranchPoints.Num();
+						     ChildBranchPointIndexInArray++)
 						{
+							const int32 ChildBranchPointIndex = ChildBranchPoints[ChildBranchPointIndexInArray];
 							const int32 ChildPointBudNumber = PointFacade.GetBudNumber(ChildBranchPointIndex);
 							const FVector3f PointPosition = PointFacade.GetPosition(ChildBranchPointIndex) + Offsets[ChildPointBudNumber];
 							const FVector3f UpdatedPosition = OriginalPosition + DeltaRotation.RotateVector(PointPosition - OriginalPosition);
@@ -785,8 +804,13 @@ void UPVSkeletonEditTool::OnGizmoTransformChanged(UTransformProxy* InTransformPr
 						{
 							if (const int32 SubBranchIndex = BranchNumbers.Find(SubBranchChildNumber); SubBranchIndex != INDEX_NONE)
 							{
-								for (int32 ChildSubBranchPointIndex : BranchFacade.GetPoints(SubBranchIndex))
+								const TArray<int32>& SubBranchPoints = BranchFacade.GetPoints(SubBranchIndex);
+								// Start from index 1 to skip the root point (index 0) of the sub-branch,
+								// as it's shared with the parent branch's point
+								for (int32 SubBranchPointIndexInArray = 1; SubBranchPointIndexInArray < SubBranchPoints.Num();
+								     SubBranchPointIndexInArray++)
 								{
+									const int32 ChildSubBranchPointIndex = SubBranchPoints[SubBranchPointIndexInArray];
 									const int32 SubPointBudNumber = PointFacade.GetBudNumber(ChildSubBranchPointIndex);
 									const FVector3f PointPosition = PointFacade.GetPosition(ChildSubBranchPointIndex) + Offsets[SubPointBudNumber];
 									const FVector3f UpdatedPosition = OriginalPosition + DeltaRotation.RotateVector(PointPosition - OriginalPosition);
diff --git a/Engine/Plugins/Experimental/ProceduralVegetationEditor/Source/ProceduralVegetationEditor/Private/Tools/NodeTools/PVSkeletonEditTool.h b/Engine/Plugins/Experimental/ProceduralVegetationEditor/Source/ProceduralVegetationEditor/Private/Tools/NodeTools/PVSkeletonEditTool.h
--- a/Engine/Plugins/Experimental/ProceduralVegetationEditor/Source/ProceduralVegetationEditor/Private/Tools/NodeTools/PVSkeletonEditTool.h
+++ b/Engine/Plugins/Experimental/ProceduralVegetationEditor/Source/ProceduralVegetationEditor/Private/Tools/NodeTools/PVSkeletonEditTool.h
@@ -35,19 +35,30 @@ struct HPVSkeletonPointProxy : HHitProxy
 };
 
 
-
 struct FSkeletonPointsSelection
 {
 	FSkeletonPointsSelection(const FManagedArrayCollection* const InCollection);
 
 	void SelectPoint(int32 InBranchIndex, int32 InBranchPointIndex);
-	void ExtendSelection(const int32 NumNeighbours, const EPointSelectionSmoothnessMethod SmoothnessMethod = EPointSelectionSmoothnessMethod::Linear);
-	void ExtendSelectionByDistance(const float MaxDistance, const TArray<FVector3f>& Offsets, const EPointSelectionSmoothnessMethod SmoothnessMethod = EPointSelectionSmoothnessMethod::Linear);
+	void ExtendSelection(
+		const ESkeletonSelectionMode SelectionMode,
+		const float ExtensionLimit,
+		const TArray<FVector3f>& Offsets,
+		const EPointSelectionSmoothnessMethod SmoothnessMethod = EPointSelectionSmoothnessMethod::Linear
+	);
+	void ApplySelectionWeights(
+		const ESkeletonSelectionMode SelectionMode,
+		const float ExtensionLimit,
+		const EPointSelectionSmoothnessMethod SmoothnessMethod
+	);
 	void ClearSelection();
 
 	bool HasSelection();
-	
-	float GetWeightBySelectedSmoothness(float Weight, const EPointSelectionSmoothnessMethod SmoothnessMethod = EPointSelectionSmoothnessMethod::Linear);
+
+	float GetWeightBySelectedSmoothness(
+		float Weight,
+		const EPointSelectionSmoothnessMethod SmoothnessMethod = EPointSelectionSmoothnessMethod::Linear
+	);
 
 	int32 GetSelectedBranchIndex()
 	{
@@ -91,14 +102,14 @@ public:
 
 	virtual void Setup() override;
 	virtual void Shutdown(EToolShutdownType ShutdownType) override;
-	
+
 	virtual void OnNodeSettingsUpdated() override;
 
 	virtual FInputRayHit IsHitByClick(const FInputDeviceRay& ClickPos) override;
 	virtual void OnClicked(const FInputDeviceRay& ClickPos) override;
 
 	virtual void Render(IToolsContextRenderAPI* RenderAPI) override;
-	
+
 private:
 	void OnGizmoTransformBegin(UTransformProxy* TransformProxy);
 	void OnGizmoTransformChanged(UTransformProxy* TransformProxy, FTransform Transform);
@@ -113,6 +124,6 @@ protected:
 
 private:
 	TUniquePtr<FSkeletonPointsSelection> Selection;
-	
+
 	FQuat4f PreviousRotation = FQuat4f::Identity;
 };
