[PVE] Fix ini error

#Synced-CL 52095318

diff --git a/Engine/Plugins/Experimental/ProceduralVegetationEditor/Config/BaseProceduralVegetationEditor.ini b/Engine/Plugins/Experimental/ProceduralVegetationEditor/Config/BaseProceduralVegetationEditor.ini
deleted file mode 100644
diff --git a/Engine/Plugins/Experimental/ProceduralVegetationEditor/Config/DefaultProceduralVegetationEditor.ini b/Engine/Plugins/Experimental/ProceduralVegetationEditor/Config/DefaultProceduralVegetationEditor.ini
new file mode 100644
--- /dev/null
+++ b/Engine/Plugins/Experimental/ProceduralVegetationEditor/Config/DefaultProceduralVegetationEditor.ini
@@ -0,0 +1,2 @@
+﻿[CoreRedirects]
++ClassRedirects=(OldName="/Script/ProceduralVegetation.PVOutputSettings",NewName="/Script/ProceduralVegetation.PVExportSettings")
\ No newline at end of file
