Remove ensure which validates that a branch has more than 1 point, as there are valid cases for branches with only 1 point

#Synced-CL 52249728

diff --git a/Engine/Plugins/Experimental/ProceduralVegetationEditor/Source/ProceduralVegetation/Private/Helpers/PVAttributesHelper.cpp b/Engine/Plugins/Experimental/ProceduralVegetationEditor/Source/ProceduralVegetation/Private/Helpers/PVAttributesHelper.cpp
--- a/Engine/Plugins/Experimental/ProceduralVegetationEditor/Source/ProceduralVegetation/Private/Helpers/PVAttributesHelper.cpp
+++ b/Engine/Plugins/Experimental/ProceduralVegetationEditor/Source/ProceduralVegetation/Private/Helpers/PVAttributesHelper.cpp
@@ -393,7 +393,7 @@ void PV::AttributesHelper::ComputePointPlantGradient(const FComputePointPlantGra
 		[&](int32 BranchIndex)->EForEachResult
 		{
 			const TArray<int32>& BranchPoints = InAttributes.BranchPoints[BranchIndex];
-			if (!ensure(BranchPoints.Num() > 1))
+			if (!ensure(BranchPoints.Num() > 0))
 			{
 				return EForEachResult::Continue;
 			}
@@ -403,7 +403,9 @@ void PV::AttributesHelper::ComputePointPlantGradient(const FComputePointPlantGra
 
 			for (int32 i = bIsTrunk ? 0 : 1; i < BranchPoints.Num(); ++i)
 			{
-				const float BranchAlpha = i / float(BranchPoints.Num() - 1);
+				const float BranchAlpha = BranchPoints.Num() > 1 
+					? i / float(BranchPoints.Num() - 1)
+					: 0.f;
 				const float PointPlantGradient = FMath::Lerp(BasePointPlantGradient, 0, BranchAlpha);
 				InAttributes.PointPlantGradient[BranchPoints[i]] = PointPlantGradient;
 			}
@@ -672,12 +674,11 @@ void PV::AttributesHelper::ComputeBudLateralMeristem(const FComputeBudLateralMer
 		for (const int32 TrunkIndex : GetTrunkIndices(InAttributes.BranchParentNumber))
 		{
 			const TArray<int32>& BranchPoints = InAttributes.BranchPoints[TrunkIndex];
-			if (ensure(BranchPoints.Num() > 1))
+			if (BranchPoints.Num() > 1)
 			{
 				const int32 SecondToLastPointIndex = BranchPoints[BranchPoints.Num() - 2];
 				return InAttributes.PointScale[SecondToLastPointIndex] * 0.01f;
 			}
-			break;
 		}
 		return 0.033f;
 	});
