[PVE] -  Fix for Grower doesn't appear in list of compatible nodes when dragging the output of other nodes

#Synced-CL 51773451

diff --git a/Engine/Plugins/Experimental/ProceduralVegetationEditor/Source/ProceduralVegetation/Private/Nodes/PVGrowerSettings.cpp b/Engine/Plugins/Experimental/ProceduralVegetationEditor/Source/ProceduralVegetation/Private/Nodes/PVGrowerSettings.cpp
--- a/Engine/Plugins/Experimental/ProceduralVegetationEditor/Source/ProceduralVegetation/Private/Nodes/PVGrowerSettings.cpp
+++ b/Engine/Plugins/Experimental/ProceduralVegetationEditor/Source/ProceduralVegetation/Private/Nodes/PVGrowerSettings.cpp
@@ -66,7 +66,6 @@ TArray<FPCGPinProperties> UPVGrowerSettings::InputPinProperties() const
 	FPCGPinProperties& Pin = Properties.Emplace_GetRef(PCGPinConstants::DefaultInputLabel, GetInputPinTypeIdentifier());
 	Pin.SetAllowMultipleConnections(false);
 	Pin.SetNormalPin();
-	Pin.SetOverrideOrUserParamPin();
 	Pin.bAllowMultipleData = false;
 
 	FPCGPinProperties& PhyllotaxyPin = Properties.Emplace_GetRef(PV::Pins::GrowerPhyllotaxyInputLabel, FPCGDataTypeIdentifier{ FPVDataTypeInfoGrowerPhyllotaxy::AsId()});
diff --git a/Engine/Plugins/Experimental/ProceduralVegetationEditor/Source/ProceduralVegetation/Private/Nodes/PVGrowerSettings.h b/Engine/Plugins/Experimental/ProceduralVegetationEditor/Source/ProceduralVegetation/Private/Nodes/PVGrowerSettings.h
--- a/Engine/Plugins/Experimental/ProceduralVegetationEditor/Source/ProceduralVegetation/Private/Nodes/PVGrowerSettings.h
+++ b/Engine/Plugins/Experimental/ProceduralVegetationEditor/Source/ProceduralVegetation/Private/Nodes/PVGrowerSettings.h
@@ -27,6 +27,8 @@ public:
 	virtual TMap<UObject*, FTransform> GetViewportObjects() const override;
 #endif
 
+	virtual bool CanCullTaskIfUnwired() const override { return false; }
+	
 protected:
 	virtual TArray<FPCGPinProperties> InputPinProperties() const override;
 	virtual FPCGDataTypeIdentifier GetInputPinTypeIdentifier() const override;
