[PVE] Fix for tilling in texture setup node.

#Synced-CL 51773451

diff --git a/Engine/Plugins/Experimental/ProceduralVegetationEditor/Source/ProceduralVegetation/Private/Implementations/PVTextureSetup.cpp b/Engine/Plugins/Experimental/ProceduralVegetationEditor/Source/ProceduralVegetation/Private/Implementations/PVTextureSetup.cpp
--- a/Engine/Plugins/Experimental/ProceduralVegetationEditor/Source/ProceduralVegetation/Private/Implementations/PVTextureSetup.cpp
+++ b/Engine/Plugins/Experimental/ProceduralVegetationEditor/Source/ProceduralVegetation/Private/Implementations/PVTextureSetup.cpp
@@ -119,7 +119,7 @@ void FPVTextureSetupImplementation::DilationPass(FRDGBuilder& GraphBuilder, floa
 	Params->TileX = TileX;
 	Params->TileY = TileY;
 	Params->Source = InputTextureSRV;
-	Params->SourceSampler = TStaticSamplerState<SF_Bilinear, AM_Clamp, AM_Clamp>::GetRHI();
+	Params->SourceSampler = TStaticSamplerState<SF_Bilinear, AM_Wrap, AM_Wrap>::GetRHI();
 	Params->Target = RWOutputTextureUAV;
 	Params->TargetSize = FIntPoint(OutputTexture->GetSurfaceWidth(), OutputTexture->GetSurfaceHeight());
 	TShaderMapRef<FPVTextureSetupDilationCS> PVTextureSetupDilationCS(GetGlobalShaderMap(GMaxRHIFeatureLevel));
