[PVE] Details attributes added in collection for grower so visualizations can work properly

#Synced-CL 51773451

diff --git a/Engine/Plugins/Experimental/ProceduralVegetationEditor/Source/ProceduralVegetation/Private/Facades/PVMetaInfoFacade.cpp b/Engine/Plugins/Experimental/ProceduralVegetationEditor/Source/ProceduralVegetation/Private/Facades/PVMetaInfoFacade.cpp
--- a/Engine/Plugins/Experimental/ProceduralVegetationEditor/Source/ProceduralVegetation/Private/Facades/PVMetaInfoFacade.cpp
+++ b/Engine/Plugins/Experimental/ProceduralVegetationEditor/Source/ProceduralVegetation/Private/Facades/PVMetaInfoFacade.cpp
@@ -47,6 +47,9 @@ namespace PV::Facades
 		GraftUpAlignmentAttribute.Add();
 		GraftHealthAttribute.Add();
 		GraftTipAttribute.Add();
+		LeafGrowthAttribute.Add();
+		AbscissionSenescenseAttribute.Add();
+		LateralElongationAttribute.Add();
 	}
 
 	bool FMetaInfoFacade::IsValid() const
@@ -114,6 +117,36 @@ namespace PV::Facades
 		return GraftData;
 	}
 
+	void FMetaInfoFacade::SetLeafGrowth(const TArray<float>& InLeafGrowth)
+	{
+		check(!IsConst());
+		if (Collection->NumElements(GroupNames::DetailsGroup) == 0)
+		{
+			Collection->AddElements(1, GroupNames::DetailsGroup);
+		}
+		LeafGrowthAttribute.Modify()[0] = InLeafGrowth;
+	}
+
+	void FMetaInfoFacade::SetAbscissionSenescense(const TArray<float>& InAbscissionSenescense)
+	{
+		check(!IsConst());
+		if (Collection->NumElements(GroupNames::DetailsGroup) == 0)
+		{
+			Collection->AddElements(1, GroupNames::DetailsGroup);
+		}
+		AbscissionSenescenseAttribute.Modify()[0] = InAbscissionSenescense;
+	}
+
+	void FMetaInfoFacade::SetLateralElongation(const TArray<float>& InLateralElongation)
+	{
+		check(!IsConst());
+		if (Collection->NumElements(GroupNames::DetailsGroup) == 0)
+		{
+			Collection->AddElements(1, GroupNames::DetailsGroup);
+		}
+		LateralElongationAttribute.Modify()[0] = InLateralElongation;
+	}
+
 	const TArray<float>& FMetaInfoFacade::GetLeafGrowth() const
 	{
 		if(LeafGrowthAttribute.IsValid() && LeafGrowthAttribute.IsValidIndex(0))
diff --git a/Engine/Plugins/Experimental/ProceduralVegetationEditor/Source/ProceduralVegetation/Private/Facades/PVTreeFacade.cpp b/Engine/Plugins/Experimental/ProceduralVegetationEditor/Source/ProceduralVegetation/Private/Facades/PVTreeFacade.cpp
--- a/Engine/Plugins/Experimental/ProceduralVegetationEditor/Source/ProceduralVegetation/Private/Facades/PVTreeFacade.cpp
+++ b/Engine/Plugins/Experimental/ProceduralVegetationEditor/Source/ProceduralVegetation/Private/Facades/PVTreeFacade.cpp
@@ -221,7 +221,7 @@ namespace PV::Facades
 				if (PV::ValidateAttributeCollection(BranchPointsAttribute, BudLateralMeristemAttribute, LateralElongationAttribute))
 				{
 					VisualizationValues.SetNum(BudLateralMeristemAttribute.Num());
-					const float CellDensity = (1 - LateralElongationAttribute[0].SeedScaleEffect) * (LateralElongationAttribute[0].Elongation * 100.0f);
+					const float CellDensityMax = LateralElongationAttribute[0].Elongation * 100.0f;
 
 					for (int32 BranchIndex = 0; BranchIndex < BranchPointsAttribute.Num(); ++BranchIndex)
 					{
@@ -230,7 +230,7 @@ namespace PV::Facades
 						{
 							const float CurrentLateralMeristem = BudLateralMeristemAttribute[BranchPointIndex].LateralMeristem;
 							VisualizationValues[BranchPointIndex] = FMath::GetMappedRangeValueClamped(
-								FVector2f{0.0f, CellDensity},
+								FVector2f{0.0f, CellDensityMax},
 								FVector2f{0.0f, 1.0f},
 								CurrentLateralMeristem
 							);
diff --git a/Engine/Plugins/Experimental/ProceduralVegetationEditor/Source/ProceduralVegetation/Private/Implementations/PVGrower.cpp b/Engine/Plugins/Experimental/ProceduralVegetationEditor/Source/ProceduralVegetation/Private/Implementations/PVGrower.cpp
--- a/Engine/Plugins/Experimental/ProceduralVegetationEditor/Source/ProceduralVegetation/Private/Implementations/PVGrower.cpp
+++ b/Engine/Plugins/Experimental/ProceduralVegetationEditor/Source/ProceduralVegetation/Private/Implementations/PVGrower.cpp
@@ -7,6 +7,7 @@
 #include "Algo/Accumulate.h"
 #include "Facades/PVAttributesNames.h"
 #include "Facades/PVBranchFacade.h"
+#include "Facades/PVMetaInfoFacade.h"
 #include "Facades/PVPointFacade.h"
 #include "Helpers/PVAttributesHelper.h"
 #include "GeometryCollection/ManagedArrayCollection.h"
@@ -90,7 +91,7 @@ void FPVGrower::Grow(const FPVSeedPoint& SeedPoint, const FPVGrowerParams& InGro
 	{
 		SplitBranchSourcePoints(Skeleton);
 	}
-	FillCollection(InCollection, Skeleton);
+	FillCollection(InCollection, Skeleton, InGrowerParams);
 	//PrintPoints();
 }
 
@@ -105,11 +106,11 @@ void FPVGrower::Grow(const FManagedArrayCollection& InCollection, const FPVGrowe
 	{
 		SplitBranchSourcePoints(Skeleton);
 	}
-	FillCollection(Collection, Skeleton);
+	FillCollection(Collection, Skeleton, InGrowerParams);
 	//PrintPoints();
 }
 
-void FPVGrower::FillCollection(FManagedArrayCollection& Collection, UPVGrowerData* Skeleton)
+void FPVGrower::FillCollection(FManagedArrayCollection& Collection, UPVGrowerData* Skeleton, const FPVGrowerParams& GrowerParams)
 {
 	PV::Facades::FPointFacade PointFacade = PV::Facades::FPointFacade(Collection);
 	PV::Facades::FBranchFacade BranchFacade = PV::Facades::FBranchFacade(Collection);
@@ -223,6 +224,45 @@ void FPVGrower::FillCollection(FManagedArrayCollection& Collection, UPVGrowerDat
 
 	check(PointFacade.IsValid());
 	check(BranchFacade.IsValid());
+
+	PV::Facades::FMetaInfoFacade MetaInfoFacade(Collection);
+
+	const TArray<float> LeafGrowthData = {
+		GrowerParams.LeafGrowth.Ethylene,
+		GrowerParams.LeafGrowth.StartScale,
+		GrowerParams.LeafGrowth.EndScale,
+		static_cast<float>(GrowerParams.LeafGrowth.Density),
+		GrowerParams.LeafGrowth.AbscisicAcid,
+		GrowerParams.LeafGrowth.EthyleneThreshold,
+		GrowerParams.LeafGrowth.LeafAuxinRetention
+	};
+	MetaInfoFacade.SetLeafGrowth(LeafGrowthData);
+
+	const TArray<float> AbscissionSenescenseData = {
+		GrowerParams.GravityParams.BranchWeightBreakThreshold,
+		static_cast<float>(GrowerParams.AgeSenescence.Threshold),
+		static_cast<float>(GrowerParams.AgeSenescence.AbscissionMin),
+		static_cast<float>(GrowerParams.AgeSenescence.AbscissionMax),
+		GrowerParams.AgeSenescence.ScaleRetention,
+		GrowerParams.LightSenescence.Threshold,
+		static_cast<float>(GrowerParams.LightSenescence.AbscissionMin),
+		static_cast<float>(GrowerParams.LightSenescence.AbscissionMax),
+		GrowerParams.LightSenescence.ScaleRetention
+	};
+	MetaInfoFacade.SetAbscissionSenescense(AbscissionSenescenseData);
+
+	const TArray<float> LateralElongationData = {
+		GrowerParams.LateralParams.Elongation,
+		GrowerParams.LateralParams.SeedScaleEffect,
+		GrowerParams.LateralParams.WhorledImpact,
+		GrowerParams.GravityParams.CellDensity.GetLowerBoundValue(),
+		GrowerParams.GravityParams.CellDensity.GetUpperBoundValue(),
+		GrowerParams.GravityParams.CellDensityCore,
+		GrowerParams.GravityParams.CellDevelopmentTime,
+		GrowerParams.GravityParams.LeafWeight,
+		GrowerParams.GravityParams.CellDensityInput
+	};
+	MetaInfoFacade.SetLateralElongation(LateralElongationData);
 }
 
 void FPVGrower::FillFromCollection(const FManagedArrayCollection& Collection, UPVGrowerData* Skeleton, bool bAgeReset, bool bUseSplitPoints)
diff --git a/Engine/Plugins/Experimental/ProceduralVegetationEditor/Source/ProceduralVegetation/Private/Implementations/PVGrower.h b/Engine/Plugins/Experimental/ProceduralVegetationEditor/Source/ProceduralVegetation/Private/Implementations/PVGrower.h
--- a/Engine/Plugins/Experimental/ProceduralVegetationEditor/Source/ProceduralVegetation/Private/Implementations/PVGrower.h
+++ b/Engine/Plugins/Experimental/ProceduralVegetationEditor/Source/ProceduralVegetation/Private/Implementations/PVGrower.h
@@ -19,7 +19,7 @@ struct FPVGrower
 
 	static void Grow(const FManagedArrayCollection& InCollection, const FPVGrowerParams& InGrowerParams, FManagedArrayCollection& Collection, bool bUseSplitPoints = true);
 
-	static void FillCollection(FManagedArrayCollection& Collection, UPVGrowerData* Skeleton);
+	static void FillCollection(FManagedArrayCollection& Collection, UPVGrowerData* Skeleton, const FPVGrowerParams& GrowerParams);
 
 	static void FillFromCollection(const FManagedArrayCollection& Collection, UPVGrowerData* Skeleton, bool bAgeReset, bool bUseSplitPoints = true);
 
diff --git a/Engine/Plugins/Experimental/ProceduralVegetationEditor/Source/ProceduralVegetation/Public/Facades/PVMetaInfoFacade.h b/Engine/Plugins/Experimental/ProceduralVegetationEditor/Source/ProceduralVegetation/Public/Facades/PVMetaInfoFacade.h
--- a/Engine/Plugins/Experimental/ProceduralVegetationEditor/Source/ProceduralVegetation/Public/Facades/PVMetaInfoFacade.h
+++ b/Engine/Plugins/Experimental/ProceduralVegetationEditor/Source/ProceduralVegetation/Public/Facades/PVMetaInfoFacade.h
@@ -30,6 +30,10 @@ namespace PV::Facades
 		void SetGraftAttributes(const FGraftData& InGraftData);
 		FGraftData GetGraftAttributes() const;
 		
+		void SetLeafGrowth(const TArray<float>& InLeafGrowth);
+		void SetAbscissionSenescense(const TArray<float>& InAbscissionSenescense);
+		void SetLateralElongation(const TArray<float>& InLateralElongation);
+
 		const TArray<float>& GetLeafGrowth() const;
 		const TArray<float>& GetAbscissionSenescense() const;
 		const TArray<float>& GetLateralElongation() const;
