[PVE] Object Interaction update for fused points

#Synced-CL 52128411

diff --git a/Engine/Plugins/Experimental/ProceduralVegetationEditor/Source/ProceduralVegetation/Private/Implementations/PVObjectInteraction.cpp b/Engine/Plugins/Experimental/ProceduralVegetationEditor/Source/ProceduralVegetation/Private/Implementations/PVObjectInteraction.cpp
--- a/Engine/Plugins/Experimental/ProceduralVegetationEditor/Source/ProceduralVegetation/Private/Implementations/PVObjectInteraction.cpp
+++ b/Engine/Plugins/Experimental/ProceduralVegetationEditor/Source/ProceduralVegetation/Private/Implementations/PVObjectInteraction.cpp
@@ -12,6 +12,7 @@
 #include "Facades/PVBranchFacade.h"
 #include "Facades/PVBudVectorsFacade.h"
 #include "Facades/PVPointFacade.h"
+#include "Facades/PVTreeFacade.h"
 
 #include "Helpers/PVUtilities.h"
 
@@ -33,11 +34,8 @@ FVector FPVObjectInteraction::GetTriangleNormal(FMeshDescriptionTriangleMeshAdap
 	}
 }
 
-void FPVObjectInteraction::ObjectInteraction(const FManagedArrayCollection& InCollection, FPVColliderParams Collider,
-                                             FManagedArrayCollection& OutCollection)
+void FPVObjectInteraction::ObjectInteraction(FManagedArrayCollection& OutCollection, FPVColliderParams Collider)
 {
-	OutCollection = InCollection;
-
 	const FTransform& Transform = Collider.Transform;
 
 	const UStaticMesh* ColliderMesh = Collider.Mesh.LoadSynchronous();
@@ -60,22 +58,16 @@ void FPVObjectInteraction::ObjectInteraction(const FManagedArrayCollection& InCo
 	PointsToRemove.Init(false, PointFacade.GetElementCount());
 	TArray<bool> BranchesToRemove;
 	BranchesToRemove.Init(false, BranchFacade.GetElementCount());
-	TArray<bool> FoliageInstancesToRemove;
-	FoliageInstancesToRemove.Init(false, FoliageFacade.NumFoliageEntries());
 
 	TManagedArray<FVector3f>& PointPositions = PointFacade.ModifyPositions();
 	TManagedArray<TArray<FVector3f>>& BudDirections = BudVectorFacade.ModifyBudDirections();
-	TManagedArray<FVector3f>& FoliagePositions = FoliageFacade.ModifyPivotPositions();
 
 	const auto ApplyOffsetToBranchSuccessors = [&](const int32 BranchIndex, const FVector3f& Offset)
 		{
-			for (const int32 BranchPoint : BranchFacade.GetPoints(BranchIndex))
+			const TArray<int32>& BranchPoints = BranchFacade.GetPoints(BranchIndex);
+			for (int32 BranchPointIndex = 1; BranchPointIndex < BranchPoints.Num(); ++BranchPointIndex)
 			{
-				PointPositions[BranchPoint] += Offset;
-			}
-			for (const int32 FoliageEntryId : FoliageFacade.GetFoliageEntryIdsForBranch(BranchIndex))
-			{
-				FoliagePositions[FoliageEntryId] += Offset;
+				PointPositions[BranchPoints[BranchPointIndex]] += Offset;
 			}
 
 			const TArray<int32> BranchChildren = BranchFacade.GetChildren(BranchIndex);
@@ -84,13 +76,10 @@ void FPVObjectInteraction::ObjectInteraction(const FManagedArrayCollection& InCo
 				const int32 ChildIndex = BranchFacade.GetBranchNumbers().Find(BranchChild);
 				if (ChildIndex != INDEX_NONE)
 				{
-					for (const int32 ChildBranchPoint : BranchFacade.GetPoints(ChildIndex))
+					const TArray<int32>& ChildBranchPoints = BranchFacade.GetPoints(ChildIndex);
+					for (int32 ChildBranchPointIndex = 1; ChildBranchPointIndex < ChildBranchPoints.Num(); ++ChildBranchPointIndex)
 					{
-						PointPositions[ChildBranchPoint] += Offset;
-					}
-					for (const int32 ChildFoliageEntryId : FoliageFacade.GetFoliageEntryIdsForBranch(ChildIndex))
-					{
-						FoliagePositions[ChildFoliageEntryId] += Offset;
+						PointPositions[ChildBranchPoints[ChildBranchPointIndex]] += Offset;
 					}
 				}
 			}
@@ -103,10 +92,6 @@ void FPVObjectInteraction::ObjectInteraction(const FManagedArrayCollection& InCo
 			{
 				PointsToRemove[BranchPoint] = true;
 			}
-			for (const int32 FoliageEntryId : FoliageFacade.GetFoliageEntryIdsForBranch(BranchIndex))
-			{
-				FoliageInstancesToRemove[FoliageEntryId] = true;
-			}
 
 			const TArray<int32> BranchChildren = BranchFacade.GetChildren(BranchIndex);
 			for (int32 BranchChild : BranchChildren)
@@ -119,13 +104,11 @@ void FPVObjectInteraction::ObjectInteraction(const FManagedArrayCollection& InCo
 					{
 						PointsToRemove[ChildBranchPoint] = true;
 					}
-					for (const int32 ChildFoliageEntryId : FoliageFacade.GetFoliageEntryIdsForBranch(ChildIndex))
-					{
-						FoliageInstancesToRemove[ChildFoliageEntryId] = true;
-					}
 				}
 			}
 		};
+	
+	const bool bIsTrimCollider = Collider.CollisionType == EPVCollisionType::TRIM_INSIDE || Collider.CollisionType == EPVCollisionType::TRIM_OUTSIDE;
 
 	TArray<int32> BranchIndices;
 	BranchIndices.Reserve(BranchFacade.GetElementCount());
@@ -134,21 +117,19 @@ void FPVObjectInteraction::ObjectInteraction(const FManagedArrayCollection& InCo
 	for (int32 BranchIndex : BranchIndices)
 	{
 		PVE_OUTER_LOOP_DEBUG_CHECK(break);
-
-		if (BranchesToRemove[BranchIndex] && Collider.CollisionType == EPVCollisionType::TRIM_OUTSIDE)
+		
+		if (bIsTrimCollider && BranchesToRemove[BranchIndex])
 		{
 			continue;
 		}
 
-		const float MaxPointScale = *Algo::MaxElement(PointFacade.GetPointScales());
-
 		const TArray<int32> BranchPoints = BranchFacade.GetPoints(BranchIndex);
 		for (int32 BranchPointIndex = 1; BranchPointIndex < BranchPoints.Num(); ++BranchPointIndex)
 		{
 			const int32 PointIndex = BranchPoints[BranchPointIndex - 1];
 			const int32 NextPointIndex = BranchPoints[BranchPointIndex];
 
-			if (PointsToRemove[PointIndex] && Collider.CollisionType == EPVCollisionType::TRIM_OUTSIDE)
+			if (bIsTrimCollider && PointsToRemove[PointIndex])
 			{
 				continue;
 			}
@@ -157,29 +138,12 @@ void FPVObjectInteraction::ObjectInteraction(const FManagedArrayCollection& InCo
 			const FVector NextPointPosition = static_cast<FVector>(PointPositions[NextPointIndex]);
 			const FVector SegmentVector = NextPointPosition - PointPosition;
 
-			// Get the previous direction
-			FVector PreviousDirection;
-			if (BranchPoints.IsValidIndex(BranchPointIndex - 2))
-			{
-				PreviousDirection = (PointPosition - BranchPoints[BranchPointIndex - 2]).
-					GetSafeNormal(UE_SMALL_NUMBER, FVector::UpVector);
-			}
-			else
-			{
-				PreviousDirection = (NextPointPosition - PointPosition).GetSafeNormal(UE_SMALL_NUMBER, FVector::UpVector);
-			}
-
-			const float PointScale = Collider.CollisionType == EPVCollisionType::AVOID
-				? MaxPointScale
-				: 0.0f;
-
 			FRay3d Ray(PointPosition, SegmentVector);
 			FRay3d InverseRay(Transform.InverseTransformPosition(Ray.Origin), Transform.InverseTransformVector(Ray.Direction));
 
 			const float PointLFR = PointFacade.GetLengthFromRoot(PointIndex);
-			const float NextPointLFR = PointFacade.GetLengthFromRoot(NextPointIndex);
 
-			const double MaxLength = SegmentVector.Length() + PointScale;
+			const double MaxLength = SegmentVector.Length();
 			const FVector MaxLengthVector = Ray.Direction * MaxLength;
 			const double InverseMaxLength = Transform.InverseTransformVector(MaxLengthVector).Length();
 
@@ -189,27 +153,23 @@ void FPVObjectInteraction::ObjectInteraction(const FManagedArrayCollection& InCo
 			{
 				const FVector HitNormal = Transform.TransformVectorNoScale(GetTriangleNormal(Mesh, HitTriangle, BaryCoords)).
 				                                    GetSafeNormal(UE_SMALL_NUMBER, FVector::UpVector);
+				const FVector InvertedNormal = -HitNormal;
 
 				if (Collider.CollisionType == EPVCollisionType::AVOID)
 				{
-					const FVector HitPoint = Transform.TransformPosition(InverseRay.PointAt(RayT));
-					const FVector InvertedNormal = -HitNormal;
+					const float PointScale = PointFacade.GetPointScale(PointIndex);
 
-					const double DistanceToHit = FVector::Distance(PointPosition, HitPoint);
+					const FVector HitPoint = Transform.TransformPosition(InverseRay.PointAt(RayT));
 					const double DistanceAfterHit = FVector::Distance(HitPoint, Ray.PointAt(MaxLength));
+
 					const FVector BiNormal = FVector::CrossProduct(Ray.Direction, InvertedNormal);
-					FVector Tangent = FVector::CrossProduct(InvertedNormal, BiNormal);
-					const FVector PreviousBiNormal = FVector::CrossProduct(PreviousDirection, InvertedNormal);
-					const FVector PreviousTangent = FVector::CrossProduct(InvertedNormal, PreviousBiNormal);
-					// Get the mean tangent of this and previous tangent
-					Tangent = (PreviousTangent + Tangent).GetSafeNormal(UE_SMALL_NUMBER, FVector::UpVector);
-					const FVector AvoidanceOffset = Tangent * DistanceToHit * (FVector::DotProduct(Ray.Direction, InvertedNormal) + 1);
-					const FVector AvoidancePosition = HitPoint + HitNormal * (PointScale + DistanceAfterHit) + AvoidanceOffset;
-					const FVector AvoidanceDirection = (AvoidancePosition - PointPosition).GetSafeNormal(UE_SMALL_NUMBER,
-						FVector::DownVector);
-
-					const FVector OffsetPosition = PointPosition + AvoidanceDirection * FVector::Distance(PointPosition,
-						NextPointPosition);
+					const FVector Tangent = FVector::CrossProduct(InvertedNormal, BiNormal).GetSafeNormal(UE_SMALL_NUMBER, FVector::UpVector);
+
+					const FVector AvoidanceOffset = Tangent * DistanceAfterHit * (FMath::Abs(FVector::DotProduct(Ray.Direction, InvertedNormal)) + 1);
+					const FVector AvoidancePosition = HitPoint + AvoidanceOffset + HitNormal * PointScale;
+					const FVector AvoidanceDirection = (AvoidancePosition - PointPosition).GetSafeNormal(UE_SMALL_NUMBER, FVector::DownVector);
+
+					const FVector OffsetPosition = PointPosition + AvoidanceDirection * FVector::Distance(PointPosition, NextPointPosition);
 					const FVector3f Offset = FVector3f(OffsetPosition - NextPointPosition);
 
 					PVE_LOOP_DEBUG_PARAMS_START()
@@ -218,15 +178,14 @@ void FPVObjectInteraction::ObjectInteraction(const FManagedArrayCollection& InCo
 						PVE_LOOP_DEBUG_DIRECTION_PARAM("Avoidance Offset", HitPoint, AvoidanceOffset);
 						PVE_LOOP_DEBUG_DIRECTION_PARAM("Avoidance Direction", NextPointPosition, AvoidanceDirection);
 						PVE_LOOP_DEBUG_VECTOR_PARAM("Avoidance Vector", PointPosition, AvoidancePosition);
-						PVE_LOOP_DEBUG_DIRECTION_PARAM("Offset", NextPointPosition, Offset);
+						PVE_LOOP_DEBUG_DIRECTION_PARAM("Offset", NextPointPosition, OffsetPosition - NextPointPosition);
 					PVE_LOOP_DEBUG_PARAMS_END()
 
 					PVE_LOOP_DEBUG_STEP(break);
 
-					constexpr int32 Smoothness = 15;
 					const float SmoothnessFactor = 1.1f + (1.0f - Collider.SmoothnessFactor);
 
-					const int32 StartIndex = FMath::Max(1, BranchPointIndex - Smoothness);
+					const int32 StartIndex = FMath::Max(1, BranchPointIndex - Collider.SmoothnessAmount);
 					FVector3f PreviousSiblingPosition;
 					for (int32 SiblingPointIndex = StartIndex; SiblingPointIndex < BranchPointIndex; ++SiblingPointIndex)
 					{
@@ -242,28 +201,15 @@ void FPVObjectInteraction::ObjectInteraction(const FManagedArrayCollection& InCo
 						{
 							const FVector3f SiblingSegmentVector = PointPositions[SiblingPoint] - PreviousSiblingPosition;
 							const FQuat4f Rotation = FQuat4f::FindBetween(SiblingSegmentVector, SiblingSegmentVector + PointOffset);
-							for (const int32 FoliageEntryId : FoliageFacade.GetFoliageEntryIdsForBranch(BranchIndex))
-							{
-								const float FoliageLFR = FoliageFacade.GetLengthFromRoot(FoliageEntryId);
-								if (FoliageLFR > SiblingPointLFR && FoliageLFR < NextSiblingPointLFR)
-								{
-									FoliagePositions[FoliageEntryId] += PointOffset;
-									FoliageFacade.SetUpVector(
-										FoliageEntryId,
-										Rotation.RotateVector(FoliageFacade.GetUpVector(FoliageEntryId))
-									);
-									FoliageFacade.SetNormalVector(
-										FoliageEntryId,
-										Rotation.RotateVector(FoliageFacade.GetNormalVector(FoliageEntryId))
-									);
-								}
-							}
-							
+
 							TArray<FVector3f>& BudDirection = BudDirections[SiblingPoint];
 							BudDirection[PV::Facades::BudDirectionsApical] = Rotation.RotateVector(BudDirection[PV::Facades::BudDirectionsApical]);
-							BudDirection[PV::Facades::BudDirectionsAxillary] = Rotation.RotateVector(BudDirection[PV::Facades::BudDirectionsAxillary]);
-							BudDirection[PV::Facades::BudDirectionsGuideCurve] = Rotation.RotateVector(BudDirection[PV::Facades::BudDirectionsGuideCurve]);
-							BudDirection[PV::Facades::BudDirectionsUpVector] = Rotation.RotateVector(BudDirection[PV::Facades::BudDirectionsUpVector]);
+							BudDirection[PV::Facades::BudDirectionsAxillary] = Rotation.
+								RotateVector(BudDirection[PV::Facades::BudDirectionsAxillary]);
+							BudDirection[PV::Facades::BudDirectionsGuideCurve] = Rotation.RotateVector(
+								BudDirection[PV::Facades::BudDirectionsGuideCurve]);
+							BudDirection[PV::Facades::BudDirectionsUpVector] = Rotation.
+								RotateVector(BudDirection[PV::Facades::BudDirectionsUpVector]);
 						}
 
 						PreviousSiblingPosition = PointPositions[SiblingPoint];
@@ -291,25 +237,6 @@ void FPVObjectInteraction::ObjectInteraction(const FManagedArrayCollection& InCo
 					}
 
 					const FQuat4f Rotation = FQuat4f::FindBetween(FVector3f(SegmentVector), FVector3f(SegmentVector) + Offset);
-					for (const int32 FoliageEntryId : FoliageFacade.GetFoliageEntryIdsForBranch(BranchIndex))
-					{
-						const float FoliageLFR = FoliageFacade.GetLengthFromRoot(FoliageEntryId);
-						if (FoliageLFR > PointLFR)
-						{
-							FoliagePositions[FoliageEntryId] += Offset;
-							if (FoliageLFR < NextPointLFR)
-							{
-								FoliageFacade.SetUpVector(
-									FoliageEntryId,
-									Rotation.RotateVector(FoliageFacade.GetUpVector(FoliageEntryId))
-								);
-								FoliageFacade.SetNormalVector(
-									FoliageEntryId,
-									Rotation.RotateVector(FoliageFacade.GetNormalVector(FoliageEntryId))
-								);
-							}
-						}
-					}
 
 					TArray<FVector3f>& BudDirection = BudDirections[NextPointIndex];
 					BudDirection[PV::Facades::BudDirectionsApical] = Rotation.RotateVector(BudDirection[PV::Facades::BudDirectionsApical]);
@@ -331,11 +258,10 @@ void FPVObjectInteraction::ObjectInteraction(const FManagedArrayCollection& InCo
 						}
 					}
 				}
-				else if (Collider.CollisionType == EPVCollisionType::TRIM_OUTSIDE)
+				else if (bIsTrimCollider)
 				{
-					// In case the edge hit from inside the collider (Normal and Ray are opposite direction)
-					// Trim the remaining branch, all child branches and all foliage instances in case of hit
-					if (FVector::DotProduct(Ray.Direction, HitNormal) > 0)
+					const double RayDot = FVector::DotProduct(Ray.Direction, HitNormal);
+					if (Collider.CollisionType == EPVCollisionType::TRIM_OUTSIDE ? RayDot > 0 : RayDot < 0)
 					{
 						for (int32 SiblingPointIndex = BranchPointIndex; SiblingPointIndex < BranchPoints.Num(); ++SiblingPointIndex)
 						{
@@ -346,19 +272,11 @@ void FPVObjectInteraction::ObjectInteraction(const FManagedArrayCollection& InCo
 							PointsToRemove[BranchPoints[0]] = true;
 							BranchesToRemove[BranchIndex] = true;
 						}
-						for (const int32 FoliageEntryId : FoliageFacade.GetFoliageEntryIdsForBranch(BranchIndex))
-						{
-							const float FoliageLFR = FoliageFacade.GetLengthFromRoot(FoliageEntryId);
-							if (FoliageLFR > PointLFR)
-							{
-								FoliageInstancesToRemove[FoliageEntryId] = true;
-							}
-						}
 
 						const float Ratio = static_cast<float>(BranchPointIndex) / static_cast<float>(BranchPoints.Num());
 
 						// Rescale the branches based on the cut point
-						for (int32 Index = 0; Index <= BranchPointIndex; ++Index)
+						for (int32 Index = !BranchFacade.IsTrunk(BranchIndex); Index <= BranchPointIndex; ++Index)
 						{
 							PointFacade.SetPointScale(BranchPoints[Index], PointFacade.GetPointScale(BranchPoints[Index]) * Ratio);
 						}
@@ -367,8 +285,10 @@ void FPVObjectInteraction::ObjectInteraction(const FManagedArrayCollection& InCo
 						{
 							if (const int32 ChildBranchIndex = BranchFacade.GetBranchNumbers().Find(ChildBranch); ChildBranchIndex != INDEX_NONE)
 							{
-								for (int32 ChildBranchPoint : BranchFacade.GetPoints(ChildBranchIndex))
+								const TArray<int32>& ChildBranchPoints = BranchFacade.GetPoints(ChildBranchIndex);
+								for (int32 ChildBranchPointIndex = 1; ChildBranchPointIndex < ChildBranchPoints.Num(); ++ChildBranchPointIndex)
 								{
+									const int32 ChildBranchPoint = ChildBranchPoints[ChildBranchPointIndex];
 									PointFacade.SetPointScale(ChildBranchPoint, PointFacade.GetPointScale(ChildBranchPoint) * Ratio);
 								}
 							}
@@ -398,11 +318,11 @@ void FPVObjectInteraction::ObjectInteraction(const FManagedArrayCollection& InCo
 	}
 
 	// Remove all the ToRemove entries using the same algorithm being used in Carve algorithm
-	if (Collider.CollisionType == EPVCollisionType::TRIM_OUTSIDE)
+	if (bIsTrimCollider)
 	{
-		FPVCarve::RemoveEntriesAndRecomputeAttributes(
+		TArray<bool> FoliageInstancesToRemove;
+		PV::Facades::FTreeFacade::RemoveEntriesAndReIndexAttributes(
 			OutCollection,
-			InCollection,
 			PointsToRemove,
 			BranchesToRemove,
 			FoliageInstancesToRemove
diff --git a/Engine/Plugins/Experimental/ProceduralVegetationEditor/Source/ProceduralVegetation/Private/Implementations/PVObjectInteraction.h b/Engine/Plugins/Experimental/ProceduralVegetationEditor/Source/ProceduralVegetation/Private/Implementations/PVObjectInteraction.h
--- a/Engine/Plugins/Experimental/ProceduralVegetationEditor/Source/ProceduralVegetation/Private/Implementations/PVObjectInteraction.h
+++ b/Engine/Plugins/Experimental/ProceduralVegetationEditor/Source/ProceduralVegetation/Private/Implementations/PVObjectInteraction.h
@@ -14,5 +14,5 @@ struct FPVObjectInteraction
 		TOptional<FVector> BaryCoords = TOptional<FVector>()
 	);
 	
-	static void ObjectInteraction(const FManagedArrayCollection& InCollection, FPVColliderParams Collider, FManagedArrayCollection& OutCollection);
+	static void ObjectInteraction(FManagedArrayCollection& OutCollection, FPVColliderParams Collider);
 };
diff --git a/Engine/Plugins/Experimental/ProceduralVegetationEditor/Source/ProceduralVegetation/Private/Nodes/PVObjectInteractionSettings.cpp b/Engine/Plugins/Experimental/ProceduralVegetationEditor/Source/ProceduralVegetation/Private/Nodes/PVObjectInteractionSettings.cpp
--- a/Engine/Plugins/Experimental/ProceduralVegetationEditor/Source/ProceduralVegetation/Private/Nodes/PVObjectInteractionSettings.cpp
+++ b/Engine/Plugins/Experimental/ProceduralVegetationEditor/Source/ProceduralVegetation/Private/Nodes/PVObjectInteractionSettings.cpp
@@ -49,8 +49,7 @@ bool FPVObjectInteractionElement::ExecuteInternal(FPCGContext* InContext) const
 			{
 				PVE_OUTER_LOOP_DEBUG_CHECK(break);
 				
-				const FManagedArrayCollection InCollection = MoveTemp(OutCollection);
-				FPVObjectInteraction::ObjectInteraction(InCollection, Collider, OutCollection);
+				FPVObjectInteraction::ObjectInteraction(OutCollection, Collider);
 			}
 
 			OutManagedArrayCollectionData->Initialize(MoveTemp(OutCollection));
diff --git a/Engine/Plugins/Experimental/ProceduralVegetationEditor/Source/ProceduralVegetation/Public/Nodes/PVObjectInteractionSettings.h b/Engine/Plugins/Experimental/ProceduralVegetationEditor/Source/ProceduralVegetation/Public/Nodes/PVObjectInteractionSettings.h
--- a/Engine/Plugins/Experimental/ProceduralVegetationEditor/Source/ProceduralVegetation/Public/Nodes/PVObjectInteractionSettings.h
+++ b/Engine/Plugins/Experimental/ProceduralVegetationEditor/Source/ProceduralVegetation/Public/Nodes/PVObjectInteractionSettings.h
@@ -10,7 +10,8 @@ UENUM()
 enum class EPVCollisionType
 {
 	AVOID UMETA(DisplayName = "Avoid"),
-	TRIM_OUTSIDE UMETA(DisplayName = "Trim Outside")
+	TRIM_OUTSIDE UMETA(DisplayName = "Trim Outside"),
+	TRIM_INSIDE UMETA(DisplayName = "Trim Inside")
 };
 
 USTRUCT()
@@ -27,6 +28,9 @@ struct FPVColliderParams
 	UPROPERTY(EditAnywhere, Category = "Collider")
 	EPVCollisionType CollisionType = EPVCollisionType::AVOID;
 	
+	UPROPERTY(EditAnywhere, Category = "Collider", meta=(EditCondition = "CollisionType == EPVCollisionType::AVOID", EditConditionHides, UIMin=0, ClampMin=0))
+	int32 SmoothnessAmount = 0;
+	
 	UPROPERTY(EditAnywhere, Category = "Collider", meta=(EditCondition = "CollisionType == EPVCollisionType::AVOID", EditConditionHides, UIMin=0, UIMax=1, ClampMin=0, ClampMax=1))
 	float SmoothnessFactor = 1.0f;
 };
