Minor cleanup of UPVData

- Fix UPVData::Initialize not moving the input collection, but rather copies it into the newly created TSharedPtr<FManagedArrayCollection>
- Remove "mutable FBox CachedBounds" as it was unused

#Synced-CL 52292167
#Synced-CL 52292167

diff --git a/Engine/Plugins/Experimental/ProceduralVegetationEditor/Source/ProceduralVegetation/Private/DataTypes/PVData.cpp b/Engine/Plugins/Experimental/ProceduralVegetationEditor/Source/ProceduralVegetation/Private/DataTypes/PVData.cpp
--- a/Engine/Plugins/Experimental/ProceduralVegetationEditor/Source/ProceduralVegetation/Private/DataTypes/PVData.cpp
+++ b/Engine/Plugins/Experimental/ProceduralVegetationEditor/Source/ProceduralVegetation/Private/DataTypes/PVData.cpp
@@ -14,9 +14,9 @@ UPVData::UPVData(const FObjectInitializer& ObjectInitializer)
 	Collection = MakeShared<FManagedArrayCollection>();
 }
 
-void UPVData::Initialize(FManagedArrayCollection&& InCollection, bool bCanTakeOwnership)
+void UPVData::Initialize(FManagedArrayCollection&& InCollection)
 {
-	Collection = MakeShared<FManagedArrayCollection>(InCollection);
+	Collection = MakeShared<FManagedArrayCollection>(MoveTemp(InCollection));
 }
 
 void UPVData::AddToCrc(FArchiveCrc32& Ar, bool bFullDataCrc) const
@@ -43,7 +43,7 @@ void UPVData::AddToCrc(FArchiveCrc32& Ar, bool bFullDataCrc) const
 
 FBox UPVData::GetBounds() const
 {
-	return CachedBounds;
+	return FBox(EForceInit::ForceInit);
 }
 
 bool UPVData::SamplePoint(const FTransform& InTransform, const FBox& InBounds, FPCGPoint& OutPoint, UPCGMetadata* OutMetadata) const
@@ -66,7 +66,7 @@ UPCGSpatialData* UPVData::CopyInternal(FPCGContext* Context) const
 	FManagedArrayCollection NewCollection;
 	Collection->CopyTo(&NewCollection);
 	
-	NewProceduralVegetationData->Initialize(MoveTemp(NewCollection), /*bCanTakeOwnership=*/false);
+	NewProceduralVegetationData->Initialize(MoveTemp(NewCollection));
 
 #if WITH_EDITOR
 	NewProceduralVegetationData->DebugSettings = DebugSettings;
diff --git a/Engine/Plugins/Experimental/ProceduralVegetationEditor/Source/ProceduralVegetation/Public/DataTypes/PVData.h b/Engine/Plugins/Experimental/ProceduralVegetationEditor/Source/ProceduralVegetation/Public/DataTypes/PVData.h
--- a/Engine/Plugins/Experimental/ProceduralVegetationEditor/Source/ProceduralVegetation/Public/DataTypes/PVData.h
+++ b/Engine/Plugins/Experimental/ProceduralVegetationEditor/Source/ProceduralVegetation/Public/DataTypes/PVData.h
@@ -199,7 +199,7 @@ class PROCEDURALVEGETATION_API UPVData : public UPCGSpatialData
 public:
 	UPVData(const FObjectInitializer& ObjectInitializer);
 	
-	void Initialize(FManagedArrayCollection&& InCollection, bool bCanTakeOwnership = false);
+	void Initialize(FManagedArrayCollection&& InCollection);
 	
 	// ~Begin UPCGData interface
 	virtual EPCGDataType GetDataType() const override { return EPCGDataType::Other; }
@@ -263,6 +263,4 @@ protected:
 
 	TArray<FPVVisualizationCollection> VisualizationCollections;
 #endif
-
-	mutable FBox CachedBounds = FBox(EForceInit::ForceInit);
 };
