[PVE] Category Toggles
Anything in the Basic category is always visible
All shows all the categories

#Synced-CL 51773451

diff --git a/Engine/Plugins/Experimental/ProceduralVegetationEditor/Source/ProceduralVegetationEditor/Private/Customizations/PVGrowerSettingsCustomization.cpp b/Engine/Plugins/Experimental/ProceduralVegetationEditor/Source/ProceduralVegetationEditor/Private/Customizations/PVGrowerSettingsCustomization.cpp
new file mode 100644
--- /dev/null
+++ b/Engine/Plugins/Experimental/ProceduralVegetationEditor/Source/ProceduralVegetationEditor/Private/Customizations/PVGrowerSettingsCustomization.cpp
@@ -0,0 +1,332 @@
+// Copyright Epic Games, Inc. All Rights Reserved.
+
+#include "PVGrowerSettingsCustomization.h"
+
+#include "DetailCategoryBuilder.h"
+#include "DetailLayoutBuilder.h"
+#include "DetailWidgetRow.h"
+#include "IDetailPropertyRow.h"
+#include "PropertyHandle.h"
+
+#include "Nodes/PVGrowerSettings.h"
+
+#include "Input/Reply.h"
+#include "Misc/Attribute.h"
+#include "Styling/AppStyle.h"
+#include "Widgets/Input/SButton.h"
+#include "Widgets/Input/SCheckBox.h"
+#include "Widgets/Layout/SWrapBox.h"
+#include "Widgets/Text/STextBlock.h"
+
+#define LOCTEXT_NAMESPACE "PVGrowerSettingsCustomization"
+
+FString FPVGrowerSettingsCustomization::ActiveSection(TEXT("Basic"));
+
+namespace
+{
+	/** Returns the top-level category name (the part before the first '|', trimmed). */
+	FString GetTopLevelCategory(const FString& Category)
+	{
+		int32 PipeIndex;
+		if (Category.FindChar(TEXT('|'), PipeIndex))
+		{
+			return Category.Left(PipeIndex).TrimEnd();
+		}
+		return Category;
+	}
+
+	/** Returns the last pipe-segment of a category string, trimmed (used as button label). */
+	FString GetSectionLabel(const FString& SectionId)
+	{
+		int32 PipeIndex;
+		if (SectionId.FindLastChar(TEXT('|'), PipeIndex))
+		{
+			return SectionId.RightChop(PipeIndex + 1).TrimStart();
+		}
+		return SectionId;
+	}
+}
+
+// ---- FPVGrowerSettingsCustomization ----
+
+TSharedRef<IDetailCustomization> FPVGrowerSettingsCustomization::MakeInstance()
+{
+	return MakeShareable(new FPVGrowerSettingsCustomization());
+}
+
+void FPVGrowerSettingsCustomization::CustomizeDetails(IDetailLayoutBuilder& DetailBuilder)
+{
+	// ---- Gather GrowerParams children and build the hide/section lists dynamically ----
+	//
+	// This single pass over the GrowerParams struct members drives everything:
+	//   - Categories to suppress in the default auto-layout (HideCategory + HideProperty)
+	//   - The ordered, deduplicated list of toggle sections for the button bar
+	//
+	// Adding a new UPROPERTY to FPVGrowerParams with a new Category value will
+	// automatically create a new toggle button and show/hide that property correctly
+	// without any changes here.
+
+	TSharedRef<IPropertyHandle> GrowerParamsHandle =
+		DetailBuilder.GetProperty(GET_MEMBER_NAME_CHECKED(UPVGrowerSettings, GrowerParams));
+	DetailBuilder.HideProperty(GrowerParamsHandle);
+
+	uint32 NumGrowerChildren = 0;
+	GrowerParamsHandle->GetNumChildren(NumGrowerChildren);
+
+	TArray<FString> DiscoveredSections; // ordered, deduplicated top-level category names
+
+	for (uint32 i = 0; i < NumGrowerChildren; ++i)
+	{
+		TSharedPtr<IPropertyHandle> Child = GrowerParamsHandle->GetChildHandle(i);
+		if (!Child.IsValid())
+		{
+			continue;
+		}
+
+		const FProperty* Prop = Child->GetProperty();
+		if (!Prop)
+		{
+			continue;
+		}
+
+		const FString Category = Prop->GetMetaData(TEXT("Category"));
+		if (Category.IsEmpty())
+		{
+			continue; // Skip non-editable / uncategorised properties (e.g. GuideSettings)
+		}
+
+		DetailBuilder.HideProperty(Child.ToSharedRef());
+		DetailBuilder.HideCategory(FName(*Category));
+
+		const FString TopLevel = GetTopLevelCategory(Category);
+		DetailBuilder.HideCategory(FName(*TopLevel));
+
+		// Basic properties are always visible — no toggle button for them.
+		if (TopLevel != TEXT("Basic"))
+		{
+			DiscoveredSections.AddUnique(Category);
+		}
+
+		// Recurse one level for ShowOnlyInnerProperties sub-structs so their inlined
+		// children are also suppressed in the auto-layout.
+		uint32 NumSubChildren = 0;
+		Child->GetNumChildren(NumSubChildren);
+		for (uint32 j = 0; j < NumSubChildren; ++j)
+		{
+			TSharedPtr<IPropertyHandle> SubChild = Child->GetChildHandle(j);
+			if (!SubChild.IsValid())
+			{
+				continue;
+			}
+
+			const FProperty* SubProp = SubChild->GetProperty();
+			if (SubProp)
+			{
+				const FString SubCategory = SubProp->GetMetaData(TEXT("Category"));
+				if (!SubCategory.IsEmpty())
+				{
+					DetailBuilder.HideCategory(FName(*SubCategory));
+				}
+			}
+
+			// Children whose parent category uses a pipe-subcategory path (e.g.
+			// "Senesence | Age Senescence") have ShowOnlyInnerProperties and the same
+			// category as the parent, so their sub-children get promoted directly into
+			// the auto-layout.  Hide them here; AddProperty below overrides HideProperty
+			// for explicit additions.
+			if (Category.Contains(TEXT("|")))
+			{
+				DetailBuilder.HideProperty(SubChild.ToSharedRef());
+			}
+		}
+	}
+
+	// Preset is a UPVGrowerSettings property (not part of GrowerParams) so hide it
+	// separately and add "Preset" as the first section manually.
+	TSharedRef<IPropertyHandle> PresetHandle =
+		DetailBuilder.GetProperty(GET_MEMBER_NAME_CHECKED(UPVGrowerSettings, Preset));
+	DetailBuilder.HideCategory(TEXT("Preset"));
+	DetailBuilder.HideProperty(PresetHandle);
+
+	// Ensure ActiveSection is still valid after a potential schema change.
+	if (ActiveSection != TEXT("All") && ActiveSection != TEXT("Preset") && !DiscoveredSections.Contains(ActiveSection) && ActiveSection != TEXT("Basic"))
+	{
+		ActiveSection = DiscoveredSections.IsEmpty() ? TEXT("Preset") : DiscoveredSections[0];
+	}
+
+	// ---- Settings category: button bar + section content ----
+
+	IDetailCategoryBuilder& SettingsCategory =
+		DetailBuilder.EditCategory(TEXT("Settings"), FText::GetEmpty(), ECategoryPriority::Important);
+
+	// Build toggle button bar from the discovered sections.
+	// "Preset" is prepended as it is a fixed UPVGrowerSettings concern.
+	TSharedRef<SWrapBox> ButtonBar = SNew(SWrapBox).UseAllottedSize(true);
+
+	auto AddSectionButton = [&ButtonBar](const FString& SectionName, const FText& Label)
+	{
+		ButtonBar->AddSlot()
+		.Padding(2.f, 2.f)
+		[
+			SNew(SCheckBox)
+			.Style(FAppStyle::Get(), "DetailsView.SectionButton")
+			.OnCheckStateChanged_Lambda([SectionName](ECheckBoxState)
+			{
+				ActiveSection = SectionName;
+			})
+			.IsChecked_Lambda([SectionName]()
+			{
+				return ActiveSection == SectionName
+					? ECheckBoxState::Checked
+					: ECheckBoxState::Unchecked;
+			})
+			[
+				SNew(STextBlock)
+				.TextStyle(FAppStyle::Get(), "SmallText")
+				.Text(Label)
+			]
+		];
+	};
+
+	AddSectionButton(TEXT("All"),    LOCTEXT("All",    "All"));
+	AddSectionButton(TEXT("Preset"), LOCTEXT("Preset", "Preset"));
+	for (const FString& Section : DiscoveredSections)
+	{
+		AddSectionButton(Section, FText::FromString(GetSectionLabel(Section)));
+	}
+
+	SettingsCategory.AddCustomRow(FText::GetEmpty())
+	.WholeRowContent()
+	[
+		ButtonBar
+	];
+
+	// ---- Preset section content ----
+
+	{
+		auto PresetVis = MakeAttributeLambda([]() -> EVisibility
+		{
+			return (ActiveSection == TEXT("Preset") || ActiveSection == TEXT("All"))
+				? EVisibility::Visible : EVisibility::Collapsed;
+		});
+
+		SettingsCategory.AddProperty(PresetHandle)
+			.ShouldAutoExpand(true)
+			.Visibility(PresetVis);
+
+		TArray<TWeakObjectPtr<UObject>> SelectedObjects;
+		DetailBuilder.GetObjectsBeingCustomized(SelectedObjects);
+
+		SettingsCategory.AddCustomRow(LOCTEXT("SaveToPreset", "Save to Preset"))
+		.Visibility(PresetVis)
+		.WholeRowContent()
+		[
+			SNew(SButton)
+			.Text(LOCTEXT("SaveToPreset", "Save to Preset"))
+			.OnClicked_Lambda([SelectedObjects]()
+			{
+				for (const TWeakObjectPtr<UObject>& WeakObj : SelectedObjects)
+				{
+					if (UObject* Obj = WeakObj.Get())
+					{
+						if (UFunction* Func = Obj->FindFunction(FName("SaveToPreset")))
+						{
+							Obj->ProcessEvent(Func, nullptr);
+						}
+					}
+				}
+				return FReply::Handled();
+			})
+		];
+
+		SettingsCategory.AddCustomRow(LOCTEXT("ApplyToSelectedPreset", "Apply to Selected Preset"))
+		.Visibility(PresetVis)
+		.WholeRowContent()
+		[
+			SNew(SButton)
+			.Text(LOCTEXT("ApplyToSelectedPreset", "Apply to Selected Preset"))
+			.OnClicked_Lambda([SelectedObjects]()
+			{
+				for (const TWeakObjectPtr<UObject>& WeakObj : SelectedObjects)
+				{
+					if (UObject* Obj = WeakObj.Get())
+					{
+						if (UFunction* Func = Obj->FindFunction(FName("ApplyToSelectedPreset")))
+						{
+							Obj->ProcessEvent(Func, nullptr);
+						}
+					}
+				}
+				return FReply::Handled();
+			})
+		];
+	}
+
+	// ---- GrowerParams section content ----
+	//
+	// Every child is added with a Visibility binding so switching sections only
+	// requires Slate to re-evaluate the attribute — no panel rebuild needed.
+
+	for (uint32 i = 0; i < NumGrowerChildren; ++i)
+	{
+		TSharedPtr<IPropertyHandle> ChildHandle = GrowerParamsHandle->GetChildHandle(i);
+		if (!ChildHandle.IsValid())
+		{
+			continue;
+		}
+
+		const FProperty* Prop = ChildHandle->GetProperty();
+		if (!Prop)
+		{
+			continue;
+		}
+
+		const FString Category = Prop->GetMetaData(TEXT("Category"));
+		if (Category.IsEmpty())
+		{
+			continue;
+		}
+
+		const FString TopLevel = GetTopLevelCategory(Category);
+
+		// Basic properties are always shown regardless of the active section.
+		if (TopLevel == TEXT("Basic"))
+		{
+			SettingsCategory.AddProperty(ChildHandle.ToSharedRef())
+				.ShouldAutoExpand(true);
+			continue;
+		}
+
+		auto VisAttr = MakeAttributeLambda([Category]() -> EVisibility
+		{
+			return (Category == ActiveSection || ActiveSection == TEXT("All"))
+				? EVisibility::Visible : EVisibility::Collapsed;
+		});
+
+		if (Category.Contains(TEXT("|")))
+		{
+			// Pipe-subcategory child (e.g. AgeSenescence, LightSenescence):
+			// sub-children were hidden above — re-add them individually so they
+			// appear as flat rows in the active section.
+			uint32 NumSubChildren = 0;
+			ChildHandle->GetNumChildren(NumSubChildren);
+			for (uint32 j = 0; j < NumSubChildren; ++j)
+			{
+				TSharedPtr<IPropertyHandle> SubChild = ChildHandle->GetChildHandle(j);
+				if (SubChild.IsValid())
+				{
+					SettingsCategory.AddProperty(SubChild.ToSharedRef())
+						.Visibility(VisAttr);
+				}
+			}
+		}
+		else
+		{
+			SettingsCategory.AddProperty(ChildHandle.ToSharedRef())
+				.ShouldAutoExpand(true)
+				.Visibility(VisAttr);
+		}
+	}
+}
+
+#undef LOCTEXT_NAMESPACE
diff --git a/Engine/Plugins/Experimental/ProceduralVegetationEditor/Source/ProceduralVegetationEditor/Private/Customizations/PVGrowerSettingsCustomization.h b/Engine/Plugins/Experimental/ProceduralVegetationEditor/Source/ProceduralVegetationEditor/Private/Customizations/PVGrowerSettingsCustomization.h
new file mode 100644
--- /dev/null
+++ b/Engine/Plugins/Experimental/ProceduralVegetationEditor/Source/ProceduralVegetationEditor/Private/Customizations/PVGrowerSettingsCustomization.h
@@ -0,0 +1,18 @@
+// Copyright Epic Games, Inc. All Rights Reserved.
+
+#pragma once
+
+#include "IDetailCustomization.h"
+
+class IDetailLayoutBuilder;
+
+class FPVGrowerSettingsCustomization : public IDetailCustomization
+{
+public:
+	static TSharedRef<IDetailCustomization> MakeInstance();
+	virtual void CustomizeDetails(IDetailLayoutBuilder& DetailBuilder) override;
+
+private:
+	// Persists across panel instances; Slate visibility lambdas read it every tick.
+	static FString ActiveSection;
+};
diff --git a/Engine/Plugins/Experimental/ProceduralVegetationEditor/Source/ProceduralVegetationEditor/Private/ProceduralVegetationEditorModule.cpp b/Engine/Plugins/Experimental/ProceduralVegetationEditor/Source/ProceduralVegetationEditor/Private/ProceduralVegetationEditorModule.cpp
--- a/Engine/Plugins/Experimental/ProceduralVegetationEditor/Source/ProceduralVegetationEditor/Private/ProceduralVegetationEditorModule.cpp
+++ b/Engine/Plugins/Experimental/ProceduralVegetationEditor/Source/ProceduralVegetationEditor/Private/ProceduralVegetationEditorModule.cpp
@@ -41,11 +41,13 @@
 #include "MessageLogModule.h"
 
 #include "Customizations/PVColliderParamsCustomization.h"
+#include "Customizations/PVGrowerSettingsCustomization.h"
 #include "Customizations/PVLoopDebugStepperCustomization.h"
 #include "Customizations/PVTextureSetupSettingsCustomization.h"
 
 #include "Helpers/PVUtilities.h"
 
+#include "Nodes/PVGrowerSettings.h"
 #include "Nodes/PVObjectInteractionSettings.h"
 
 #include "Templates/SharedPointer.h"
@@ -77,6 +79,7 @@ void FProceduralVegetationEditorModule::StartupModule()
 	PropertyEditor.RegisterCustomPropertyTypeLayout(FLoopDebugStepper::StaticStruct()->GetFName(), FOnGetPropertyTypeCustomizationInstance::CreateStatic(&FPVLoopDebugStepperCustomization::MakeInstance));
 
 	PropertyEditor.RegisterCustomClassLayout(UPVTextureSetupSettings::StaticClass()->GetFName(), FOnGetDetailCustomizationInstance::CreateStatic(&FPVTextureSetupSettingsCustomization::MakeInstance));
+	PropertyEditor.RegisterCustomClassLayout(UPVGrowerSettings::StaticClass()->GetFName(), FOnGetDetailCustomizationInstance::CreateStatic(&FPVGrowerSettingsCustomization::MakeInstance));
 	
 	FMessageLogModule& MessageLogModule = FModuleManager::LoadModuleChecked<FMessageLogModule>("MessageLog");
 	MessageLogModule.RegisterLogListing(PVEditor::MessageLogName, LOCTEXT("PVEditorLogLabel", "Procedural Vegetation Editor"));
@@ -97,7 +100,7 @@ void FProceduralVegetationEditorModule::ShutdownModule()
 			PropertyEditor->UnregisterCustomPropertyTypeLayout(FPVFloatRamp::StaticStruct()->GetFName());
 			PropertyEditor->UnregisterCustomPropertyTypeLayout(FLoopDebugStepper::StaticStruct()->GetFName());
 			PropertyEditor->UnregisterCustomClassLayout(UPVTextureSetupSettings::StaticClass()->GetFName());
-	
+			PropertyEditor->UnregisterCustomClassLayout(UPVGrowerSettings::StaticClass()->GetFName());
 		}
 	}
 
