[PVE] Enable conditional nodes from PCG for testing.

#Synced-CL 52095318

diff --git a/Engine/Plugins/Experimental/ProceduralVegetationEditor/Source/ProceduralVegetationEditor/Private/PVEditorSchema.cpp b/Engine/Plugins/Experimental/ProceduralVegetationEditor/Source/ProceduralVegetationEditor/Private/PVEditorSchema.cpp
--- a/Engine/Plugins/Experimental/ProceduralVegetationEditor/Source/ProceduralVegetationEditor/Private/PVEditorSchema.cpp
+++ b/Engine/Plugins/Experimental/ProceduralVegetationEditor/Source/ProceduralVegetationEditor/Private/PVEditorSchema.cpp
@@ -5,6 +5,9 @@
 #include "PCGParamData.h"
 #include "PCGSubgraph.h"
 #include "Elements/PCGCreateAttribute.h"
+#include "Elements/ControlFlow/PCGBranch.h"
+#include "Elements/ControlFlow/PCGMultiSelect.h"
+#include "Elements/ControlFlow/PCGSwitch.h"
 #include "Elements/Metadata/PCGMetadataMathsOpElement.h"
 #include "Nodes/PVBaseSettings.h"
 #include "Nodes/PVPointScatterSettings.h"
@@ -32,6 +35,9 @@ const FPCGGraphEditorCustomization& UPVEditorSchema::GetGraphEditorCustomization
 			UPCGMetadataMathsSettings::StaticClass(),
 			UPVSubgraphSettings::StaticClass(),
 			UPVPointScatterSettings::StaticClass(),
+			UPCGBranchSettings::StaticClass(),
+			UPCGSwitchSettings::StaticClass(),
+			UPCGMultiSelectSettings::StaticClass()
 		}
 	};
 
