Fix crash in UPVData::Initialize when passed a FGeometryCollection by copying it into a FManagedArrayCollection before initialization
Delete FGeometryCollection and FTransformCollection overloads of Initialize to prevent the incorrect usage at compile time

#Synced-CL 52292167

diff --git a/Engine/Plugins/Experimental/ProceduralVegetationEditor/Source/ProceduralVegetation/Private/Nodes/PVMeshBuilderSettings.cpp b/Engine/Plugins/Experimental/ProceduralVegetationEditor/Source/ProceduralVegetation/Private/Nodes/PVMeshBuilderSettings.cpp
--- a/Engine/Plugins/Experimental/ProceduralVegetationEditor/Source/ProceduralVegetation/Private/Nodes/PVMeshBuilderSettings.cpp
+++ b/Engine/Plugins/Experimental/ProceduralVegetationEditor/Source/ProceduralVegetation/Private/Nodes/PVMeshBuilderSettings.cpp
@@ -232,7 +232,9 @@ bool FPVMeshBuilderElement::ExecuteInternal(FPCGContext* InContext) const
 			FGeometryCollection OutGeometryCollection;
 			FPVMeshBuilder::GenerateGeometryCollection(SkeletonCollection, PlantProfileCollection, Settings->MesherSettings, OutGeometryCollection);
 
-			OutManagedArrayCollectionData->Initialize(MoveTemp(OutGeometryCollection));
+			FManagedArrayCollection ManagedArrayCollection;
+			OutGeometryCollection.CopyTo(&ManagedArrayCollection);
+			OutManagedArrayCollectionData->Initialize(MoveTemp(ManagedArrayCollection));
 
 			FPCGTaggedData& CollectionOutput = InContext->OutputData.TaggedData.Emplace_GetRef();
 			CollectionOutput.Data = OutManagedArrayCollectionData;
diff --git a/Engine/Plugins/Experimental/ProceduralVegetationEditor/Source/ProceduralVegetation/Public/DataTypes/PVData.h b/Engine/Plugins/Experimental/ProceduralVegetationEditor/Source/ProceduralVegetation/Public/DataTypes/PVData.h
--- a/Engine/Plugins/Experimental/ProceduralVegetationEditor/Source/ProceduralVegetation/Public/DataTypes/PVData.h
+++ b/Engine/Plugins/Experimental/ProceduralVegetationEditor/Source/ProceduralVegetation/Public/DataTypes/PVData.h
@@ -11,6 +11,8 @@
 #include "PVData.generated.h"
 
 struct FPCGContext;
+class FGeometryCollection;
+class FTransformCollection;
 
 UENUM()
 enum class EPVRenderType : uint8
@@ -200,6 +202,8 @@ public:
 	UPVData(const FObjectInitializer& ObjectInitializer);
 	
 	void Initialize(FManagedArrayCollection&& InCollection);
+	void Initialize(FGeometryCollection&&) = delete;
+	void Initialize(FTransformCollection&&) = delete;
 	
 	// ~Begin UPCGData interface
 	virtual EPCGDataType GetDataType() const override { return EPCGDataType::Other; }
