RotateAxillaryBranches UX pass

- Re-name "RotateAxillaryBranches" to "Rotate Branches"
- Add branch gradient and plant gradient curves for controlling the rotation amount along the plant
- Add "randomness" parameter to allow random variation of the rotation
- Fix issue where the bAlternatingRotations boolean would not properly alternate the roations of branches on the same parent

#Synced-CL 52338011

diff --git a/Engine/Plugins/Experimental/ProceduralVegetationEditor/Config/DefaultProceduralVegetationEditor.ini b/Engine/Plugins/Experimental/ProceduralVegetationEditor/Config/DefaultProceduralVegetationEditor.ini
--- a/Engine/Plugins/Experimental/ProceduralVegetationEditor/Config/DefaultProceduralVegetationEditor.ini
+++ b/Engine/Plugins/Experimental/ProceduralVegetationEditor/Config/DefaultProceduralVegetationEditor.ini
@@ -1,2 +1,4 @@
 ﻿[CoreRedirects]
-+ClassRedirects=(OldName="/Script/ProceduralVegetation.PVOutputSettings",NewName="/Script/ProceduralVegetation.PVExportSettings")
\ No newline at end of file
++ClassRedirects=(OldName="/Script/ProceduralVegetation.PVOutputSettings",NewName="/Script/ProceduralVegetation.PVExportSettings")
++StructRedirects=(OldName="/Script/ProceduralVegetation.PVRotateAxillaryBranchesParams",NewName="/Script/ProceduralVegetation.PVRotateBranchesParams")
++ClassRedirects=(OldName="/Script/ProceduralVegetation.PVRotateAxillaryBranchesSettings",NewName="/Script/ProceduralVegetation.PVRotateBranchesSettings")
\ No newline at end of file
diff --git a/Engine/Plugins/Experimental/ProceduralVegetationEditor/Source/ProceduralVegetation/Private/Implementations/PVRotateAxillaryBranches.cpp b/Engine/Plugins/Experimental/ProceduralVegetationEditor/Source/ProceduralVegetation/Private/Implementations/PVRotateAxillaryBranches.cpp
deleted file mode 100644
diff --git a/Engine/Plugins/Experimental/ProceduralVegetationEditor/Source/ProceduralVegetation/Private/Implementations/PVRotateAxillaryBranches.h b/Engine/Plugins/Experimental/ProceduralVegetationEditor/Source/ProceduralVegetation/Private/Implementations/PVRotateAxillaryBranches.h
deleted file mode 100644
diff --git a/Engine/Plugins/Experimental/ProceduralVegetationEditor/Source/ProceduralVegetation/Private/Implementations/PVRotateBranches.cpp b/Engine/Plugins/Experimental/ProceduralVegetationEditor/Source/ProceduralVegetation/Private/Implementations/PVRotateBranches.cpp
new file mode 100644
--- /dev/null
+++ b/Engine/Plugins/Experimental/ProceduralVegetationEditor/Source/ProceduralVegetation/Private/Implementations/PVRotateBranches.cpp
@@ -0,0 +1,117 @@
+// Copyright Epic Games, Inc. All Rights Reserved.
+
+#include "PVRotateBranches.h"
+#include "Helpers/PVTransformHelper.h"
+#include "Helpers/PVPlantTraversalHelper.h"
+#include "Helpers/PVAttributesHelper.h"
+
+void FPVRotateBranches::ApplyRotateBranches(const FPVRotateBranchesParams& InRotateBranchesParams, FManagedArrayCollection& OutCollection)
+{
+	using namespace PV::Transform;
+
+	const FRotateBranchPointsAttributeCollection Attributes(OutCollection);
+	if (!Attributes.IsValid())
+	{
+		return;
+	}
+
+	PV::FBudDevelopmentAttributeConstView BudDevelopmentAttribute = PV::FBudDevelopmentAttribute::FindAttribute(OutCollection);
+	if (!PV::ValidateAttributeCollection(BudDevelopmentAttribute))
+	{
+		return;
+	}
+
+	PV::FPointPlantGradientAttributeView PointPlantGradientAttribute = PV::FPointPlantGradientAttribute::AddAttribute(OutCollection);
+	PV::AttributesHelper::ComputePointPlantGradient({
+		PointPlantGradientAttribute,
+		Attributes.BranchPoints,
+		Attributes.BranchParentNumber,
+		Attributes.BranchNumber,
+		Attributes.BranchChildren
+	});
+
+
+	const int32 StartGeneration = InRotateBranchesParams.StartGeneration;
+	const int32 EndGeneration = InRotateBranchesParams.NumGenerations <= 0
+		? MAX_int32
+		: InRotateBranchesParams.StartGeneration + InRotateBranchesParams.NumGenerations;
+
+	FRandomStream RandomStream(InRotateBranchesParams.RandomSeed);
+	int32 CurrentBranchIndex = INDEX_NONE;
+
+	for (const int32 TrunkIndex : PV::PlantTraversalHelper::GetTrunkIndices(Attributes.BranchParentNumber))
+	{
+		const TArray<int32>& CurrentBranchPoints = Attributes.BranchPoints[TrunkIndex];
+		if (CurrentBranchPoints.Num() == 0)
+		{
+			continue;
+		}
+
+		struct FPointState
+		{
+			int32 PrevBranchIndex = INDEX_NONE;
+			int32 PrevBranchPointIndex = INDEX_NONE;
+			float RotationSign = 1;
+		};
+
+		const auto CalcRotationFunction = [&](const FRotateBranchPointParams& Params, FRotateBranchPointResult& Result, FPointState& PointState)
+		{
+			if (Params.BranchIndex != CurrentBranchIndex)
+			{
+				CurrentBranchIndex = Params.BranchIndex;
+			}
+
+			if (PointState.PrevBranchIndex != Params.BranchIndex && Params.BranchIndex != TrunkIndex)
+			{
+				const int32 Generation = BudDevelopmentAttribute[Params.PointIndex].Generation;
+				if (Generation < StartGeneration || Generation > EndGeneration)
+				{
+					return;
+				}
+
+				const float PlantGradient = PointPlantGradientAttribute[Params.ParentPointIndex];
+				const float BranchGradient = PV::AttributesHelper::GetBranchPointBranchGradient(Attributes.PointLengthFromRoot, Attributes.BranchPoints, PointState.PrevBranchIndex, PointState.PrevBranchPointIndex);
+				const float PlantGradientMultiplier = InRotateBranchesParams.PlantGradientMultiplier.EditorCurveData.Eval(1.f - PlantGradient);
+				const float BranchGradientMultiplier = InRotateBranchesParams.BranchGradientMultiplier.EditorCurveData.Eval(1.f - BranchGradient);
+
+				float BlendedMultiplier = 0.f;
+				switch (InRotateBranchesParams.GradientBlendMode)
+				{
+				case EPVGradientBlendMode::Multiply:
+					BlendedMultiplier = PlantGradientMultiplier * BranchGradientMultiplier;
+					break;
+				case EPVGradientBlendMode::Add:
+					BlendedMultiplier = FMath::Clamp(PlantGradientMultiplier + BranchGradientMultiplier, -1, 1);
+					break;
+				case EPVGradientBlendMode::Min:
+					BlendedMultiplier = FMath::Min(PlantGradientMultiplier, BranchGradientMultiplier);
+					break;
+				case EPVGradientBlendMode::Max:
+					BlendedMultiplier = FMath::Max(PlantGradientMultiplier, BranchGradientMultiplier);
+					break;
+				case EPVGradientBlendMode::Lerp:
+				default:
+					BlendedMultiplier = FMath::Lerp(PlantGradientMultiplier, BranchGradientMultiplier, InRotateBranchesParams.BranchGradientBias);
+					break;
+				}
+
+				const float RandomAngleRad = InRotateBranchesParams.Randomness > 0
+					? RandomStream.FRandRange(-UE_PI, UE_PI) * InRotateBranchesParams.Randomness
+					: 0.f;
+				const float RotationDeg = (FMath::DegreesToRadians(InRotateBranchesParams.Rotation * PointState.RotationSign) + RandomAngleRad) * BlendedMultiplier;
+				const FVector3f& ParentApical = Attributes.BudDirection[Params.ParentPointIndex].Apical;
+				Result.Rotation = FQuat4f(ParentApical, RotationDeg);
+			}
+
+			if (InRotateBranchesParams.bAlternatingRotations && Params.ChildBranchIndices.Num() > 0)
+			{
+				PointState.RotationSign *= -1;
+			}
+
+			PointState.PrevBranchIndex = Params.BranchIndex;
+			PointState.PrevBranchPointIndex = Params.BranchPointIndex;
+		};
+
+		RecursiveRotateBranchPoints<FPointState>(Attributes, TrunkIndex, FPointState(), CalcRotationFunction);
+	}
+}
diff --git a/Engine/Plugins/Experimental/ProceduralVegetationEditor/Source/ProceduralVegetation/Private/Implementations/PVRotateBranches.h b/Engine/Plugins/Experimental/ProceduralVegetationEditor/Source/ProceduralVegetation/Private/Implementations/PVRotateBranches.h
new file mode 100644
--- /dev/null
+++ b/Engine/Plugins/Experimental/ProceduralVegetationEditor/Source/ProceduralVegetation/Private/Implementations/PVRotateBranches.h
@@ -0,0 +1,87 @@
+// Copyright Epic Games, Inc. All Rights Reserved.
+
+#pragma once
+
+#include "CoreMinimal.h"
+#include "GeometryCollection/ManagedArrayCollection.h"
+#include "Utils/PVFloatRamp.h"
+#include "PVRotateBranches.generated.h"
+
+UENUM()
+enum class EPVGradientBlendMode : uint8
+{
+	Multiply,
+	Add,
+	Min,
+	Max,
+	Lerp
+};
+
+USTRUCT()
+struct FPVRotateBranchesParams
+{
+	GENERATED_BODY()
+
+	UPROPERTY(EditAnywhere, Category = "Settings", meta = (Units = Degrees, ClampMin = -180.0f, ClampMax = 180.0f, UIMin = -180.0f, UIMax = 180.0f))
+	float Rotation = 0;
+
+	// The generation to start applying the rotation
+	UPROPERTY(EditAnywhere, Category = "Settings", meta = (ClampMin = 1, UIMin = 1, UIMax = 20))
+	int32 StartGeneration = 1;
+
+	// The number of generations to apply the rotation for (0 = no limit)
+	UPROPERTY(EditAnywhere, Category = "Settings", meta = (ClampMin = 0, UIMin = 0, UIMax = 20))
+	int32 NumGenerations = 0;
+
+	// Scales the rotation amount based on the branch's position along the whole plant.
+	// The X axis is the plant gradient (0 = root, 1 = tip) and the Y axis is the rotation multiplier.
+	UPROPERTY(EditAnywhere, Category="Settings", meta=(XAxisMin=0.0f, XAxisMax=1.0f, YAxisMin=-1.0f, YAxisMax=1.0f))
+	FPVFloatRamp PlantGradientMultiplier;
+
+	// Scales the rotation amount based on the branch's position along its parent branch.
+	// The X axis is the branch gradient (0 = branch root, 1 = branch tip) and the Y axis is the rotation multiplier.
+	UPROPERTY(EditAnywhere, Category="Settings", meta=(XAxisMin=0.0f, XAxisMax=1.0f, YAxisMin=-1.0f, YAxisMax=1.0f))
+	FPVFloatRamp BranchGradientMultiplier;
+
+	// How PlantGradientMultiplier and BranchGradientMultiplier are combined.
+	UPROPERTY(EditAnywhere, Category="Settings")
+	EPVGradientBlendMode GradientBlendMode = EPVGradientBlendMode::Lerp;
+
+	// 0=use only plant gradient, 1=use only branch gradient
+	UPROPERTY(EditAnywhere, Category="Settings", meta=(ClampMin=0, ClampMax=1, UIMin=0, UIMax=1, EditCondition="GradientBlendMode==EPVGradientBlendMode::Lerp", EditConditionHides))
+	float BranchGradientBias = 0.f;
+
+	// When enabled, successive axillary branches alternate the rotation direction (e.g. +angle, -angle, +angle, ...).
+	UPROPERTY(EditAnywhere, Category = "Settings")
+	bool bAlternatingRotations = true;
+
+	// How much random variation is applied to the rotation angle (0 = no randomness, 1 = full random offset).
+	UPROPERTY(EditAnywhere, Category = "Settings", meta = (ClampMin = 0, ClampMax = 1, UIMin = 0, UIMax = 1))
+	float Randomness = 0.f;
+
+	// Seed used to generate the random rotation offsets. Change this to get a different random variation.
+	UPROPERTY(EditAnywhere, Category = "Settings")
+	int32 RandomSeed = 0;
+
+	FPVRotateBranchesParams()
+	{
+		InitializeLinearCurve(PlantGradientMultiplier);
+		InitializeLinearCurve(BranchGradientMultiplier);
+	}
+
+private:
+	static void InitializeLinearCurve(FPVFloatRamp& Curve)
+	{
+		Curve.EditorCurveData.Reset();
+		const FKeyHandle Handle0 = Curve.EditorCurveData.AddKey(0.0f, 0.0f);
+		const FKeyHandle Handle1 = Curve.EditorCurveData.AddKey(1.0f, 1.0f);
+
+		Curve.EditorCurveData.SetKeyInterpMode(Handle0, RCIM_Linear);
+		Curve.EditorCurveData.SetKeyInterpMode(Handle1, RCIM_Linear);
+	}
+};
+
+struct FPVRotateBranches
+{
+	static void ApplyRotateBranches(const FPVRotateBranchesParams& InRotateBranchesParams, FManagedArrayCollection& OutCollection);
+};
diff --git a/Engine/Plugins/Experimental/ProceduralVegetationEditor/Source/ProceduralVegetation/Private/Nodes/PVRotateAxillaryBranchesSettings.cpp b/Engine/Plugins/Experimental/ProceduralVegetationEditor/Source/ProceduralVegetation/Private/Nodes/PVRotateAxillaryBranchesSettings.cpp
deleted file mode 100644
diff --git a/Engine/Plugins/Experimental/ProceduralVegetationEditor/Source/ProceduralVegetation/Private/Nodes/PVRotateAxillaryBranchesSettings.h b/Engine/Plugins/Experimental/ProceduralVegetationEditor/Source/ProceduralVegetation/Private/Nodes/PVRotateAxillaryBranchesSettings.h
deleted file mode 100644
diff --git a/Engine/Plugins/Experimental/ProceduralVegetationEditor/Source/ProceduralVegetation/Private/Nodes/PVRotateBranchesSettings.cpp b/Engine/Plugins/Experimental/ProceduralVegetationEditor/Source/ProceduralVegetation/Private/Nodes/PVRotateBranchesSettings.cpp
new file mode 100644
--- /dev/null
+++ b/Engine/Plugins/Experimental/ProceduralVegetationEditor/Source/ProceduralVegetation/Private/Nodes/PVRotateBranchesSettings.cpp
@@ -0,0 +1,74 @@
+// Copyright Epic Games, Inc. All Rights Reserved.
+
+#include "PVRotateBranchesSettings.h"
+#include "DataTypes/PVData.h"
+#include "DataTypes/PVGrowthData.h"
+#include "Data/PCGBasePointData.h"
+#include "GeometryCollection/ManagedArrayCollection.h"
+#include "Helpers/PCGSettingsHelpers.h"
+
+#define LOCTEXT_NAMESPACE "PVRotateBranchesSettings"
+
+#if WITH_EDITOR
+FText UPVRotateBranchesSettings::GetDefaultNodeTitle() const
+{
+	return LOCTEXT("NodeTitle", "RotateBranches");
+}
+
+FText UPVRotateBranchesSettings::GetNodeTooltipText() const
+{
+	return LOCTEXT("NodeTooltip",
+		"Rotates the axillary branches around the apical direction of its parent branch."
+		"\n\nPress Ctrl + L to lock/unlock node output"
+	); 
+}
+#endif
+
+FPCGDataTypeIdentifier UPVRotateBranchesSettings::GetInputPinTypeIdentifier() const
+{
+	return FPCGDataTypeIdentifier{ FPVDataTypeInfoGrowth::AsId() };
+}
+
+FPCGDataTypeIdentifier UPVRotateBranchesSettings::GetOutputPinTypeIdentifier() const
+{
+	return FPCGDataTypeIdentifier{ FPVDataTypeInfoGrowth::AsId() };
+}
+
+FPCGElementPtr UPVRotateBranchesSettings::CreateElement() const
+{
+	return MakeShared<FPVRotateBranchesElement>();
+}
+
+bool FPVRotateBranchesElement::ExecuteInternal(FPCGContext* InContext) const
+{
+	TRACE_CPUPROFILER_EVENT_SCOPE(FPVRotateBranchesElement::Execute);
+
+	check(InContext);
+
+	const UPVRotateBranchesSettings* Settings = InContext->GetInputSettings<UPVRotateBranchesSettings>();
+	check(Settings);
+
+	for (const FPCGTaggedData& Input : InContext->InputData.GetInputsByPin(PCGPinConstants::DefaultInputLabel))
+	{
+		if (const UPVGrowthData* InputData = Cast<UPVGrowthData>(Input.Data))
+		{
+			FManagedArrayCollection OutCollection = InputData->GetCollection();
+
+			UPVGrowthData* OutManagedArrayCollectionData = FPCGContext::NewObject_AnyThread<UPVGrowthData>(InContext);
+
+			FPVRotateBranches::ApplyRotateBranches(Settings->RotateBranchesParams, OutCollection);
+
+			OutManagedArrayCollectionData->Initialize(MoveTemp(OutCollection));
+			InContext->OutputData.TaggedData.Emplace(OutManagedArrayCollectionData);
+		}
+		else
+		{
+			PCGLog::InputOutput::LogInvalidInputDataError(InContext);
+			return true;
+		}
+	}
+
+	return true;
+}
+
+#undef LOCTEXT_NAMESPACE
\ No newline at end of file
diff --git a/Engine/Plugins/Experimental/ProceduralVegetationEditor/Source/ProceduralVegetation/Private/Nodes/PVRotateBranchesSettings.h b/Engine/Plugins/Experimental/ProceduralVegetationEditor/Source/ProceduralVegetation/Private/Nodes/PVRotateBranchesSettings.h
new file mode 100644
--- /dev/null
+++ b/Engine/Plugins/Experimental/ProceduralVegetationEditor/Source/ProceduralVegetation/Private/Nodes/PVRotateBranchesSettings.h
@@ -0,0 +1,39 @@
+// Copyright Epic Games, Inc. All Rights Reserved.
+#pragma once
+
+#include "CoreMinimal.h"
+#include "Nodes/PVBaseSettings.h"
+#include "Implementations/PVRotateBranches.h"
+#include "PVRotateBranchesSettings.generated.h"
+
+UCLASS(BlueprintType, ClassGroup = (Procedural))
+class UPVRotateBranchesSettings : public UPVBaseSettings
+{
+	GENERATED_BODY()
+
+public:
+#if WITH_EDITOR
+	virtual FName GetDefaultNodeName() const override { return FName(TEXT("ProceduralVegetationRotateBranches")); }
+	virtual FText GetDefaultNodeTitle() const override;
+	virtual FText GetNodeTooltipText() const override;
+	virtual FLinearColor GetNodeTitleColor() const override { return FLinearColor::Yellow; }
+
+	virtual TArray<EPVRenderType> GetDefaultRenderType() const override { return TArray{ EPVRenderType::PointData }; }
+#endif
+
+protected:
+	virtual FPCGDataTypeIdentifier GetInputPinTypeIdentifier() const override;
+	virtual FPCGDataTypeIdentifier GetOutputPinTypeIdentifier() const override;
+
+	virtual FPCGElementPtr CreateElement() const override;
+
+public:
+	UPROPERTY(EditAnywhere, Category = "RotateBranches", meta = (ShowOnlyInnerProperties, PCG_Overridable))
+	FPVRotateBranchesParams RotateBranchesParams;
+};
+
+class FPVRotateBranchesElement : public FPVBaseElement
+{
+protected:
+	virtual bool ExecuteInternal(FPCGContext* InContext) const override;
+};
