[PVE] Remove redundant BranchPoints sorting and add ensure in ParametricDistributionHelper and HormoneDistributionHelper. It is generally expected that branch points be sorted. If they are not, then thats a bug and we should try and fix it.

#Synced-CL 52095318

diff --git a/Engine/Plugins/Experimental/ProceduralVegetationEditor/Source/ProceduralVegetation/Private/Helpers/PVHormoneDistributionHelper.cpp b/Engine/Plugins/Experimental/ProceduralVegetationEditor/Source/ProceduralVegetation/Private/Helpers/PVHormoneDistributionHelper.cpp
--- a/Engine/Plugins/Experimental/ProceduralVegetationEditor/Source/ProceduralVegetation/Private/Helpers/PVHormoneDistributionHelper.cpp
+++ b/Engine/Plugins/Experimental/ProceduralVegetationEditor/Source/ProceduralVegetation/Private/Helpers/PVHormoneDistributionHelper.cpp
@@ -47,6 +47,14 @@ void PV::HormoneDistributionHelper::ComputeHormoneBasedAttachmentPoints(
 
 		if (BranchPoints.Num() < 2) continue;
 
+#if DO_ENSURE
+		for (int32 i = 0; i < BranchPoints.Num() - 1; ++i)
+		{
+			ensureMsgf(PointLengthFromRootAttribute[BranchPoints[i]] <= PointLengthFromRootAttribute[BranchPoints[i + 1]],
+				TEXT("BranchPoints are not ordered root-to-tip at index %d for BranchIndex %d"), i, BranchIndex);
+		}
+#endif
+
 		const float FirstPointLengthFromRoot = PointLengthFromRootAttribute[BranchPoints[0]];
 		const float LastPointLengthFromRoot = PointLengthFromRootAttribute[BranchPoints.Last()];
 		const float BranchLength = LastPointLengthFromRoot - FirstPointLengthFromRoot;
diff --git a/Engine/Plugins/Experimental/ProceduralVegetationEditor/Source/ProceduralVegetation/Private/Helpers/PVParametricDistributionHelper.cpp b/Engine/Plugins/Experimental/ProceduralVegetationEditor/Source/ProceduralVegetation/Private/Helpers/PVParametricDistributionHelper.cpp
--- a/Engine/Plugins/Experimental/ProceduralVegetationEditor/Source/ProceduralVegetation/Private/Helpers/PVParametricDistributionHelper.cpp
+++ b/Engine/Plugins/Experimental/ProceduralVegetationEditor/Source/ProceduralVegetation/Private/Helpers/PVParametricDistributionHelper.cpp
@@ -178,14 +178,16 @@ void PV::ParametricDistributionHelper::ComputeAttachmentPoints(
 			continue;
 		}
 
-		TArray<int32> SortedPointIndices = BranchPoints;
-		SortedPointIndices.Sort([PointLengthFromRootAttribute](const int32& A, const int32& B)
-			{
-				return PointLengthFromRootAttribute[A] < PointLengthFromRootAttribute[B];
-			});
+#if DO_ENSURE
+		for (int32 i = 0; i < BranchPoints.Num() - 1; ++i)
+		{
+			ensureMsgf(PointLengthFromRootAttribute[BranchPoints[i]] <= PointLengthFromRootAttribute[BranchPoints[i + 1]],
+				TEXT("BranchPoints are not ordered root-to-tip at index %d for BranchIndex %d"), i, BranchIndex);
+		}
+#endif
 
-		const int32 FirstPointIndex = SortedPointIndices[0];
-		const int32 TipPointIndex = SortedPointIndices.Last();
+		const int32 FirstPointIndex = BranchPoints[0];
+		const int32 TipPointIndex   = BranchPoints.Last();
 		const float PlantGradientAtRoot = 1.0f - PointPlantGradientAttribute[FirstPointIndex];
 		const float LFRofRoot = PointLengthFromRootAttribute[FirstPointIndex];
 		const float LFRofTip = PointLengthFromRootAttribute[TipPointIndex];
@@ -243,13 +245,13 @@ void PV::ParametricDistributionHelper::ComputeAttachmentPoints(
 
 				Pos = SampleVectorByFloat(
 					PlantGradient,
-					SortedPointIndices,
+					BranchPoints,
 					PointPlantGradientAttribute,
 					PointPositionAttribute);
 
 				AttachmentPointLengthFromRoot = SampleFloatByFloat(
 					PlantGradient,
-					SortedPointIndices,
+					BranchPoints,
 					PointPlantGradientAttribute,
 					PointLengthFromRootAttribute);
 
@@ -260,7 +262,7 @@ void PV::ParametricDistributionHelper::ComputeAttachmentPoints(
 
 				FindPointByNormalizedLengthFromRoot(
 					BranchGradient,
-					SortedPointIndices,
+					BranchPoints,
 					PointLengthFromRootAttribute,
 					LFRofRoot,
 					LFRofTip,
@@ -284,7 +286,7 @@ void PV::ParametricDistributionHelper::ComputeAttachmentPoints(
 
 				if (!FindPointByNormalizedLengthFromRoot(
 					BranchGradient,
-					SortedPointIndices,
+					BranchPoints,
 					PointLengthFromRootAttribute,
 					LFRofRoot,
 					LFRofTip,
@@ -393,7 +395,7 @@ void PV::ParametricDistributionHelper::ComputeAttachmentPoints(
 
 				if (bIsTipInstance)
 				{
-					const int32 SecondToLastPointIndex = SortedPointIndices[SortedPointIndices.Num() - 2];
+					const int32 SecondToLastPointIndex = BranchPoints[BranchPoints.Num() - 2];
 					const FVector3f TipPosition = PointPositionAttribute[TipPointIndex];
 					const FVector3f SecondToLastPointPosition = PointPositionAttribute[SecondToLastPointIndex];
 					FVector3f TipTangent = (TipPosition - SecondToLastPointPosition).GetSafeNormal();
@@ -476,14 +478,17 @@ TMap<int32, float> PV::ParametricDistributionHelper::ComputeBranchLengths(FPoint
 			continue;
 		}
 
-		TArray<int32> SortedPointIndices = BranchPoints;
-		SortedPointIndices.Sort([PointLengthFromRootAttribute](const int32& A, const int32& B)
-			{
-				return PointLengthFromRootAttribute[A] < PointLengthFromRootAttribute[B];
-			});
+#if DO_ENSURE
+		for (int32 i = 0; i < BranchPoints.Num() - 1; ++i)
+		{
+			ensureMsgf(
+				PointLengthFromRootAttribute[BranchPoints[i]] <= PointLengthFromRootAttribute[BranchPoints[i + 1]],
+				TEXT("BranchPoints are not ordered root-to-tip at index %d for BranchIndex %d"), i, BranchIndex);
+		}
+#endif
 
-		const int32 RootPointIndex = SortedPointIndices[0];
-		const int32 BranchTipPointIndex = SortedPointIndices.Last();
+		const int32 RootPointIndex      = BranchPoints[0];
+		const int32 BranchTipPointIndex = BranchPoints.Last();
 
 		const float RootPointLFR = PointLengthFromRootAttribute[RootPointIndex];
 		const float TipPointLFR = PointLengthFromRootAttribute[BranchTipPointIndex];
@@ -495,30 +500,30 @@ TMap<int32, float> PV::ParametricDistributionHelper::ComputeBranchLengths(FPoint
 	return BranchIndicesToBranchLengths;
 }
 
-FVector3f PV::ParametricDistributionHelper::SampleVectorByFloat(const float SampleValue, const TArray<int32>& SortedPointIndices,
+FVector3f PV::ParametricDistributionHelper::SampleVectorByFloat(const float SampleValue, const TArray<int32>& BranchPoints,
                                                                 FPointPlantGradientAttributeConstView PointPlantGradientAttribute,
                                                                 FPointPositionAttributeConstView PointPositionAttribute)
 {
-	if (SortedPointIndices.Num() < 2)
+	if (BranchPoints.Num() < 2)
 	{
 		return FVector3f();
 	}
 
-	const int32 FirstPointIndex = SortedPointIndices[0];
-	const int32 LastPointIndex = SortedPointIndices.Last();
+	const int32 FirstPointIndex = BranchPoints[0];
+	const int32 LastPointIndex = BranchPoints.Last();
 	float Value0 = 1.0f - PointPlantGradientAttribute[FirstPointIndex];
 	float Value1 = 1.0f - PointPlantGradientAttribute[LastPointIndex];
 	FVector3f OutValue(Value0, Value0, Value0);
 
 	if (FMath::IsWithinInclusive(SampleValue, Value0, Value1))
 	{
-		for (int32 i = 1; i < SortedPointIndices.Num(); i++)
+		for (int32 i = 1; i < BranchPoints.Num(); i++)
 		{
-			Value1 = 1.0f - PointPlantGradientAttribute[SortedPointIndices[i]];
+			Value1 = 1.0f - PointPlantGradientAttribute[BranchPoints[i]];
 			if (FMath::IsWithinInclusive(SampleValue, Value0, Value1))
 			{
-				FVector3f P0 = PointPositionAttribute[SortedPointIndices[i - 1]];
-				FVector3f P1 = PointPositionAttribute[SortedPointIndices[i]];
+				FVector3f P0 = PointPositionAttribute[BranchPoints[i - 1]];
+				FVector3f P1 = PointPositionAttribute[BranchPoints[i]];
 
 				float Blend = FMath::GetMappedRangeValueClamped(
 					FVector2f(Value0, Value1),
@@ -542,30 +547,30 @@ FVector3f PV::ParametricDistributionHelper::SampleVectorByFloat(const float Samp
 	return OutValue;
 }
 
-float PV::ParametricDistributionHelper::SampleFloatByFloat(const float SampleValue, const TArray<int32>& SortedPointIndices,
+float PV::ParametricDistributionHelper::SampleFloatByFloat(const float SampleValue, const TArray<int32>& BranchPoints,
                                                            FPointPlantGradientAttributeConstView PointPlantGradientAttribute,
                                                            FPointLengthFromRootAttributeConstView PointLengthFromRootAttribute)
 {
-	if (SortedPointIndices.Num() < 2)
+	if (BranchPoints.Num() < 2)
 	{
 		return -1.0f;
 	}
 
-	const int32 FirstPointIndex = SortedPointIndices[0];
-	const int32 LastPointIndex = SortedPointIndices.Last();
+	const int32 FirstPointIndex = BranchPoints[0];
+	const int32 LastPointIndex = BranchPoints.Last();
 	float Value0 = 1.0f - PointPlantGradientAttribute[FirstPointIndex];
 	float Value1 = 1.0f - PointPlantGradientAttribute[LastPointIndex];
 	float OutValue = Value0;
 
 	if (FMath::IsWithinInclusive(SampleValue, Value0, Value1))
 	{
-		for (int32 i = 1; i < SortedPointIndices.Num(); i++)
+		for (int32 i = 1; i < BranchPoints.Num(); i++)
 		{
-			Value1 = 1.0f - PointPlantGradientAttribute[SortedPointIndices[i]];
+			Value1 = 1.0f - PointPlantGradientAttribute[BranchPoints[i]];
 			if (FMath::IsWithinInclusive(SampleValue, Value0, Value1))
 			{
-				float P0 = PointLengthFromRootAttribute[SortedPointIndices[i - 1]];
-				float P1 = PointLengthFromRootAttribute[SortedPointIndices[i]];
+				float P0 = PointLengthFromRootAttribute[BranchPoints[i - 1]];
+				float P1 = PointLengthFromRootAttribute[BranchPoints[i]];
 
 				float Blend = FMath::GetMappedRangeValueClamped(
 					FVector2f(Value0, Value1),
@@ -589,12 +594,12 @@ float PV::ParametricDistributionHelper::SampleFloatByFloat(const float SampleVal
 	return OutValue;
 }
 
-bool PV::ParametricDistributionHelper::FindPointByNormalizedLengthFromRoot(float SampleValue, const TArray<int32>& SortedPointIndices,
+bool PV::ParametricDistributionHelper::FindPointByNormalizedLengthFromRoot(float SampleValue, const TArray<int32>& BranchPoints,
                                                                            FPointLengthFromRootAttributeConstView PointLengthFromRootAttribute,
                                                                            const float LFRofRoot, const float LFRofTip, int32& PreviousPointIndex,
                                                                            int32& NextPointIndex, float& OutAlpha)
 {
-	if (SortedPointIndices.Num() < 2)
+	if (BranchPoints.Num() < 2)
 	{
 		return false;
 	}
@@ -606,10 +611,10 @@ bool PV::ParametricDistributionHelper::FindPointByNormalizedLengthFromRoot(float
 		FVector2f(LFRofRoot, LFRofTip),
 		SampleValue);
 
-	for (int32 i = 0; i < SortedPointIndices.Num() - 1; ++i)
+	for (int32 i = 0; i < BranchPoints.Num() - 1; ++i)
 	{
-		const int32 P = SortedPointIndices[i];
-		const int32 N = SortedPointIndices[i + 1];
+		const int32 P = BranchPoints[i];
+		const int32 N = BranchPoints[i + 1];
 		const float LFR0 = PointLengthFromRootAttribute[P];
 		const float LFR1 = PointLengthFromRootAttribute[N];
 
@@ -628,8 +633,8 @@ bool PV::ParametricDistributionHelper::FindPointByNormalizedLengthFromRoot(float
 		}
 	}
 
-	PreviousPointIndex = SortedPointIndices[0];
-	NextPointIndex = SortedPointIndices.Last();
+	PreviousPointIndex = BranchPoints[0];
+	NextPointIndex = BranchPoints.Last();
 	OutAlpha = (SampleValue < 0.5f)
 		? 0.0f
 		: 1.0f;
diff --git a/Engine/Plugins/Experimental/ProceduralVegetationEditor/Source/ProceduralVegetation/Private/Helpers/PVParametricDistributionHelper.h b/Engine/Plugins/Experimental/ProceduralVegetationEditor/Source/ProceduralVegetation/Private/Helpers/PVParametricDistributionHelper.h
--- a/Engine/Plugins/Experimental/ProceduralVegetationEditor/Source/ProceduralVegetation/Private/Helpers/PVParametricDistributionHelper.h
+++ b/Engine/Plugins/Experimental/ProceduralVegetationEditor/Source/ProceduralVegetation/Private/Helpers/PVParametricDistributionHelper.h
@@ -96,19 +96,19 @@ namespace PV::ParametricDistributionHelper
 
 	FVector3f SampleVectorByFloat(
 		const float SampleValue,
-		const TArray<int32>& SortedPointIndices,
+		const TArray<int32>& BranchPoints,
 		FPointPlantGradientAttributeConstView PointPlantGradientAttribute,
 		FPointPositionAttributeConstView PointPositionAttribute);
 
 	float SampleFloatByFloat(
 		const float SampleValue,
-		const TArray<int32>& SortedPointIndices,
+		const TArray<int32>& BranchPoints,
 		FPointPlantGradientAttributeConstView PointPlantGradientAttribute,
 		FPointLengthFromRootAttributeConstView PointLengthFromRootAttribute);
 
 	bool FindPointByNormalizedLengthFromRoot(
 		float SampleValue,
-		const TArray<int32>& SortedPointIndices,
+		const TArray<int32>& BranchPoints,
 		FPointLengthFromRootAttributeConstView PointLengthFromRootAttribute,
 		const float LFRofRoot,
 		const float LFRofTip,
