[PVE] Fixed Njord_pixelIdx issue. The current implementation was not setting the unique IDs for the branches. Reverting to the old pixel id implementation with new attribute system.

#Synced-CL 52095318

diff --git a/Engine/Plugins/Experimental/ProceduralVegetationEditor/Source/ProceduralVegetation/Private/Helpers/PVAttributesHelper.cpp b/Engine/Plugins/Experimental/ProceduralVegetationEditor/Source/ProceduralVegetation/Private/Helpers/PVAttributesHelper.cpp
--- a/Engine/Plugins/Experimental/ProceduralVegetationEditor/Source/ProceduralVegetation/Private/Helpers/PVAttributesHelper.cpp
+++ b/Engine/Plugins/Experimental/ProceduralVegetationEditor/Source/ProceduralVegetation/Private/Helpers/PVAttributesHelper.cpp
@@ -63,29 +63,26 @@ void PV::AttributesHelper::ComputeBudNumbers(const FComputeBudNumbersAttributes&
 
 void PV::AttributesHelper::ComputeNjordPixelIndex(const FComputeNjordPixelIndexAttributes& InAttributes)
 {
-	using namespace PV::PlantTraversalHelper;
+	if (!InAttributes.IsValid())
+	{
+		return;
+	}
 
-	RecursiveWalkBranches<int32>(
-		InAttributes.BranchParentNumber,
-		InAttributes.BranchChildren,
-		InAttributes.BranchNumber,
-		1, 
-		[&](int32 BranchIndex, int32& NjordPixelIndex)->EForEachResult
-		{
-			const bool bIsTrunk = IsTrunk(InAttributes.BranchParentNumber, BranchIndex);
-			const TArray<int32>& BranchPoints = InAttributes.BranchPoints[BranchIndex];
-			for (int32 i = 0; i < BranchPoints.Num(); ++i)
-			{
-				if (bIsTrunk || i != 0)
-				{
-					InAttributes.PointNjordPixelIndex[BranchPoints[i]] = NjordPixelIndex;
-					NjordPixelIndex += 1;
-				}
-			}
+	TMap<int32, TArray<int32>> BudNumberPointIndicesMap;
+	for (int32 PointIndex = 0; PointIndex < InAttributes.PointBudNumber.Num(); ++PointIndex)
+	{
+		const int32 BudNumber = InAttributes.PointBudNumber[PointIndex];
+		BudNumberPointIndicesMap.FindOrAdd(BudNumber).Add(PointIndex);
+	}
 
-			return EForEachResult::Continue;
+	for (const auto& [BudNumber, PointIndices] : BudNumberPointIndicesMap)
+	{
+		const int32 Count = PointIndices.Num();
+		for (int32 i = 0; i < Count; ++i)
+		{
+			InAttributes.PointNjordPixelIndex[PointIndices[i]] = BudNumber + static_cast<float>(i) / Count;
 		}
-	);
+	}
 }
 
 void PV::AttributesHelper::ComputeLengthFromRoot(const FComputeLengthFromRootAttributes& InAttributes)
@@ -948,11 +945,8 @@ bool PV::AttributesHelper::RecomputeAllGrowthDataAttributes(const FRecomputeGrow
 	});
 
 	ComputeNjordPixelIndex({
-		InAttributes.PointNjordPixelIndex,
-		InAttributes.BranchPoints,
-		InAttributes.BranchParentNumber,
-		InAttributes.BranchNumber,
-		InAttributes.BranchChildren
+		InAttributes.PointBudNumber,
+		InAttributes.PointNjordPixelIndex
 	});
 
 	ComputeLengthFromRoot({
diff --git a/Engine/Plugins/Experimental/ProceduralVegetationEditor/Source/ProceduralVegetation/Private/Helpers/PVAttributesHelper.h b/Engine/Plugins/Experimental/ProceduralVegetationEditor/Source/ProceduralVegetation/Private/Helpers/PVAttributesHelper.h
--- a/Engine/Plugins/Experimental/ProceduralVegetationEditor/Source/ProceduralVegetation/Private/Helpers/PVAttributesHelper.h
+++ b/Engine/Plugins/Experimental/ProceduralVegetationEditor/Source/ProceduralVegetation/Private/Helpers/PVAttributesHelper.h
@@ -45,37 +45,25 @@ namespace PV::AttributesHelper
 	struct FComputeNjordPixelIndexAttributes
 	{
 		PV::FPointNjordPixelIndexAttributeView PointNjordPixelIndex;
-		PV::FBranchPointsAttributeConstView BranchPoints;
-		PV::FBranchParentNumberAttributeConstView BranchParentNumber;
-		PV::FBranchNumberAttributeConstView BranchNumber;
-		PV::FBranchChildrenAttributeConstView BranchChildren;
+		FPointBudNumberAttributeView PointBudNumber;
 
 		FComputeNjordPixelIndexAttributes(
-			PV::FPointNjordPixelIndexAttributeView InPointNjordPixelIndex,
-			PV::FBranchPointsAttributeConstView InBranchPoints,
-			PV::FBranchParentNumberAttributeConstView InBranchParentNumber,
-			PV::FBranchNumberAttributeConstView InBranchNumber,
-			PV::FBranchChildrenAttributeConstView InBranchChildren)
+			FPointBudNumberAttributeView InPointBudNumber,
+			PV::FPointNjordPixelIndexAttributeView InPointNjordPixelIndex)
 			: PointNjordPixelIndex(InPointNjordPixelIndex)
-			, BranchPoints(InBranchPoints)
-			, BranchParentNumber(InBranchParentNumber)
-			, BranchNumber(InBranchNumber)
-			, BranchChildren(InBranchChildren)
+			, PointBudNumber(InPointBudNumber)
 		{
 		}
 
 		FComputeNjordPixelIndexAttributes(FManagedArrayCollection& InCollection)
 			: PointNjordPixelIndex(PV::FPointNjordPixelIndexAttribute::FindAttribute(InCollection))
-			, BranchPoints(PV::FBranchPointsAttribute::FindAttribute(InCollection))
-			, BranchParentNumber(PV::FBranchParentNumberAttribute::FindAttribute(InCollection))
-			, BranchNumber(PV::FBranchNumberAttribute::FindAttribute(InCollection))
-			, BranchChildren(PV::FBranchChildrenAttribute::FindAttribute(InCollection))
+			, PointBudNumber(FPointBudNumberAttribute::FindAttribute(InCollection))
 		{
 		}
 
 		bool IsValid() const
 		{
-			return ValidateAttributeCollection(PointNjordPixelIndex, BranchPoints, BranchParentNumber, BranchNumber, BranchChildren);
+			return ValidateAttributeCollection(PointNjordPixelIndex, PointBudNumber);
 		}
 	};
 
