[PVE] Fix UV for fused root branch point
- Extract UV computation into FPVMeshBuilder::ComputePointUVData to encapsulate fused-point handling
- For non-trunk branches, the 0th branch ring is the fused point shared with the parent; UV is now read from BranchPoints[1] and extrapolated backward to the child's origin

#Synced-CL 52379040

diff --git a/Engine/Plugins/Experimental/ProceduralVegetationEditor/Source/ProceduralVegetation/Private/Implementations/PVMaterialSettings.cpp b/Engine/Plugins/Experimental/ProceduralVegetationEditor/Source/ProceduralVegetation/Private/Implementations/PVMaterialSettings.cpp
--- a/Engine/Plugins/Experimental/ProceduralVegetationEditor/Source/ProceduralVegetation/Private/Implementations/PVMaterialSettings.cpp
+++ b/Engine/Plugins/Experimental/ProceduralVegetationEditor/Source/ProceduralVegetation/Private/Implementations/PVMaterialSettings.cpp
@@ -5,6 +5,7 @@
 #include "ProceduralVegetationModule.h"
 #include "Algo/MaxElement.h"
 #include "Facades/PVBranchFacade.h"
+#include "Facades/PVPlantFacade.h"
 #include "Facades/PVPointFacade.h"
 #include "Facades/PVTreeFacade.h"
 
@@ -324,6 +325,7 @@ void FPVMaterialSettings::SetOutputUVs(FManagedArrayCollection& Collection, TArr
 {
 	PV::Facades::FBranchFacade BranchFacade = PV::Facades::FBranchFacade(Collection);
 	PV::Facades::FPointFacade PointFacade = PV::Facades::FPointFacade(Collection);
+	PV::Facades::FPlantFacade PlantFacade = PV::Facades::FPlantFacade(Collection);
 	
 	for (int BranchIndex = 0; BranchIndex < BranchFacade.GetElementCount(); BranchIndex++ )
 	{
@@ -349,8 +351,11 @@ void FPVMaterialSettings::SetOutputUVs(FManagedArrayCollection& Collection, TArr
 		UVXTextureOffset = BranchUVData[BranchIndex].XOffset;
 		
 		auto& BranchPoints = BranchFacade.GetPoints(BranchIndex);
-	
-		for (int j=0; j < BranchPoints.Num(); j++)
+
+		const bool IsTrunk = PlantFacade.IsTrunkIndex(BranchIndex);
+		const int StartingIndex = IsTrunk ? 0 : 1;
+		
+		for (int j= StartingIndex; j < BranchPoints.Num(); j++)
 		{
 			FVector2f& UV = PointUVData[BranchPoints[j]].BaseUVs;
 			UV.Y *= -1;
diff --git a/Engine/Plugins/Experimental/ProceduralVegetationEditor/Source/ProceduralVegetation/Private/Implementations/PVMeshBuilder.cpp b/Engine/Plugins/Experimental/ProceduralVegetationEditor/Source/ProceduralVegetation/Private/Implementations/PVMeshBuilder.cpp
--- a/Engine/Plugins/Experimental/ProceduralVegetationEditor/Source/ProceduralVegetation/Private/Implementations/PVMeshBuilder.cpp
+++ b/Engine/Plugins/Experimental/ProceduralVegetationEditor/Source/ProceduralVegetation/Private/Implementations/PVMeshBuilder.cpp
@@ -845,6 +845,47 @@ TMap<int32, TArray<int32>> FPVMeshBuilder::GetPointsIndicesToFoliageIndicesMap(c
 	return Map;
 }
 
+FMeshPointUVData FPVMeshBuilder::ComputePointUVData(const int32 PointIndex, const int32 BranchPointIndex, const bool bPrimitiveIsTrunk,
+                                               const TArray<int32>& PrimitivePoints, const PV::Facades::FPointFacade& PointFacade,
+                                               const TManagedArray<FVector3f>& PointPositions, const float PointScale)
+{
+	// For non-trunk branches, the 0th branch ring is the fused junction shared with the parent.
+	// Its UV attributes belong to the parent, so read UV from BranchPoints[1] and extrapolate
+	// backward by the B0→B1 segment length to recover the child's origin UV.
+	const bool bRootBranchPoint = !bPrimitiveIsTrunk && BranchPointIndex == 0;
+
+#if DO_ENSURE
+	if (bRootBranchPoint)
+	{
+		ensureMsgf(PrimitivePoints.IsValidIndex(1),
+			TEXT("Non-trunk branch at PrimitiveIndex 0 has fewer than 2 points"));
+	}
+#endif
+
+	const int32 UVPointIndex = (bRootBranchPoint && PrimitivePoints.IsValidIndex(1)) ? PrimitivePoints[1] : PointIndex;
+
+	float TextureCoordV = PointFacade.GetTextureCoordV(UVPointIndex);
+	const float TextureCoordUOffset = PointFacade.GetTextureCoordUOffset(UVPointIndex);
+	const FVector2f URange = PointFacade.GetURange(UVPointIndex);
+
+	const float Interval = URange.GetMax() - URange.GetMin();
+	const float URatio = FMath::IsNearlyZero(Interval) ? 1.0f : 1.0f / Interval;
+	TextureCoordV /= URatio;
+
+	if (bRootBranchPoint && PrimitivePoints.IsValidIndex(1))
+	{
+		// TextureCoordV decreases from root to tip (SetOutputUVs applies UV.Y *= -1), so ring 0
+		// must have a higher value than ring 1 to sit "behind" it in texture space.
+		// We extrapolate backward by ADDING the B0→B1 segment's UV contribution to ring 1's value.
+		// TODO: We need to figure out why changing the root point flips all the branch UVs 
+		const float DistToFirstChild = FVector3f::Dist(PointPositions[PointIndex], PointPositions[UVPointIndex]);
+		const float SegLength = PointScale > 0.0f ? DistToFirstChild / (PointScale * UE_TWO_PI) : 0.0f;
+		TextureCoordV += SegLength * Interval;
+	}
+
+	return { TextureCoordV, TextureCoordUOffset, URange };
+}
+
 void FPVMeshBuilder::GenerateBranchMeshData(const bool bPrimitiveIsTrunk, const int32 GenerationNumber, const TArray<int32>& PrimitivePoints,
                                             const FDisplacementData& DisplacementData, const FPVMeshBuilderParams& MeshBuilderParams,
                                             const TArray<int32>& TargetMeshDivisions, const FManagedArrayCollection& Collection, const FManagedArrayCollection& InPlantProfileCollection,
@@ -906,19 +947,12 @@ void FPVMeshBuilder::GenerateBranchMeshData(const bool bPrimitiveIsTrunk, const
 
 		const float LengthFromRoot = PointFacade.GetLengthFromRoot(PointIndex);
 
-		float TextureCoordV = PointFacade.GetTextureCoordV(PointIndex);
-		float TextureCoordUOffset = PointFacade.GetTextureCoordUOffset(PointIndex);
-		FVector2f OldRange = FVector2f(0, 1.0f);
-		FVector2f NewRange = PointFacade.GetURange(PointIndex);
-
-		float Interval = NewRange.GetMax() - NewRange.GetMin();
-		float URatio = FMath::IsNearlyZero(Interval)
-			? 1.0f
-			: 1.0f / Interval;
-		TextureCoordV /= URatio;
-
 		const FVector3f& PointPosition = PointPositions[PointIndex];
 		const float PointScale = PointScales[PointIndex];
+		const FVector2f OldRange = FVector2f(0, 1.0f);
+
+		const auto [TextureCoordV, TextureCoordUOffset, NewRange] =
+			ComputePointUVData(PointIndex, Index, bPrimitiveIsTrunk, PrimitivePoints, PointFacade, PointPositions, PointScale);
 		
 		Displacement_UV_V += FVector3f::Dist(PointPosition, PreviousPosition) * MaxScaleRatio * (MaxPointScale / PointScale);
 		
diff --git a/Engine/Plugins/Experimental/ProceduralVegetationEditor/Source/ProceduralVegetation/Private/Implementations/PVMeshBuilder.h b/Engine/Plugins/Experimental/ProceduralVegetationEditor/Source/ProceduralVegetation/Private/Implementations/PVMeshBuilder.h
--- a/Engine/Plugins/Experimental/ProceduralVegetationEditor/Source/ProceduralVegetation/Private/Implementations/PVMeshBuilder.h
+++ b/Engine/Plugins/Experimental/ProceduralVegetationEditor/Source/ProceduralVegetation/Private/Implementations/PVMeshBuilder.h
@@ -289,6 +289,13 @@ public:
 	}
 };
 
+struct FMeshPointUVData
+{
+	float TextureCoordV;
+	float TextureCoordUOffset;
+	FVector2f URange;
+};
+
 struct FPVMeshBuilder
 {
 	static void GenerateGeometryCollection(const FManagedArrayCollection& InSkeletonCollection, const FManagedArrayCollection& InPlantProfileCollection, const FPVMeshBuilderParams& MeshBuilderParams,
@@ -333,6 +340,10 @@ private:
 
 	static TMap<int32, TArray<int32>> GetPointsIndicesToFoliageIndicesMap(const FManagedArrayCollection& Collection);
 
+	static FMeshPointUVData ComputePointUVData(int32 PointIndex, int32 BranchPointIndex, bool bPrimitiveIsTrunk,
+	                                       const TArray<int32>& PrimitivePoints, const PV::Facades::FPointFacade& PointFacade,
+	                                       const TManagedArray<FVector3f>& PointPositions, float PointScale);
+
 	static void GenerateBranchMeshData(const bool bPrimitiveIsTrunk, const int32 GenerationNumber, const TArray<int32>& PrimitivePoints,
 	                                   const FDisplacementData& DisplacementData, const FPVMeshBuilderParams& MeshBuilderParams,
 	                                   const TArray<int32>& TargetMeshDivisions, const FManagedArrayCollection& Collection, const FManagedArrayCollection& InPlantProfileCollection, 
