First draft of viewport interaction for setting plant root point and rotation

- Add custom tool for PVTexture2DImporter node
- Add custom gizmo to allow draging plant root point and rotate plants
- Allow disabling rendering of node output (Used to hide the not so helpful output of the 2D importer while still allowing rendering of the visualization collections)
- Move 2D importer visualizations to center of viewport to improve user experience

#Synced-CL 52391845

diff --git a/Engine/Plugins/Experimental/ProceduralVegetationEditor/Source/ProceduralVegetation/Private/Implementations/PVImporter_Texture2D.cpp b/Engine/Plugins/Experimental/ProceduralVegetationEditor/Source/ProceduralVegetation/Private/Implementations/PVImporter_Texture2D.cpp
--- a/Engine/Plugins/Experimental/ProceduralVegetationEditor/Source/ProceduralVegetation/Private/Implementations/PVImporter_Texture2D.cpp
+++ b/Engine/Plugins/Experimental/ProceduralVegetationEditor/Source/ProceduralVegetation/Private/Implementations/PVImporter_Texture2D.cpp
@@ -1014,8 +1014,7 @@ PV::Texture2DImport::EImportResult PV::Texture2DImport::ImportGrowthDataFromText
 		PerimeterAttributeDefinition.GetOrAttachAttribute(DynamicMesh)
 	);
 
-	const float HypotenuseLength = FMath::Sqrt(InParams.Scale * InParams.Scale * 2);
-	const FVector VisualizationOffset = FVector(0, -InParams.Scale / 2, HypotenuseLength);
+	const FVector VisualizationOffset = FVector(0, -InParams.Scale / 2, 25.f);
 
 	// Add end points visualization before applying root points as that transforms the DynamicMesh
 	AddDynamicMeshVerticesToCollection(
diff --git a/Engine/Plugins/Experimental/ProceduralVegetationEditor/Source/ProceduralVegetation/Private/Nodes/PVTexture2DImporterSettings.cpp b/Engine/Plugins/Experimental/ProceduralVegetationEditor/Source/ProceduralVegetation/Private/Nodes/PVTexture2DImporterSettings.cpp
--- a/Engine/Plugins/Experimental/ProceduralVegetationEditor/Source/ProceduralVegetation/Private/Nodes/PVTexture2DImporterSettings.cpp
+++ b/Engine/Plugins/Experimental/ProceduralVegetationEditor/Source/ProceduralVegetation/Private/Nodes/PVTexture2DImporterSettings.cpp
@@ -138,16 +138,6 @@ bool FPVTexture2DImporterElement::ExecuteInternal(FPCGContext* InContext) const
 		TipVisualization.Color = FColor::Red;
 		TipVisualization.GizmoScale = 6.5f;
 
-		FPVVisualizationSettings& UserRootPositionsVisualization = OutVisualizationSettings.AddDefaulted_GetRef();
-		UserRootPositionsVisualization.DebugType = EPVDebugType::Custom;
-		UserRootPositionsVisualization.bShowAnchorPoints = false;
-		UserRootPositionsVisualization.bDrawPointAsMesh = false;
-		UserRootPositionsVisualization.AttributeToFilter = PVImportNames::UserRootPositionsAttribute;
-		UserRootPositionsVisualization.CustomGroupToFilter = PVImportNames::UserRootPositionsGroup;
-		UserRootPositionsVisualization.VisualizationMode = EPVDebugValueVisualizationMode::Point;
-		UserRootPositionsVisualization.Color = FColor::Blue;
-		UserRootPositionsVisualization.GizmoScale = 10;
-
 		FPVVisualizationSettings& LabelVisualization = OutVisualizationSettings.AddDefaulted_GetRef();
 		LabelVisualization.DebugType = EPVDebugType::Custom;
 		LabelVisualization.bShowAnchorPoints = false;
@@ -157,7 +147,7 @@ bool FPVTexture2DImporterElement::ExecuteInternal(FPCGContext* InContext) const
 		LabelVisualization.CustomPivotScaleAttributeName = PVImportNames::LabelScaleAttribute;
 		LabelVisualization.VisualizationMode = EPVDebugValueVisualizationMode::Text;
 		LabelVisualization.Color = FColor::Black;
-		LabelVisualization.TextScale = 10.f;
+		LabelVisualization.TextScale = 30.f;
 
 		return OutVisualizationSettings;
 	});
@@ -193,23 +183,24 @@ bool FPVTexture2DImporterElement::ExecuteInternal(FPCGContext* InContext) const
 		return false;
 	}
 
+	const bool bVisualizeDebugCollection = ImportTexture2DOutput.DebugCollection.IsSet() && PV::Utilities::DebugModeEnabled();
+
 	// Add Combined output
 	{
 		UPVGrowthData* OutVariantData = FPCGContext::NewObject_AnyThread<UPVGrowthData>(InContext);
 
 		OutVariantData->AddManagedResource(TextureVisualizationMaterialInstance);
 #if WITH_EDITOR
-		OutVariantData->AddVisualizationCollection(ImportTexture2DOutput.VisualizationCollection, VisualizationSettings);
-#endif
-
-		if (ImportTexture2DOutput.DebugCollection.IsSet() && PV::Utilities::DebugModeEnabled())
+		if (bVisualizeDebugCollection)
 		{
-			OutVariantData->Initialize(MoveTemp(ImportTexture2DOutput.DebugCollection.GetValue()));
+			OutVariantData->AddVisualizationCollection(ImportTexture2DOutput.DebugCollection.GetValue(), Settings->GetDebugVisualizationSettings().VisualizationSettings);
 		}
 		else
 		{
-			OutVariantData->Initialize(MoveTemp(ImportTexture2DOutput.CombinedPlantCollection));
+			OutVariantData->AddVisualizationCollection(ImportTexture2DOutput.VisualizationCollection, VisualizationSettings);
 		}
+#endif
+		OutVariantData->Initialize(MoveTemp(ImportTexture2DOutput.CombinedPlantCollection));
 
 		FPCGTaggedData& CollectionOutput = InContext->OutputData.TaggedData.Emplace_GetRef();
 		CollectionOutput.Data = OutVariantData;
@@ -225,7 +216,14 @@ bool FPVTexture2DImporterElement::ExecuteInternal(FPCGContext* InContext) const
 		// TODO: For now we add the visualization to all outputs as we cannot control which pin gets inspected
 		OutVariantData->AddManagedResource(TextureVisualizationMaterialInstance);
 #if WITH_EDITOR
-		OutVariantData->AddVisualizationCollection(ImportTexture2DOutput.VisualizationCollection, VisualizationSettings);
+		if (bVisualizeDebugCollection)
+		{
+			OutVariantData->AddVisualizationCollection(ImportTexture2DOutput.DebugCollection.GetValue(), Settings->GetDebugVisualizationSettings().VisualizationSettings);
+		}
+		else
+		{
+			OutVariantData->AddVisualizationCollection(ImportTexture2DOutput.VisualizationCollection, VisualizationSettings);
+		}
 #endif
 
 		FPCGTaggedData& CollectionOutput = InContext->OutputData.TaggedData.Emplace_GetRef();
diff --git a/Engine/Plugins/Experimental/ProceduralVegetationEditor/Source/ProceduralVegetation/Private/Nodes/PVTexture2DImporterSettings.h b/Engine/Plugins/Experimental/ProceduralVegetationEditor/Source/ProceduralVegetation/Private/Nodes/PVTexture2DImporterSettings.h
--- a/Engine/Plugins/Experimental/ProceduralVegetationEditor/Source/ProceduralVegetation/Private/Nodes/PVTexture2DImporterSettings.h
+++ b/Engine/Plugins/Experimental/ProceduralVegetationEditor/Source/ProceduralVegetation/Private/Nodes/PVTexture2DImporterSettings.h
@@ -25,6 +25,7 @@ public:
 	virtual EPCGSettingsType GetType() const override { return EPCGSettingsType::InputOutput; }
 	virtual TArray<EPVRenderType> GetDefaultRenderType() const override { return TArray{ EPVRenderType::PointData, EPVRenderType::Mesh }; }
 	virtual TArray<EPVRenderType> GetSupportedRenderTypes() const override { return TArray{ EPVRenderType::PointData, EPVRenderType::Mesh }; }
+	virtual bool IsCollectionRenderingEnabled() const override { return false; }
 #endif
 
 protected:
diff --git a/Engine/Plugins/Experimental/ProceduralVegetationEditor/Source/ProceduralVegetation/Public/Nodes/PVBaseSettings.h b/Engine/Plugins/Experimental/ProceduralVegetationEditor/Source/ProceduralVegetation/Public/Nodes/PVBaseSettings.h
--- a/Engine/Plugins/Experimental/ProceduralVegetationEditor/Source/ProceduralVegetation/Public/Nodes/PVBaseSettings.h
+++ b/Engine/Plugins/Experimental/ProceduralVegetationEditor/Source/ProceduralVegetation/Public/Nodes/PVBaseSettings.h
@@ -28,7 +28,9 @@ public:
 	
 	virtual const FPVDebugSettings& GetDebugVisualizationSettings() const override { return Visualization.DebugSettings; }
 	virtual bool IsDebug() const override { return bDebug; }
-	
+	virtual bool IsCollectionRenderingEnabled() const override { return true; }
+	virtual bool IsVisualizationCollectionsRenderingEnabled() const override { return true; }
+
 	virtual void PostEditChangeProperty(FPropertyChangedEvent& PropertyChangedEvent) override;
 #endif
 
diff --git a/Engine/Plugins/Experimental/ProceduralVegetationEditor/Source/ProceduralVegetation/Public/Nodes/PVRenderSettings.h b/Engine/Plugins/Experimental/ProceduralVegetationEditor/Source/ProceduralVegetation/Public/Nodes/PVRenderSettings.h
--- a/Engine/Plugins/Experimental/ProceduralVegetationEditor/Source/ProceduralVegetation/Public/Nodes/PVRenderSettings.h
+++ b/Engine/Plugins/Experimental/ProceduralVegetationEditor/Source/ProceduralVegetation/Public/Nodes/PVRenderSettings.h
@@ -29,5 +29,7 @@ public:
 	virtual const FPVDebugSettings& GetDebugVisualizationSettings() const = 0;
 
 	virtual bool IsDebug() const = 0;
+	virtual bool IsCollectionRenderingEnabled() const = 0;
+	virtual bool IsVisualizationCollectionsRenderingEnabled() const = 0;
 #endif
 };
diff --git a/Engine/Plugins/Experimental/ProceduralVegetationEditor/Source/ProceduralVegetation/Public/Nodes/PVSubgraphSettings.h b/Engine/Plugins/Experimental/ProceduralVegetationEditor/Source/ProceduralVegetation/Public/Nodes/PVSubgraphSettings.h
--- a/Engine/Plugins/Experimental/ProceduralVegetationEditor/Source/ProceduralVegetation/Public/Nodes/PVSubgraphSettings.h
+++ b/Engine/Plugins/Experimental/ProceduralVegetationEditor/Source/ProceduralVegetation/Public/Nodes/PVSubgraphSettings.h
@@ -35,6 +35,8 @@ public:
 	
 	virtual const FPVDebugSettings& GetDebugVisualizationSettings() const override { return Visualization.DebugSettings; }
 	virtual bool IsDebug() const override { return bDebug; }
+	virtual bool IsCollectionRenderingEnabled() const override { return true; }
+	virtual bool IsVisualizationCollectionsRenderingEnabled() const override { return true; }
 #endif
 	
 	//~UPCGSettings interface implementation
diff --git a/Engine/Plugins/Experimental/ProceduralVegetationEditor/Source/ProceduralVegetationEditor/Private/Tools/NodeTools/PVTexture2DImporterTool.cpp b/Engine/Plugins/Experimental/ProceduralVegetationEditor/Source/ProceduralVegetationEditor/Private/Tools/NodeTools/PVTexture2DImporterTool.cpp
new file mode 100644
--- /dev/null
+++ b/Engine/Plugins/Experimental/ProceduralVegetationEditor/Source/ProceduralVegetationEditor/Private/Tools/NodeTools/PVTexture2DImporterTool.cpp
@@ -0,0 +1,203 @@
+// Copyright Epic Games, Inc. All Rights Reserved.
+
+#include "PVTexture2DImporterTool.h"
+
+#include "Nodes/PVTexture2DImporterSettings.h"
+
+#include "BaseBehaviors/ClickDragBehavior.h"
+#include "HitProxies.h"
+#include "UnrealClient.h"
+#include "InteractiveToolManager.h"
+#include "PrimitiveDrawInterface.h"
+
+#define LOCTEXT_NAMESPACE "PVTexture2DImportTool"
+
+namespace PVTexture2DImporterTool
+{
+	class HPVRootPointProxy : public HHitProxy
+	{
+		DECLARE_HIT_PROXY();
+
+	public:
+		const int32 PlantIndex = INDEX_NONE;
+		const bool bIsRotation = false;
+
+		HPVRootPointProxy(const int32 InPlantIndex, const bool bInIsRotation)
+			: HHitProxy(HPP_Wireframe)
+			, PlantIndex(InPlantIndex)
+			, bIsRotation(bInIsRotation)
+		{
+		}
+
+		virtual EMouseCursor::Type GetMouseCursor() override
+		{
+			return EMouseCursor::GrabHand;
+		}
+	};
+
+	IMPLEMENT_HIT_PROXY(HPVRootPointProxy, HHitProxy)
+
+	template<typename Vector3dType = float, typename Vector2dType = float>
+	UE::Math::TVector<Vector3dType> UVToVec3(const UE::Math::TVector2<Vector2dType>& Vec2D, float Scale)
+	{
+		return UE::Math::TVector<Vector3dType>(0, Vec2D.X * Scale - Scale / 2.0, (1.0 - Vec2D.Y) * Scale + 25.0);
+	}
+
+	template<typename Vector2dType = float, typename Vector3dType = float>
+	UE::Math::TVector2<Vector2dType> Vec3ToUV(const UE::Math::TVector<Vector3dType>& Vec3D, float Scale)
+	{
+		if (FMath::IsNearlyZero(Scale))
+		{
+			return UE::Math::TVector2<Vector2dType>::ZeroVector;
+		}
+		return UE::Math::TVector2<Vector2dType>((Vec3D.Y + Scale / 2.0) / Scale, 1.0 - (Vec3D.Z - 25.0) / Scale);
+	}
+}
+
+UPVTexture2DImporterTool::UPVTexture2DImporterTool()
+{
+	SetFlags(RF_Transactional);
+}
+
+void UPVTexture2DImporterTool::Setup()
+{
+	Super::Setup();
+
+	UClickDragInputBehavior* DragBehavior = NewObject<UClickDragInputBehavior>(this);
+	DragBehavior->Initialize(this);
+	AddInputBehavior(DragBehavior);
+}
+
+void UPVTexture2DImporterTool::Render(IToolsContextRenderAPI* RenderAPI)
+{
+	using namespace PVTexture2DImporterTool;
+
+	const TObjectPtr<UPVTexture2DImporterSettings> Texture2DImporterSettings = GetNodeSettings<UPVTexture2DImporterSettings>();
+	FPrimitiveDrawInterface* PDI = RenderAPI->GetPrimitiveDrawInterface();
+
+	const FPVImportTexture2DParams& ImportTexture2DParams = Texture2DImporterSettings->Params;
+
+	constexpr float ScaleFactor = 0.15f;
+	const float GizmoScale = ImportTexture2DParams.Scale * ScaleFactor;
+	for (int32 PlantIndex = 0; PlantIndex < ImportTexture2DParams.PlantSettings.Num(); ++PlantIndex)
+	{
+		const FPVImportTexture2DPlantSettings& PlantSettings = ImportTexture2DParams.PlantSettings[PlantIndex];
+		const FVector3f RootPointPosition = UVToVec3(PlantSettings.RootPosition, ImportTexture2DParams.Scale);
+
+		PDI->SetHitProxy(new HPVRootPointProxy(PlantIndex, false));
+		PDI->DrawPoint(FVector(RootPointPosition), FColor::Blue, 0.2f * GizmoScale, SDPG_Foreground);
+		PDI->SetHitProxy(nullptr);
+
+		const float RotationRad = FMath::DegreesToRadians(PlantSettings.Rotation);
+		const FVector RotationDir(0.f, -FMath::Sin(RotationRad), FMath::Cos(RotationRad));
+		const FVector RotationGizmoEndPoint = FVector(RootPointPosition) + RotationDir * GizmoScale;
+
+		PDI->DrawLine(FVector(RootPointPosition), RotationGizmoEndPoint, FColor::Orange, SDPG_Foreground, 0.025f * GizmoScale);
+		PDI->SetHitProxy(new HPVRootPointProxy(PlantIndex, true));
+		PDI->DrawPoint(RotationGizmoEndPoint, FColor::Orange, 0.15f * GizmoScale, SDPG_Foreground);
+		PDI->SetHitProxy(nullptr);
+	}
+}
+
+FInputRayHit UPVTexture2DImporterTool::CanBeginClickDragSequence(const FInputDeviceRay& PressPos)
+{
+	using namespace PVTexture2DImporterTool;
+
+	if (FViewport* Viewport = GetToolManager()->GetContextQueriesAPI()->GetFocusedViewport())
+	{
+		const HHitProxy* HitProxy = Viewport->GetHitProxy(PressPos.ScreenPosition.X, PressPos.ScreenPosition.Y);
+		if (HitProxy && HitProxy->IsA(HPVRootPointProxy::StaticGetType()))
+		{
+			return FInputRayHit(1.f);
+		}
+	}
+	return FInputRayHit();
+}
+
+void UPVTexture2DImporterTool::OnClickPress(const FInputDeviceRay& PressPos)
+{
+	using namespace PVTexture2DImporterTool;
+
+	FViewport* Viewport = GetToolManager()->GetContextQueriesAPI()->GetFocusedViewport();
+	if (!Viewport)
+	{
+		return;
+	}
+
+	const HHitProxy* HitProxy = Viewport->GetHitProxy(PressPos.ScreenPosition.X, PressPos.ScreenPosition.Y);
+	if (HitProxy && HitProxy->IsA(HPVRootPointProxy::StaticGetType()))
+	{
+		const HPVRootPointProxy* RootPointProxy = static_cast<const HPVRootPointProxy*>(HitProxy);
+		DraggedPlantIndex = RootPointProxy->PlantIndex;
+		bDraggingRotation = RootPointProxy->bIsRotation;
+		GetToolManager()->BeginUndoTransaction(bDraggingRotation ? LOCTEXT("RotatePlant", "Rotate Plant") : LOCTEXT("MoveRootPoint", "Move Root Point"));
+		GetNodeSettings<UPVTexture2DImporterSettings>()->Modify();
+	}
+}
+
+void UPVTexture2DImporterTool::OnClickDrag(const FInputDeviceRay& DragPos)
+{
+	using namespace PVTexture2DImporterTool;
+
+	const TObjectPtr<UPVTexture2DImporterSettings> Settings = GetNodeSettings<UPVTexture2DImporterSettings>();
+
+	if (DraggedPlantIndex == INDEX_NONE || !Settings->Params.PlantSettings.IsValidIndex(DraggedPlantIndex))
+	{
+		return;
+	}
+
+	// All root points lie in the plane at X=0 in world space (UVToVec3 always produces X=0).
+	// Intersect the drag ray with that plane to find where the cursor is in 3D.
+	const FRay& WorldRay = DragPos.WorldRay;
+	if (FMath::IsNearlyZero(WorldRay.Direction.X))
+	{
+		return;
+	}
+
+	const float T = -WorldRay.Origin.X / WorldRay.Direction.X;
+	if (T < 0.f)
+	{
+		return;
+	}
+
+	const float Scale = Settings->Params.Scale;
+	const FVector2f HitPointUV = Vec3ToUV(FVector3f(WorldRay.PointAt(T)), Scale);
+
+	FPVImportTexture2DPlantSettings& PlantSettings = Settings->Params.PlantSettings[DraggedPlantIndex];
+
+	if (bDraggingRotation)
+	{
+		const FVector2f Offset = PlantSettings.RootPosition - HitPointUV;
+		PlantSettings.Rotation = FMath::RadiansToDegrees(FMath::Atan2(Offset.X, Offset.Y));
+	}
+	else
+	{
+		PlantSettings.RootPosition = FVector2f(
+			FMath::Clamp(HitPointUV.X, 0.f, 1.f),
+			FMath::Clamp(HitPointUV.Y, 0.f, 1.f)
+		);
+	}
+}
+
+void UPVTexture2DImporterTool::OnClickRelease(const FInputDeviceRay& ReleasePos)
+{
+	EndClickDragBehavior();
+}
+
+void UPVTexture2DImporterTool::OnTerminateDragSequence()
+{
+	EndClickDragBehavior();
+}
+
+void UPVTexture2DImporterTool::EndClickDragBehavior()
+{
+	if (DraggedPlantIndex != INDEX_NONE)
+	{
+		GetNodeSettings()->PostEditChange();
+		GetToolManager()->EndUndoTransaction();
+		DraggedPlantIndex = INDEX_NONE;
+		bDraggingRotation = false;
+	}
+}
+
+#undef LOCTEXT_NAMESPACE
diff --git a/Engine/Plugins/Experimental/ProceduralVegetationEditor/Source/ProceduralVegetationEditor/Private/Tools/NodeTools/PVTexture2DImporterTool.h b/Engine/Plugins/Experimental/ProceduralVegetationEditor/Source/ProceduralVegetationEditor/Private/Tools/NodeTools/PVTexture2DImporterTool.h
new file mode 100644
--- /dev/null
+++ b/Engine/Plugins/Experimental/ProceduralVegetationEditor/Source/ProceduralVegetationEditor/Private/Tools/NodeTools/PVTexture2DImporterTool.h
@@ -0,0 +1,40 @@
+// Copyright Epic Games, Inc. All Rights Reserved.
+
+#pragma once
+
+#include "Tools/PVBaseInteractiveTool.h"
+#include "BaseBehaviors/BehaviorTargetInterfaces.h"
+#include "Params/PVImportTexture2DParams.h"
+#include "PVTexture2DImporterTool.generated.h"
+
+class UTransformProxy;
+class IToolsContextRenderAPI;
+
+UCLASS()
+class UPVTexture2DImporterTool : public UPVBaseInteractiveTool, public IClickDragBehaviorTarget
+{
+	GENERATED_BODY()
+
+public:
+	UPVTexture2DImporterTool();
+
+	// Begin UInteractiveTool interface
+	virtual void Setup() override;
+	virtual void Render(IToolsContextRenderAPI* RenderAPI) override;
+	// End UInteractiveTool interface
+
+	// Begin IClickDragBehaviorTarget interface
+	virtual FInputRayHit CanBeginClickDragSequence(const FInputDeviceRay& PressPos) override;
+	virtual void OnClickPress(const FInputDeviceRay& PressPos) override;
+	virtual void OnClickDrag(const FInputDeviceRay& DragPos) override;
+	virtual void OnClickRelease(const FInputDeviceRay& ReleasePos) override;
+	virtual void OnTerminateDragSequence() override;
+	// End IClickDragBehaviorTarget interface
+
+private:
+	void EndClickDragBehavior();
+
+private:
+	int32 DraggedPlantIndex = INDEX_NONE;
+	bool bDraggingRotation = false;
+};
diff --git a/Engine/Plugins/Experimental/ProceduralVegetationEditor/Source/ProceduralVegetationEditor/Private/Tools/PVToolsRegistry.cpp b/Engine/Plugins/Experimental/ProceduralVegetationEditor/Source/ProceduralVegetationEditor/Private/Tools/PVToolsRegistry.cpp
--- a/Engine/Plugins/Experimental/ProceduralVegetationEditor/Source/ProceduralVegetationEditor/Private/Tools/PVToolsRegistry.cpp
+++ b/Engine/Plugins/Experimental/ProceduralVegetationEditor/Source/ProceduralVegetationEditor/Private/Tools/PVToolsRegistry.cpp
@@ -8,20 +8,14 @@
 #include "Contexts/PVToolContext.h"
 
 #include "Nodes/PVSkeletonEditSettings.h"
+#include "Nodes/PVTexture2DImporterSettings.h"
 
 #include "NodeTools/PVColliderTransformTool.h"
 #include "NodeTools/PVSkeletonEditTool.h"
+#include "NodeTools/PVTexture2DImporterTool.h"
 
 #include "Tools/EdModeInteractiveToolsContext.h"
 
-#define REGISTER_PVE_TOOL(Tool, Node)																								\
-	{																																\
-		UPVToolBuilder* const ToolBuilder = NewObject<UPVToolBuilder>(ToolsContext);												\
-		ToolBuilder->ToolClass = Tool::StaticClass();																				\
-		ToolManager->RegisterToolType(Tool::StaticClass()->GetName(), ToolBuilder);													\
-		NodeSettingsToToolMapping.Add(Node::StaticClass(), Tool::StaticClass()->GetName());											\
-	}
-
 FPVToolsRegistry::FPVToolsRegistry(const TObjectPtr<UEditorInteractiveToolsContext>& InToolsContext)
 	: ToolsContext(InToolsContext)
 {
@@ -34,8 +28,9 @@ void FPVToolsRegistry::RegisterTools()
 {
 	check(ToolsContext);
 
-	REGISTER_PVE_TOOL(UPVColliderTransformTool, UPVObjectInteractionSettings);
-	REGISTER_PVE_TOOL(UPVSkeletonEditTool, UPVSkeletonEditSettings);
+	RegisterNodeTool(UPVColliderTransformTool::StaticClass(), UPVObjectInteractionSettings::StaticClass());
+	RegisterNodeTool(UPVSkeletonEditTool::StaticClass(), UPVSkeletonEditSettings::StaticClass());
+	RegisterNodeTool(UPVTexture2DImporterTool::StaticClass(), UPVTexture2DImporterSettings::StaticClass());
 }
 
 void FPVToolsRegistry::UnregisterTools()
@@ -96,3 +91,21 @@ bool FPVToolsRegistry::HasActiveTool()
 {
 	return ToolManager->HasActiveTool(EToolSide::Left);
 }
+
+void FPVToolsRegistry::RegisterNodeTool(TSubclassOf<UPVBaseInteractiveTool> ToolClass, TSubclassOf<UPVBaseSettings> NodeClass)
+{
+	if (!ensure(ToolClass && NodeClass))
+	{
+		return;
+	}
+
+	if (!ensure(!NodeSettingsToToolMapping.Contains(NodeClass)))
+	{
+		return;
+	}
+
+	UPVToolBuilder* const ToolBuilder = NewObject<UPVToolBuilder>(ToolsContext);
+	ToolBuilder->ToolClass = ToolClass;
+	ToolManager->RegisterToolType(ToolClass->GetName(), ToolBuilder);
+	NodeSettingsToToolMapping.Add(NodeClass, ToolClass->GetName());
+}
diff --git a/Engine/Plugins/Experimental/ProceduralVegetationEditor/Source/ProceduralVegetationEditor/Private/Tools/PVToolsRegistry.h b/Engine/Plugins/Experimental/ProceduralVegetationEditor/Source/ProceduralVegetationEditor/Private/Tools/PVToolsRegistry.h
--- a/Engine/Plugins/Experimental/ProceduralVegetationEditor/Source/ProceduralVegetationEditor/Private/Tools/PVToolsRegistry.h
+++ b/Engine/Plugins/Experimental/ProceduralVegetationEditor/Source/ProceduralVegetationEditor/Private/Tools/PVToolsRegistry.h
@@ -10,6 +10,7 @@
 #include "UObject/ObjectPtr.h"
 
 class UPVBaseSettings;
+class UPVBaseInteractiveTool;
 class UInteractiveToolManager;
 class UEditorInteractiveToolsContext;
 
@@ -43,6 +44,10 @@ public:
 	 */
 	bool HasActiveTool();
 
+private:
+
+	void RegisterNodeTool(TSubclassOf<UPVBaseInteractiveTool> ToolClass, TSubclassOf<UPVBaseSettings> NodeClass);
+
 protected:
 	TMap<TObjectPtr<UClass>, FString> NodeSettingsToToolMapping;
 	
diff --git a/Engine/Plugins/Experimental/ProceduralVegetationEditor/Source/ProceduralVegetationEditor/Private/Visualizations/PVDataVisualization.cpp b/Engine/Plugins/Experimental/ProceduralVegetationEditor/Source/ProceduralVegetationEditor/Private/Visualizations/PVDataVisualization.cpp
--- a/Engine/Plugins/Experimental/ProceduralVegetationEditor/Source/ProceduralVegetationEditor/Private/Visualizations/PVDataVisualization.cpp
+++ b/Engine/Plugins/Experimental/ProceduralVegetationEditor/Source/ProceduralVegetationEditor/Private/Visualizations/PVDataVisualization.cpp
@@ -115,20 +115,34 @@ FPCGSetupSceneFunc FPVDataVisualization::GetViewportSetupFunc(const UPCGSettings
 			InOutParams.Scene->HandleTogglePostProcessing();
 		}
  
-		const FManagedArrayCollection& Collection = WeakData->GetCollection();
- 
+		InOutParams.Scene->GetLineBatcher()->Flush();
+
 		FBoxSphereBounds Bounds(ForceInit);
- 
-		for (const EPVRenderType& RenderType : RenderSettings->GetCurrentRenderTypes())
+
+		const bool bIsCollectionRenderingEnabled = RenderSettings->IsCollectionRenderingEnabled();
+		const bool bIsVisualizationCollectionsRenderingEnabled = RenderSettings->IsVisualizationCollectionsRenderingEnabled();
+
+		if (bIsCollectionRenderingEnabled || bIsVisualizationCollectionsRenderingEnabled)
 		{
-			if (RenderMap.Contains(RenderType))
+			for (const EPVRenderType& RenderType : RenderSettings->GetCurrentRenderTypes())
 			{
-				InOutParams.Scene->GetLineBatcher()->Flush();
-				RenderMap[RenderType].Execute(WeakData.Get(), Collection, InOutParams, Bounds);
+				const FPVRenderCallback* RenderCallback = RenderMap.Find(RenderType);
+				if (!ensure(RenderCallback != nullptr))
+				{
+					continue;
+				}
+
+				if (bIsCollectionRenderingEnabled)
+				{
+					RenderCallback->Execute(WeakData.Get(), WeakData->GetCollection(), InOutParams, Bounds);
+				}
 
-				for (const FPVVisualizationCollection& VisualizationCollection : WeakData->GetVisualizationCollections())
+				if (bIsVisualizationCollectionsRenderingEnabled)
 				{
-					RenderMap[RenderType].Execute(WeakData.Get(), VisualizationCollection.Collection, InOutParams, Bounds);
+					for (const FPVVisualizationCollection& VisualizationCollection : WeakData->GetVisualizationCollections())
+					{
+						RenderCallback->Execute(WeakData.Get(), VisualizationCollection.Collection, InOutParams, Bounds);
+					}
 				}
 			}
 		}
