[PVE] Fix scale for fused branch points
- Extract scale computation into AttributesHelper::GetBranchPointScale to encapsulate fused-point handling
- For non-trunk branches, the 0th branch ring is the fused point shared with the parent; scale is now extrapolated using the child's local taper rate (S1 - S2) divided by PI, matching the Houdini growth formula
- Cap the extrapolated scale at 80% of the parent's scale to prevent overshooting

#Synced-CL 52379040

diff --git a/Engine/Plugins/Experimental/ProceduralVegetationEditor/Source/ProceduralVegetation/Private/Helpers/PVAttributesHelper.cpp b/Engine/Plugins/Experimental/ProceduralVegetationEditor/Source/ProceduralVegetation/Private/Helpers/PVAttributesHelper.cpp
--- a/Engine/Plugins/Experimental/ProceduralVegetationEditor/Source/ProceduralVegetation/Private/Helpers/PVAttributesHelper.cpp
+++ b/Engine/Plugins/Experimental/ProceduralVegetationEditor/Source/ProceduralVegetation/Private/Helpers/PVAttributesHelper.cpp
@@ -42,6 +42,40 @@ float PV::AttributesHelper::GetBranchPointBranchGradient(
 	);
 }
 
+float PV::AttributesHelper::GetBranchPointScale(
+	PV::FPointScaleAttributeConstView PointScaleAttribute,
+	PV::FBranchPointsAttributeConstView BranchPointsAttribute,
+	PV::FBranchParentNumberAttributeConstView BranchParentNumberAttribute,
+	int32 BranchIndex,
+	int32 BranchPointIndex
+)
+{
+	if (!BranchPointsAttribute.IsValidIndex(BranchIndex) || !BranchPointsAttribute[BranchIndex].IsValidIndex(BranchPointIndex))
+	{
+		return 0.0f;
+	}
+
+	const TArray<int32>& BranchPoints = BranchPointsAttribute[BranchIndex];
+	const int32 PointIndex = BranchPoints[BranchPointIndex];
+
+	// Ring 0 of a non-trunk branch is the fused junction shared with the parent.
+	// Extrapolate the child's scale at ring 0 using the local taper (S1 - S2) as an
+	// elongation estimate, capped at 80% of the parent's scale to avoid overshooting.
+	const bool bIsRootFusedPoint = !PV::PlantTraversalHelper::IsTrunk(BranchParentNumberAttribute, BranchIndex) && BranchPointIndex == 0;
+	
+	if (bIsRootFusedPoint && BranchPoints.IsValidIndex(1))
+	{
+		const float S1 = PointScaleAttribute[BranchPoints[1]];
+		const float ParentScale = PointScaleAttribute[PointIndex];
+		const float Elongation = BranchPoints.IsValidIndex(2) ? S1 - PointScaleAttribute[BranchPoints[2]] : 0.0f;
+		const float NewScale = S1 + Elongation / UE_PI;
+		
+		return NewScale > ParentScale ? ParentScale * 0.8f : NewScale;
+	}
+
+	return PointScaleAttribute[PointIndex];
+}
+
 void PV::AttributesHelper::ComputeBudNumbers(const FComputeBudNumbersAttributes& InAttributes)
 {
 	for (int32 PointIndex = 0; PointIndex < InAttributes.PointBudNumber.Num(); ++PointIndex)
diff --git a/Engine/Plugins/Experimental/ProceduralVegetationEditor/Source/ProceduralVegetation/Private/Helpers/PVAttributesHelper.h b/Engine/Plugins/Experimental/ProceduralVegetationEditor/Source/ProceduralVegetation/Private/Helpers/PVAttributesHelper.h
--- a/Engine/Plugins/Experimental/ProceduralVegetationEditor/Source/ProceduralVegetation/Private/Helpers/PVAttributesHelper.h
+++ b/Engine/Plugins/Experimental/ProceduralVegetationEditor/Source/ProceduralVegetation/Private/Helpers/PVAttributesHelper.h
@@ -8,10 +8,20 @@ namespace PV::AttributesHelper
 	float GetBranchPointBranchGradient(
 		PV::FPointLengthFromRootAttributeConstView PointLengthFromRootAttribute,
 		PV::FBranchPointsAttributeConstView BranchPointsAttribute,
-		int32 BranchIndex, 
+		int32 BranchIndex,
 		int32 BranchPointIndex
 	);
 
+	float GetBranchPointScale(
+		PV::FPointScaleAttributeConstView PointScaleAttribute,
+		PV::FBranchPointsAttributeConstView BranchPointsAttribute,
+		PV::FBranchParentNumberAttributeConstView BranchParentNumberAttribute,
+		int32 BranchIndex,
+		int32 BranchPointIndex
+	);
+	
+	
+
 	struct FComputeBudNumbersAttributes
 	{
 		FPointBudNumberAttributeView PointBudNumber;
diff --git a/Engine/Plugins/Experimental/ProceduralVegetationEditor/Source/ProceduralVegetation/Private/Implementations/PVMaterialSettings.cpp b/Engine/Plugins/Experimental/ProceduralVegetationEditor/Source/ProceduralVegetation/Private/Implementations/PVMaterialSettings.cpp
--- a/Engine/Plugins/Experimental/ProceduralVegetationEditor/Source/ProceduralVegetation/Private/Implementations/PVMaterialSettings.cpp
+++ b/Engine/Plugins/Experimental/ProceduralVegetationEditor/Source/ProceduralVegetation/Private/Implementations/PVMaterialSettings.cpp
@@ -8,6 +8,7 @@
 #include "Facades/PVPlantFacade.h"
 #include "Facades/PVPointFacade.h"
 #include "Facades/PVTreeFacade.h"
+#include "Helpers/PVAttributesHelper.h"
 
 struct FBranchUVData
 {
@@ -67,6 +68,10 @@ void FPVMaterialSettings::SetBranchUVMaterial(FManagedArrayCollection& Collectio
 	PV::Facades::FPointFacade PointFacade = PV::Facades::FPointFacade(Collection);
 	PV::Facades::FBranchFacade BranchFacade = PV::Facades::FBranchFacade(Collection);
 
+	const auto PointScaleAttribute = PV::FPointScaleAttribute::GetAttribute(Collection);
+	const auto BranchPointsAttribute = PV::FBranchPointsAttribute::GetAttribute(Collection);
+	const auto BranchParentNumberAttribute = PV::FBranchParentNumberAttribute::GetAttribute(Collection);
+
 	BranchUVData.Init(FBranchUVData(), BranchFacade.GetElementCount());
 	
 	float MaxValue = 10000;
@@ -109,7 +114,7 @@ void FPVMaterialSettings::SetBranchUVMaterial(FManagedArrayCollection& Collectio
 		const int32 GenerationNumber = PV::Facades::FTreeFacade::GetGenerationNumberFromSourceBud(Collection, BranchIndex);
 	
 		TArray<int> BudDevelopment = PointFacade.GetBudDevelopment(BranchRootPoint);
-		float PScale = PointFacade.GetPointScale(BranchRootPoint);
+		float PScale = PV::AttributesHelper::GetBranchPointScale(PointScaleAttribute, BranchPointsAttribute, BranchParentNumberAttribute, BranchIndex, 0);
 		check(BudDevelopment.Num() > 2);
 		
 		if (MaterialMode == EUVMaterialMode::Generation)
diff --git a/Engine/Plugins/Experimental/ProceduralVegetationEditor/Source/ProceduralVegetation/Private/Implementations/PVMeshBuilder.cpp b/Engine/Plugins/Experimental/ProceduralVegetationEditor/Source/ProceduralVegetation/Private/Implementations/PVMeshBuilder.cpp
--- a/Engine/Plugins/Experimental/ProceduralVegetationEditor/Source/ProceduralVegetation/Private/Implementations/PVMeshBuilder.cpp
+++ b/Engine/Plugins/Experimental/ProceduralVegetationEditor/Source/ProceduralVegetation/Private/Implementations/PVMeshBuilder.cpp
@@ -849,7 +849,7 @@ FMeshPointUVData FPVMeshBuilder::ComputePointUVData(const int32 PointIndex, cons
                                                const TArray<int32>& PrimitivePoints, const PV::Facades::FPointFacade& PointFacade,
                                                const TManagedArray<FVector3f>& PointPositions, const float PointScale)
 {
-	// For non-trunk branches, the 0th branch ring is the fused junction shared with the parent.
+	// For non-trunk branches, the 0th branch ring is the fused point shared with the parent.
 	// Its UV attributes belong to the parent, so read UV from BranchPoints[1] and extrapolate
 	// backward by the B0→B1 segment length to recover the child's origin UV.
 	const bool bRootBranchPoint = !bPrimitiveIsTrunk && BranchPointIndex == 0;
@@ -886,7 +886,7 @@ FMeshPointUVData FPVMeshBuilder::ComputePointUVData(const int32 PointIndex, cons
 	return { TextureCoordV, TextureCoordUOffset, URange };
 }
 
-void FPVMeshBuilder::GenerateBranchMeshData(const bool bPrimitiveIsTrunk, const int32 GenerationNumber, const TArray<int32>& PrimitivePoints,
+void FPVMeshBuilder::GenerateBranchMeshData(const bool bPrimitiveIsTrunk, const int32 GenerationNumber, const int32 BranchIndex, const TArray<int32>& PrimitivePoints,
                                             const FDisplacementData& DisplacementData, const FPVMeshBuilderParams& MeshBuilderParams,
                                             const TArray<int32>& TargetMeshDivisions, const FManagedArrayCollection& Collection, const FManagedArrayCollection& InPlantProfileCollection,
                                             FLocalDynamicMeshData& OutLocalMeshData)
@@ -912,9 +912,12 @@ void FPVMeshBuilder::GenerateBranchMeshData(const bool bPrimitiveIsTrunk, const
 
 	const TManagedArray<FVector3f>& PointPositions = PointFacade.GetPositions();
 	const TManagedArray<TArray<FVector3f>>& PointBudDirections = BudVectorFacade.GetBudDirections();
-	const TManagedArray<float>& PointScales = PointFacade.GetPointScales();
 
-	const float MaxPointScale = *Algo::MaxElement(PointScales);
+	const auto PointScaleAttribute = PV::FPointScaleAttribute::GetAttribute(Collection);
+	const auto BranchPointsAttribute = PV::FBranchPointsAttribute::GetAttribute(Collection);
+	const auto BranchParentNumberAttribute = PV::FBranchParentNumberAttribute::GetAttribute(Collection);
+
+	const float MaxPointScale = *Algo::MaxElement(PointScaleAttribute);
 	const float MaxScaleRatio = 1.0f / (MaxPointScale * UE_TWO_PI);
 	int32 PolyGroupIndex = 0;
 	TArray<int32> CurrentIndices;
@@ -940,40 +943,42 @@ void FPVMeshBuilder::GenerateBranchMeshData(const bool bPrimitiveIsTrunk, const
 	const float TipLengthFromRoot = PointFacade.GetLengthFromRoot(PrimitivePoints.Last());
 
 	FVector3f TangentVector = FVector3f::UpVector;
-	for (int32 Index = 0; Index < PrimitivePoints.Num(); ++Index)
+	for (int32 BranchPointIndex = 0; BranchPointIndex < PrimitivePoints.Num(); ++BranchPointIndex)
 	{
-		const int32& PointIndex = PrimitivePoints[Index];
+		const int32& PointIndex = PrimitivePoints[BranchPointIndex];
 		const float NjordPixelId = PointIndex;
 
 		const float LengthFromRoot = PointFacade.GetLengthFromRoot(PointIndex);
 
 		const FVector3f& PointPosition = PointPositions[PointIndex];
-		const float PointScale = PointScales[PointIndex];
-		const FVector2f OldRange = FVector2f(0, 1.0f);
 
+		const float PointScale = PV::AttributesHelper::GetBranchPointScale(PointScaleAttribute, BranchPointsAttribute, BranchParentNumberAttribute, BranchIndex, BranchPointIndex);
+		
 		const auto [TextureCoordV, TextureCoordUOffset, NewRange] =
-			ComputePointUVData(PointIndex, Index, bPrimitiveIsTrunk, PrimitivePoints, PointFacade, PointPositions, PointScale);
+			ComputePointUVData(PointIndex, BranchPointIndex, bPrimitiveIsTrunk, PrimitivePoints, PointFacade, PointPositions, PointScale);
+		
+		const FVector2f OldRange = FVector2f(0, 1.0f);
 		
 		Displacement_UV_V += FVector3f::Dist(PointPosition, PreviousPosition) * MaxScaleRatio * (MaxPointScale / PointScale);
 		
-		if (PrimitivePoints.IsValidIndex(Index - 1) && PrimitivePoints.IsValidIndex(Index + 1))
+		if (PrimitivePoints.IsValidIndex(BranchPointIndex - 1) && PrimitivePoints.IsValidIndex(BranchPointIndex + 1))
 		[[likely]]
 		{
-			const FVector3f PreviousTangent = (PointPosition - PointPositions[PrimitivePoints[Index - 1]]).GetSafeNormal();
-			const FVector3f NextTangent = (PointPositions[PrimitivePoints[Index + 1]] - PointPosition).GetSafeNormal();
+			const FVector3f PreviousTangent = (PointPosition - PointPositions[PrimitivePoints[BranchPointIndex - 1]]).GetSafeNormal();
+			const FVector3f NextTangent = (PointPositions[PrimitivePoints[BranchPointIndex + 1]] - PointPosition).GetSafeNormal();
 			TangentVector = FVector3f::SlerpNormals(PreviousTangent, NextTangent, 0.5f);
 		}
-		else if (Index == PrimitivePoints.Num() - 1)
+		else if (BranchPointIndex == PrimitivePoints.Num() - 1)
 		[[unlikely]]
 		{
-			const int32 PreviousPointIndex = PrimitivePoints[Index - 1];
+			const int32 PreviousPointIndex = PrimitivePoints[BranchPointIndex - 1];
 			const FVector3f& PreviousPointPosition = PointPositions[PreviousPointIndex];
 			TangentVector = (PointPosition - PreviousPointPosition).GetSafeNormal();
 		}
-		else if (Index == 0)
+		else if (BranchPointIndex == 0)
 		[[unlikely]]
 		{
-			const int32 NextPointIndex = PrimitivePoints[Index + 1];
+			const int32 NextPointIndex = PrimitivePoints[BranchPointIndex + 1];
 			const FVector3f& NextPointPosition = PointPositions[NextPointIndex];
 			TangentVector = (NextPointPosition - PointPosition).GetSafeNormal();
 		}
@@ -981,7 +986,7 @@ void FPVMeshBuilder::GenerateBranchMeshData(const bool bPrimitiveIsTrunk, const
 		// Guard against near-zero tangent when consecutive branch points are coincident
 		if (TangentVector.IsNearlyZero())
 		{
-			for (int32 SearchIndex = Index + 1; SearchIndex < PrimitivePoints.Num(); ++SearchIndex)
+			for (int32 SearchIndex = BranchPointIndex + 1; SearchIndex < PrimitivePoints.Num(); ++SearchIndex)
 			{
 				TangentVector = (PointPositions[PrimitivePoints[SearchIndex]] - PointPosition).GetSafeNormal();
 				if (!TangentVector.IsNearlyZero())
@@ -1114,7 +1119,7 @@ void FPVMeshBuilder::GenerateBranchMeshData(const bool bPrimitiveIsTrunk, const
 	{
 		const int32 LastPointIndex = PrimitivePoints.Last();
 		float CurrentLengthFromRoot = PointFacade.GetLengthFromRoot(LastPointIndex);
-		float CurrentScale = PointScales[LastPointIndex];
+		float CurrentScale = PointScaleAttribute[LastPointIndex];
 		const int32 CapDivisions = TargetMeshDivisions[LastPointIndex];
 		float StepSize = CurrentScale * 0.5f;
 		
@@ -1435,7 +1440,7 @@ void FPVMeshBuilder::GetLocalDynamicMeshData(FManagedArrayCollection& Collection
 
 				MeshesData[PrimitiveIndex].MaterialID = BranchFacade.GetBranchUVMaterial(PrimitiveIndex);
 
-				GenerateBranchMeshData(PrimitiveIsTrunk, GenerationNumber, PrimitivePoints, DisplacementData, MeshBuilderParams, TargetMeshDivisions,
+				GenerateBranchMeshData(PrimitiveIsTrunk, GenerationNumber, PrimitiveIndex, PrimitivePoints, DisplacementData, MeshBuilderParams, TargetMeshDivisions,
 					Collection, InPlantProfileCollection, MeshesData[PrimitiveIndex]);
 			});
 }
diff --git a/Engine/Plugins/Experimental/ProceduralVegetationEditor/Source/ProceduralVegetation/Private/Implementations/PVMeshBuilder.h b/Engine/Plugins/Experimental/ProceduralVegetationEditor/Source/ProceduralVegetation/Private/Implementations/PVMeshBuilder.h
--- a/Engine/Plugins/Experimental/ProceduralVegetationEditor/Source/ProceduralVegetation/Private/Implementations/PVMeshBuilder.h
+++ b/Engine/Plugins/Experimental/ProceduralVegetationEditor/Source/ProceduralVegetation/Private/Implementations/PVMeshBuilder.h
@@ -344,9 +344,9 @@ private:
 	                                       const TArray<int32>& PrimitivePoints, const PV::Facades::FPointFacade& PointFacade,
 	                                       const TManagedArray<FVector3f>& PointPositions, float PointScale);
 
-	static void GenerateBranchMeshData(const bool bPrimitiveIsTrunk, const int32 GenerationNumber, const TArray<int32>& PrimitivePoints,
+	static void GenerateBranchMeshData(const bool bPrimitiveIsTrunk, const int32 GenerationNumber, int32 BranchIndex, const TArray<int32>& PrimitivePoints,
 	                                   const FDisplacementData& DisplacementData, const FPVMeshBuilderParams& MeshBuilderParams,
-	                                   const TArray<int32>& TargetMeshDivisions, const FManagedArrayCollection& Collection, const FManagedArrayCollection& InPlantProfileCollection, 
+	                                   const TArray<int32>& TargetMeshDivisions, const FManagedArrayCollection& Collection, const FManagedArrayCollection& InPlantProfileCollection,
 	                                   FLocalDynamicMeshData& OutLocalMeshData);
 
 	static void UpdateFoliagePivotPoints(const TSet<int32>& PointsToRemove, FManagedArrayCollection& OutCollection);
