[PVE] FMeshVertex::NjordPixelId was internally setting PointIndex to the mesh. It was confusing to name it Pixel ID. So renamed FMeshVertex::NjordPixelId to PointIndex in PVMeshBuilder.

#Synced-CL 52379040

diff --git a/Engine/Plugins/Experimental/ProceduralVegetationEditor/Source/ProceduralVegetation/Private/Implementations/PVMeshBuilder.cpp b/Engine/Plugins/Experimental/ProceduralVegetationEditor/Source/ProceduralVegetation/Private/Implementations/PVMeshBuilder.cpp
--- a/Engine/Plugins/Experimental/ProceduralVegetationEditor/Source/ProceduralVegetation/Private/Implementations/PVMeshBuilder.cpp
+++ b/Engine/Plugins/Experimental/ProceduralVegetationEditor/Source/ProceduralVegetation/Private/Implementations/PVMeshBuilder.cpp
@@ -183,7 +183,7 @@ void FPVMeshBuilder::GenerateGeometryCollection(const FManagedArrayCollection& I
 			UVs[VertexIndex] = Vertex.UV;
 			UVs1[VertexIndex] = Vertex.UV1;
 			UVs2[VertexIndex] = Vertex.UV2;
-			PointIDs[VertexIndex] = Vertex.NjordPixelId;
+			PointIDs[VertexIndex] = Vertex.PointIndex;
 			++VertexIndex;
 		}
 
@@ -946,10 +946,7 @@ void FPVMeshBuilder::GenerateBranchMeshData(const bool bPrimitiveIsTrunk, const
 	for (int32 BranchPointIndex = 0; BranchPointIndex < PrimitivePoints.Num(); ++BranchPointIndex)
 	{
 		const int32& PointIndex = PrimitivePoints[BranchPointIndex];
-		const float NjordPixelId = PointIndex;
-
 		const float LengthFromRoot = PointFacade.GetLengthFromRoot(PointIndex);
-
 		const FVector3f& PointPosition = PointPositions[PointIndex];
 
 		const float PointScale = PV::AttributesHelper::GetBranchPointScale(PointScaleAttribute, BranchPointsAttribute, BranchParentNumberAttribute, BranchIndex, BranchPointIndex);
@@ -1041,7 +1038,7 @@ void FPVMeshBuilder::GenerateBranchMeshData(const bool bPrimitiveIsTrunk, const
 		if (PointScale == 0.0f)
 		[[unlikely]]
 		{
-			CurrentIndices.Add(OutLocalMeshData.Vertices.Emplace(PointPosition, TangentVector, FVector2f(UVOGenerationOffset, TextureCoordV), UV1, UV2, NjordPixelId));
+			CurrentIndices.Add(OutLocalMeshData.Vertices.Emplace(PointPosition, TangentVector, FVector2f(UVOGenerationOffset, TextureCoordV), UV1, UV2, PointIndex));
 		}
 		else
 		[[likely]]
@@ -1092,7 +1089,7 @@ void FPVMeshBuilder::GenerateBranchMeshData(const bool bPrimitiveIsTrunk, const
 
 				const FVector3f Normal = FVector3f((RadialPoint - PointPosition).GetUnsafeNormal());
 
-				CurrentIndices.Add(OutLocalMeshData.Vertices.Emplace(RadialPoint, Normal, FVector2f(UV_U, TextureCoordV), UV1, UV2, NjordPixelId));
+				CurrentIndices.Add(OutLocalMeshData.Vertices.Emplace(RadialPoint, Normal, FVector2f(UV_U, TextureCoordV), UV1, UV2, PointIndex));
 			}
 		}
 
@@ -1103,7 +1100,7 @@ void FPVMeshBuilder::GenerateBranchMeshData(const bool bPrimitiveIsTrunk, const
 		UV1.X = 0;
 		
 		const FLocalDynamicMeshData::FVertex FirstVertex = OutLocalMeshData.Vertices[CurrentIndices[0]];
-		CurrentIndices.Add(OutLocalMeshData.Vertices.Emplace(FirstVertex.Position, FirstVertex.Normal,UV0, UV1, UV2 ,NjordPixelId));
+		CurrentIndices.Add(OutLocalMeshData.Vertices.Emplace(FirstVertex.Position, FirstVertex.Normal,UV0, UV1, UV2 , PointIndex));
 
 		PreviousPosition = PointPosition;
 		
diff --git a/Engine/Plugins/Experimental/ProceduralVegetationEditor/Source/ProceduralVegetation/Private/Implementations/PVMeshBuilder.h b/Engine/Plugins/Experimental/ProceduralVegetationEditor/Source/ProceduralVegetation/Private/Implementations/PVMeshBuilder.h
--- a/Engine/Plugins/Experimental/ProceduralVegetationEditor/Source/ProceduralVegetation/Private/Implementations/PVMeshBuilder.h
+++ b/Engine/Plugins/Experimental/ProceduralVegetationEditor/Source/ProceduralVegetation/Private/Implementations/PVMeshBuilder.h
@@ -218,7 +218,7 @@ struct FLocalDynamicMeshData
 		FVector2f UV;
 		FVector2f UV1;
 		FVector2f UV2;
-		float NjordPixelId;
+		float PointIndex;
 	};
 
 	TArray<FVertex> Vertices;
