[PVE] Add param debug visualization support to the Procedural Vegetation Editor, introducing thread-safe macros for capturing point, direction, vector, and text parameters during mesh generation
Refactor debug visualizer API to accept UPVData directly and unify loop and param debug rendering paths
Rename loop debug data types to a generalized param debug naming convention shared between loop and non-loop debug contexts

#Synced-CL 52466418

#Synced-CL 52466418

diff --git a/Engine/Plugins/Experimental/ProceduralVegetationEditor/Source/ProceduralVegetation/Private/Implementations/PVMeshBuilder.cpp b/Engine/Plugins/Experimental/ProceduralVegetationEditor/Source/ProceduralVegetation/Private/Implementations/PVMeshBuilder.cpp
--- a/Engine/Plugins/Experimental/ProceduralVegetationEditor/Source/ProceduralVegetation/Private/Implementations/PVMeshBuilder.cpp
+++ b/Engine/Plugins/Experimental/ProceduralVegetationEditor/Source/ProceduralVegetation/Private/Implementations/PVMeshBuilder.cpp
@@ -1023,6 +1023,11 @@ void FPVMeshBuilder::GenerateBranchMeshData(const bool bPrimitiveIsTrunk, const
 			UpVector = FVector3f(Axis1);
 		}
 
+		PVE_DEBUG_LOCK()
+			PVE_PARAM_DEBUG_DIRECTION_PARAM("UpVector", PointPosition, UpVector);
+			PVE_PARAM_DEBUG_VECTOR_PARAM("TangentVector", PointPosition, PointPosition + TangentVector);
+		PVE_DEBUG_UNLOCK()
+
 		PreviousUpVector = UpVector;
 
 		FVector2f UV1 = FVector2f(0, 1.0 - Displacement_UV_V);
diff --git a/Engine/Plugins/Experimental/ProceduralVegetationEditor/Source/ProceduralVegetation/Private/Nodes/PVMeshBuilderSettings.cpp b/Engine/Plugins/Experimental/ProceduralVegetationEditor/Source/ProceduralVegetation/Private/Nodes/PVMeshBuilderSettings.cpp
--- a/Engine/Plugins/Experimental/ProceduralVegetationEditor/Source/ProceduralVegetation/Private/Nodes/PVMeshBuilderSettings.cpp
+++ b/Engine/Plugins/Experimental/ProceduralVegetationEditor/Source/ProceduralVegetation/Private/Nodes/PVMeshBuilderSettings.cpp
@@ -229,9 +229,13 @@ bool FPVMeshBuilderElement::ExecuteInternal(FPCGContext* InContext) const
 
 			UPVMeshData* OutManagedArrayCollectionData = FPCGContext::NewObject_AnyThread<UPVMeshData>(InContext);
 
+			PVE_PARAM_DEBUG_INIT(Settings, OutManagedArrayCollectionData);
+			
 			FGeometryCollection OutGeometryCollection;
 			FPVMeshBuilder::GenerateGeometryCollection(SkeletonCollection, PlantProfileCollection, Settings->MesherSettings, OutGeometryCollection);
 
+			PVE_PARAM_DEBUG_END()
+			
 			FManagedArrayCollection ManagedArrayCollection;
 			OutGeometryCollection.CopyTo(&ManagedArrayCollection);
 			OutManagedArrayCollectionData->Initialize(MoveTemp(ManagedArrayCollection));
diff --git a/Engine/Plugins/Experimental/ProceduralVegetationEditor/Source/ProceduralVegetation/Public/DataTypes/PVData.h b/Engine/Plugins/Experimental/ProceduralVegetationEditor/Source/ProceduralVegetation/Public/DataTypes/PVData.h
--- a/Engine/Plugins/Experimental/ProceduralVegetationEditor/Source/ProceduralVegetation/Public/DataTypes/PVData.h
+++ b/Engine/Plugins/Experimental/ProceduralVegetationEditor/Source/ProceduralVegetation/Public/DataTypes/PVData.h
@@ -132,7 +132,7 @@ struct FPVVisualizationSettings
 };
 
 USTRUCT()
-struct FLoopDebugData
+struct FPVParamDebugData
 {
 	GENERATED_BODY()
 	
@@ -147,7 +147,7 @@ struct FLoopDebugData
 };
 
 USTRUCT()
-struct FPVLoopDebugSettings
+struct FPVParamDebuggerSettings
 {
 	GENERATED_BODY()
 
@@ -158,7 +158,7 @@ struct FPVLoopDebugSettings
 	FText ParamName;
 
 	UPROPERTY(EditAnywhere, Category = LoopDebugSettings)
-	FLoopDebugData Data;
+	FPVParamDebugData Data;
 };
 
 USTRUCT()
@@ -170,8 +170,8 @@ struct FPVDebugSettings
 	TArray<FPVVisualizationSettings> VisualizationSettings;
 
 	UPROPERTY()
-	TArray<FPVLoopDebugSettings> LoopDebugVisualizationSettings;
-
+	TArray<FPVParamDebuggerSettings> ParamDebugVisualizationSettings;
+	
 	bool bAutoFocusLoopDebug = false;
 };
 
diff --git a/Engine/Plugins/Experimental/ProceduralVegetationEditor/Source/ProceduralVegetation/Public/Helpers/PVUtilities.h b/Engine/Plugins/Experimental/ProceduralVegetationEditor/Source/ProceduralVegetation/Public/Helpers/PVUtilities.h
--- a/Engine/Plugins/Experimental/ProceduralVegetationEditor/Source/ProceduralVegetation/Public/Helpers/PVUtilities.h
+++ b/Engine/Plugins/Experimental/ProceduralVegetationEditor/Source/ProceduralVegetation/Public/Helpers/PVUtilities.h
@@ -68,7 +68,16 @@ struct FLoopDebugStepper
 	}
 };
 
-struct FLoopDebuggerState
+struct FPVParamDebuggerState
+{
+	UPVData* Data = nullptr;
+	
+#if WITH_EDITORONLY_DATA
+	FPVDebugSettings DebugSettings;
+#endif
+};
+
+struct FPVLoopDebuggerState
 {
 	uint32 NumSteps = 0;
 	const FLoopDebugStepper* LoopDebugStepper = nullptr;
@@ -107,60 +116,115 @@ struct FLoopDebuggerState
 
 #if WITH_EDITORONLY_DATA
 
-inline TUniquePtr<FLoopDebuggerState> GlobalLoopDebuggerState;
+inline FCriticalSection GlobalPVDebuggerLock;
+inline TUniquePtr<FPVParamDebuggerState> GlobalPVDebuggerState;
+
+#define PVE_DEBUG_LOCK()																									\
+	{																														\
+		FScopeLock Lock(&GlobalPVDebuggerLock);
+#define PVE_DEBUG_UNLOCK()																									\
+	}
+
+#define PVE_PARAM_DEBUG_INIT(InputSettings, PVData)																			\
+	if (InputSettings->bDebug)																								\
+	{																														\
+		GlobalPVDebuggerState = MakeUnique<FPVParamDebuggerState>(PVData);													\
+		GlobalPVDebuggerState->DebugSettings = PVData->GetDebugSettings();													\
+	}
+
+#define PVE_PARAM_DEBUG_POINT_PARAM(InParamName, InPivot)																	\
+	if (GlobalPVDebuggerState)																								\
+		GlobalPVDebuggerState->DebugSettings.ParamDebugVisualizationSettings.Emplace(										\
+			EPVDebugValueVisualizationMode::Point,																			\
+			FText::FromString(TEXT(InParamName)),																			\
+			FPVParamDebugData(FVector(InPivot), NullOpt, NullOpt)															\
+		);
+
+#define PVE_PARAM_DEBUG_DIRECTION_PARAM(InParamName, InPivot, InDirection)													\
+	if (GlobalPVDebuggerState)																								\
+		GlobalPVDebuggerState->DebugSettings.ParamDebugVisualizationSettings.Emplace(										\
+			EPVDebugValueVisualizationMode::Direction,																		\
+			FText::FromString(TEXT(InParamName)),																			\
+			FPVParamDebugData(FVector(InPivot), FVector(InDirection), NullOpt)												\
+		);
+
+#define PVE_PARAM_DEBUG_VECTOR_PARAM(InParamName, InPivot, InVector)														\
+	if (GlobalPVDebuggerState)																								\
+		GlobalPVDebuggerState->DebugSettings.ParamDebugVisualizationSettings.Emplace(										\
+			EPVDebugValueVisualizationMode::Vector,																			\
+			FText::FromString(TEXT(InParamName)),																			\
+			FPVParamDebugData(FVector(InPivot), FVector(InVector), NullOpt)													\
+		);
+
+#define PVE_PARAM_DEBUG_TEXT_PARAM(InParamName, InPivot, InText)															\
+	if (GlobalPVDebuggerState)																								\
+		GlobalPVDebuggerState->DebugSettings.ParamDebugVisualizationSettings.Emplace(										\
+			EPVDebugValueVisualizationMode::Text,																			\
+			FText::FromString(TEXT(InParamName)),																			\
+			FPVParamDebugData(FVector(InPivot), NullOpt, InText)															\
+		);
+
+#define PVE_PARAM_DEBUG_END()																								\
+	if (GlobalPVDebuggerState)																								\
+	{																														\
+		GlobalPVDebuggerState->Data->SetDebugSettings(GlobalPVDebuggerState->DebugSettings);								\
+		GlobalPVDebuggerState.Reset();																						\
+	}
+
+inline TUniquePtr<FPVLoopDebuggerState> GlobalPVLoopDebuggerState;
 
 #define PVE_LOOP_DEBUG_INIT(InputSettings, LoopDebugStepperName, PVData)													\
 	if (InputSettings->bDebug && InputSettings->LoopDebugStepperName.bDebug)												\
 	{																														\
-		GlobalLoopDebuggerState = MakeUnique<FLoopDebuggerState>(0, &InputSettings->LoopDebugStepperName, PVData);			\
+		GlobalPVLoopDebuggerState = MakeUnique<FPVLoopDebuggerState>(0, &InputSettings->LoopDebugStepperName, PVData);		\
 	}
 
 #define PVE_LOOP_DEBUG_PARAMS_START()																						\
 	if (																													\
-		GlobalLoopDebuggerState &&																							\
-		GlobalLoopDebuggerState->IsLastLoopIteration() &&																	\
-		!GlobalLoopDebuggerState->bComplete																					\
+		GlobalPVLoopDebuggerState &&																						\
+		GlobalPVLoopDebuggerState->IsLastLoopIteration() &&																	\
+		!GlobalPVLoopDebuggerState->bComplete																				\
 	)																														\
 	{																														\
-		FPVDebugSettings __DebugSettings = GlobalLoopDebuggerState->GetDebugSettings();										\
-		__DebugSettings.bAutoFocusLoopDebug = GlobalLoopDebuggerState->LoopDebugStepper->bDebugFocus;
+		FPVDebugSettings __DebugSettings = GlobalPVLoopDebuggerState->GetDebugSettings();									\
+		__DebugSettings.bAutoFocusLoopDebug = GlobalPVLoopDebuggerState->LoopDebugStepper->bDebugFocus;
 
 #define PVE_LOOP_DEBUG_POINT_PARAM(InParamName, InPivot)																	\
-		__DebugSettings.LoopDebugVisualizationSettings.Emplace(																\
+		__DebugSettings.ParamDebugVisualizationSettings.Emplace(															\
 			EPVDebugValueVisualizationMode::Point,																			\
 			FText::FromString(TEXT(InParamName)),																			\
-			FLoopDebugData(FVector(InPivot), NullOpt, NullOpt)																\
+			FPVParamDebugData(FVector(InPivot), NullOpt, NullOpt)															\
 		);
 
 #define PVE_LOOP_DEBUG_DIRECTION_PARAM(InParamName, InPivot, InDirection)													\
-		__DebugSettings.LoopDebugVisualizationSettings.Emplace(																\
+		__DebugSettings.ParamDebugVisualizationSettings.Emplace(															\
 			EPVDebugValueVisualizationMode::Direction,																		\
 			FText::FromString(TEXT(InParamName)),																			\
-			FLoopDebugData(FVector(InPivot), FVector(InDirection), NullOpt)													\
+			FPVParamDebugData(FVector(InPivot), FVector(InDirection), NullOpt)												\
 		);
 
 #define PVE_LOOP_DEBUG_VECTOR_PARAM(InParamName, InPivot, InVector)															\
-		__DebugSettings.LoopDebugVisualizationSettings.Emplace(																\
+		__DebugSettings.ParamDebugVisualizationSettings.Emplace(															\
 			EPVDebugValueVisualizationMode::Vector,																			\
 			FText::FromString(TEXT(InParamName)),																			\
-			FLoopDebugData(FVector(InPivot), FVector(InVector), NullOpt)													\
+			FPVParamDebugData(FVector(InPivot), FVector(InVector), NullOpt)													\
 		);
 
 #define PVE_LOOP_DEBUG_TEXT_PARAM(InParamName, InPivot, InText)																\
-		__DebugSettings.LoopDebugVisualizationSettings.Emplace(																\
+		__DebugSettings.ParamDebugVisualizationSettings.Emplace(															\
 			EPVDebugValueVisualizationMode::Text,																			\
 			FText::FromString(TEXT(InParamName)),																			\
-			FLoopDebugData(FVector(InPivot), NullOpt, InText)																\
+			FPVParamDebugData(FVector(InPivot), NullOpt, InText)															\
 		);
 
 #define PVE_LOOP_DEBUG_PARAMS_END()																							\
-		if (GlobalLoopDebuggerState)																						\
-			GlobalLoopDebuggerState->SetDebugSettings(MoveTemp(__DebugSettings));											\
+		if (GlobalPVLoopDebuggerState)																						\
+			GlobalPVLoopDebuggerState->SetDebugSettings(MoveTemp(__DebugSettings));											\
 	}
 #define PVE_OUTER_LOOP_DEBUG_CHECK(ExitMethod)																				\
 	if (																													\
-		GlobalLoopDebuggerState &&																							\
-		GlobalLoopDebuggerState->bComplete																					\
+		GlobalPVLoopDebuggerState &&																						\
+		GlobalPVLoopDebuggerState->bComplete																				\
 	)																														\
 	{																														\
 		ExitMethod;																											\
@@ -168,27 +232,39 @@ inline TUniquePtr<FLoopDebuggerState> GlobalLoopDebuggerState;
 
 #define PVE_LOOP_DEBUG_CHECK(ExitMethod)																					\
 	if (																													\
-		GlobalLoopDebuggerState &&																							\
-		GlobalLoopDebuggerState->IsLastLoopIteration()																		\
+		GlobalPVLoopDebuggerState &&																						\
+		GlobalPVLoopDebuggerState->IsLastLoopIteration()																	\
 	)																														\
 	{																														\
-		GlobalLoopDebuggerState->NumSteps++;																				\
-		GlobalLoopDebuggerState->bComplete = true;																			\
+		GlobalPVLoopDebuggerState->NumSteps++;																				\
+		GlobalPVLoopDebuggerState->bComplete = true;																		\
 		ExitMethod;																											\
 	}
 
 #define PVE_LOOP_DEBUG_STEP(ExitMethod)																						\
-	if (GlobalLoopDebuggerState)																							\
+	if (GlobalPVLoopDebuggerState)																							\
 	{																														\
 		PVE_LOOP_DEBUG_CHECK(ExitMethod)																					\
-		GlobalLoopDebuggerState->NumSteps++;																				\
+		GlobalPVLoopDebuggerState->NumSteps++;																				\
 	}
 
 #define PVE_LOOP_DEBUG_END()																								\
-	GlobalLoopDebuggerState.Reset();																						\
+	GlobalPVLoopDebuggerState.Reset();
 
 #else
 
+#define PVE_DEBUG_LOCK()
+#define PVE_DEBUG_UNLOCK()
+
+#define PVE_PARAM_DEBUG_INIT(InputSettings, PVData)
+
+#define PVE_PARAM_DEBUG_POINT_PARAM(InParamName, InPivot)
+#define PVE_PARAM_DEBUG_DIRECTION_PARAM(InParamName, InPivot, InDirection)
+#define PVE_PARAM_DEBUG_VECTOR_PARAM(InParamName, InPivot, InVector)
+#define PVE_PARAM_DEBUG_TEXT_PARAM(InParamName, InPivot, InText)
+
+#define PVE_PARAM_DEBUG_END()
+
 #define PVE_LOOP_DEBUG_INIT(InputSettings, LoopDebugStepperName, PVData)
 
 #define PVE_LOOP_DEBUG_PARAMS_START()
diff --git a/Engine/Plugins/Experimental/ProceduralVegetationEditor/Source/ProceduralVegetationEditor/Private/Visualizations/DebugVisualization/PVDebugVisualizer.cpp b/Engine/Plugins/Experimental/ProceduralVegetationEditor/Source/ProceduralVegetationEditor/Private/Visualizations/DebugVisualization/PVDebugVisualizer.cpp
--- a/Engine/Plugins/Experimental/ProceduralVegetationEditor/Source/ProceduralVegetationEditor/Private/Visualizations/DebugVisualization/PVDebugVisualizer.cpp
+++ b/Engine/Plugins/Experimental/ProceduralVegetationEditor/Source/ProceduralVegetationEditor/Private/Visualizations/DebugVisualization/PVDebugVisualizer.cpp
@@ -25,41 +25,69 @@ FVisualizerDrawContext::FVisualizerDrawContext(
 	, SceneSetupParams(InSceneParams)
 {}
 
-void FPVDebugVisualizer::DrawDebugVisualizations(
-	const TArray<FPVVisualizationSettings>& VisualizationSettings, 
-	const FManagedArrayCollection& Collection, 
-	FPCGSceneSetupParams& InOutParams
-)
+void FPVDebugVisualizer::Draw(const UPVData* InData, FPCGSceneSetupParams& InOutParams, bool bDrawVisualizationCollections)
 {
+	check(InData);
+
+	const TArray<FPVVisualizationSettings>& VisualizationSettings = InData->GetDebugSettings().VisualizationSettings;
 	for (const FPVVisualizationSettings& VisualizationSetting : VisualizationSettings)
 	{
-		FVisualizerDrawContext Context(Collection, VisualizationSetting, InOutParams);
-		if (const FPVDebugVisualizationPtr VisualizationPtr = CreateVisualizer(VisualizationSetting))
+		if (bDrawVisualizationCollections)
+		{
+			for (const FPVVisualizationCollection& VisualizationCollection : InData->GetVisualizationCollections())
+			{
+				FVisualizerDrawContext Context(VisualizationCollection.Collection, VisualizationSetting, InOutParams);
+				if (const FPVDebugVisualizationPtr VisualizationPtr = CreateVisualizer(VisualizationSetting))
+				{
+					VisualizationPtr->Draw(Context);
+				}
+			}
+		}
+		else
 		{
-			VisualizationPtr->Draw(Context);
+			FVisualizerDrawContext Context(InData->GetCollection(), VisualizationSetting, InOutParams);
+			if (const FPVDebugVisualizationPtr VisualizationPtr = CreateVisualizer(VisualizationSetting))
+			{
+				VisualizationPtr->Draw(Context);
+			}
 		}
 	}
 }
 
-void FPVDebugVisualizer::DrawLoopDebug(
-	const TArray<FPVLoopDebugSettings>& LoopDebugVisualizationSettings, 
-	bool bAutoFocusLoopDebug, 
-	FPCGSceneSetupParams& InOutParams
-)
+void FPVDebugVisualizer::DrawDebugParams(const UPVData* InData, FPCGSceneSetupParams& InOutParams)
 {
+	check(InData);
+	
 	TArray<TPair<FColor, FText>> LegendItems;
+	const TArray<FPVParamDebuggerSettings>& ParamDebugVisualizationSettings = InData->GetDebugSettings().ParamDebugVisualizationSettings;
+
+	TArray<FString> UniqueNames;
+	for (const FPVParamDebuggerSettings& DebugVisualizationSetting : ParamDebugVisualizationSettings)
+	{
+		UniqueNames.AddUnique(DebugVisualizationSetting.ParamName.ToString());
+	}
 
 	float Alpha = 0.0f;
-	const float AlphaStep = 1.0f / static_cast<float>(LoopDebugVisualizationSettings.Num());
+	const float AlphaStep = 1.0f / static_cast<float>(UniqueNames.Num());
+	TArray<FLinearColor> UniqueColors;
+	for (const FString& UniqueName : UniqueNames) 
+	{
+		const FLinearColor Color = PV::Utilities::GetRandomHueColor(Alpha);
+		UniqueColors.Add(Color);
+		LegendItems.Emplace(Color.ToFColorSRGB(), FText::FromString(UniqueName));
+		Alpha += AlphaStep;
+	}
+	
 	UPVLineBatchComponent* LineBatchComponent = NewObject<UPVLineBatchComponent>();
 
 	FBox BoundingBox(ForceInit);
-	for (const FPVLoopDebugSettings& LoopDebugVisualizationSetting : LoopDebugVisualizationSettings)
+	for (const FPVParamDebuggerSettings& ParamDebugVisualizationSetting : ParamDebugVisualizationSettings)
 	{
-		const FLinearColor Color = PV::Utilities::GetRandomHueColor(Alpha);
-		VisualizeLoopDebuggerParam(InOutParams, LoopDebugVisualizationSetting, LineBatchComponent, Color, BoundingBox);
-		LegendItems.Emplace(Color.ToFColorSRGB(), LoopDebugVisualizationSetting.ParamName);
-		Alpha += AlphaStep;
+		if (const int32 ColorIndex = UniqueNames.Find(ParamDebugVisualizationSetting.ParamName.ToString()); ColorIndex != INDEX_NONE)
+		{
+			const FLinearColor Color = UniqueColors[ColorIndex];
+			VisualizeDebuggerParam(InOutParams, ParamDebugVisualizationSetting, LineBatchComponent, Color, BoundingBox);
+		}
 	}
 
 	InOutParams.ManagedResources.Add(LineBatchComponent);
@@ -75,7 +103,7 @@ void FPVDebugVisualizer::DrawLoopDebug(
 		return;
 	}
 
-	if (bAutoFocusLoopDebug)
+	if (InData->GetDebugSettings().bAutoFocusLoopDebug)
 	{
 		if (const TSharedPtr<FEditorViewportClient> ViewportClient = ViewportWidget->GetViewportClient())
 		{
@@ -84,19 +112,20 @@ void FPVDebugVisualizer::DrawLoopDebug(
 	}
 }
 
-void FPVDebugVisualizer::VisualizeLoopDebuggerParam(
+void FPVDebugVisualizer::VisualizeDebuggerParam(
 	FPCGSceneSetupParams& InOutParams,
-	const FPVLoopDebugSettings& Settings,
+	const FPVParamDebuggerSettings& Settings,
 	UPVLineBatchComponent* LineBatcher,
 	FLinearColor Color,
 	FBox& BoundingBox
 )
 {
 	constexpr float DefaultTextScale = 0.05f;
+	constexpr ESceneDepthPriorityGroup DepthGroup = SDPG_Foreground;
 
 	if (Settings.VisualizationMode == EPVDebugValueVisualizationMode::Point)
 	{
-		LineBatcher->AddPoint(Settings.Data.Pivot, 5, Color, SDPG_Foreground);
+		LineBatcher->AddPoint(Settings.Data.Pivot, 5, Color, DepthGroup);
 		BoundingBox += Settings.Data.Pivot;
 	}
 	else if (Settings.VisualizationMode == EPVDebugValueVisualizationMode::Direction && Settings.Data.Direction.IsSet())
@@ -120,9 +149,9 @@ void FPVDebugVisualizer::VisualizeLoopDebuggerParam(
 
 			const float ArrowSqrt = FMath::Sqrt(ArrowSize);
 
-			LineBatcher->AddLine(StartPoint, EndPoint, Color, SDPG_Foreground);
-			LineBatcher->AddLine(EndPoint, EndPoint + TM.TransformPosition(FVector(-ArrowSqrt, ArrowSqrt, 0)), Color, SDPG_Foreground);
-			LineBatcher->AddLine(EndPoint, EndPoint + TM.TransformPosition(FVector(-ArrowSqrt, -ArrowSqrt, 0)), Color, SDPG_Foreground);
+			LineBatcher->AddLine(StartPoint, EndPoint, Color, DepthGroup);
+			LineBatcher->AddLine(EndPoint, EndPoint + TM.TransformPosition(FVector(-ArrowSqrt, ArrowSqrt, 0)), Color, DepthGroup);
+			LineBatcher->AddLine(EndPoint, EndPoint + TM.TransformPosition(FVector(-ArrowSqrt, -ArrowSqrt, 0)), Color, DepthGroup);
 		}
 	}
 	else if (Settings.VisualizationMode == EPVDebugValueVisualizationMode::Vector && Settings.Data.Direction.IsSet())
@@ -131,7 +160,7 @@ void FPVDebugVisualizer::VisualizeLoopDebuggerParam(
 			Settings.Data.Pivot,
 			Settings.Data.Direction.GetValue(),
 			Color,
-			SDPG_Foreground
+			DepthGroup
 		);
 		BoundingBox += Settings.Data.Pivot;
 		// BoundingBox += Settings.Data.Direction.GetValue();
diff --git a/Engine/Plugins/Experimental/ProceduralVegetationEditor/Source/ProceduralVegetationEditor/Private/Visualizations/DebugVisualization/PVDebugVisualizer.h b/Engine/Plugins/Experimental/ProceduralVegetationEditor/Source/ProceduralVegetationEditor/Private/Visualizations/DebugVisualization/PVDebugVisualizer.h
--- a/Engine/Plugins/Experimental/ProceduralVegetationEditor/Source/ProceduralVegetationEditor/Private/Visualizations/DebugVisualization/PVDebugVisualizer.h
+++ b/Engine/Plugins/Experimental/ProceduralVegetationEditor/Source/ProceduralVegetationEditor/Private/Visualizations/DebugVisualization/PVDebugVisualizer.h
@@ -24,13 +24,13 @@ struct FVisualizerDrawContext
 class FPVDebugVisualizer
 {
 public:
-	static void DrawDebugVisualizations(const TArray<FPVVisualizationSettings>& VisualizationSettings, const FManagedArrayCollection& Collection, FPCGSceneSetupParams& InOutParams);
-	static void DrawLoopDebug(const TArray<FPVLoopDebugSettings>& LoopDebugVisualizationSettings, bool bAutoFocusLoopDebug, FPCGSceneSetupParams& InOutParams);
+	static void Draw(const UPVData* InData, FPCGSceneSetupParams& InOutParams, bool bDrawVisualizationCollections = false);
+	static void DrawDebugParams(const UPVData* InData, FPCGSceneSetupParams& InOutParams);
 
 private:
-	static void VisualizeLoopDebuggerParam(
+	static void VisualizeDebuggerParam(
 		FPCGSceneSetupParams& InOutParams,
-		const FPVLoopDebugSettings& Settings,
+		const FPVParamDebuggerSettings& Settings,
 		UPVLineBatchComponent* LineBatcher,
 		FLinearColor Color,
 		FBox& BoundingBox
diff --git a/Engine/Plugins/Experimental/ProceduralVegetationEditor/Source/ProceduralVegetationEditor/Private/Visualizations/PVDataVisualization.cpp b/Engine/Plugins/Experimental/ProceduralVegetationEditor/Source/ProceduralVegetationEditor/Private/Visualizations/PVDataVisualization.cpp
--- a/Engine/Plugins/Experimental/ProceduralVegetationEditor/Source/ProceduralVegetationEditor/Private/Visualizations/PVDataVisualization.cpp
+++ b/Engine/Plugins/Experimental/ProceduralVegetationEditor/Source/ProceduralVegetationEditor/Private/Visualizations/PVDataVisualization.cpp
@@ -162,14 +162,11 @@ FPCGSetupSceneFunc FPVDataVisualization::GetViewportSetupFunc(const UPCGSettings
 
 		if (PV::Utilities::DebugModeEnabled() && RenderSettings->IsDebug())
 		{
-			FPVDebugVisualizer::DrawDebugVisualizations(WeakData->GetDebugSettings().VisualizationSettings, WeakData->GetCollection(), InOutParams);
-			FPVDebugVisualizer::DrawLoopDebug(WeakData->GetDebugSettings().LoopDebugVisualizationSettings, WeakData->GetDebugSettings().bAutoFocusLoopDebug, InOutParams);
-		}
-
-		for (const FPVVisualizationCollection& VisualizationCollection : WeakData->GetVisualizationCollections())
-		{
-			FPVDebugVisualizer::DrawDebugVisualizations(VisualizationCollection.VisualizationSettings, VisualizationCollection.Collection, InOutParams);
+			FPVDebugVisualizer::Draw(WeakData.Get(), InOutParams);
+			FPVDebugVisualizer::DrawDebugParams(WeakData.Get(), InOutParams);
 		}
+		
+		FPVDebugVisualizer::Draw(WeakData.Get(), InOutParams, true);
  
 		InOutParams.FocusBounds = FBoxSphereBounds(Bounds.Origin, Bounds.BoxExtent, Bounds.SphereRadius);
 	};
