[PVE] Fix fused branch root point scale in RemoveBranches, BoneFacade, and Utilities

Use GetBranchPointScale for branch root PScale in:
- PVRemoveBranches::GatherBranchRemovalData (ERemoveBranchesBasis::Radius) -- was
using the parent's scale at the fused point, causing radius-based pruning to
threshold against the wrong (larger) value

- PVBoneFacade::SetBranchSimulationGroupIndex -- was using the parent's scale at
the fused point for bone thickness normalization, causing non-trunk branches to
be grouped incorrectly

- PVBoneFacade::CreateBoneData -- was using the parent's scale at the fused point
for the child branch connection radius test, causing extra branches to be
incorrectly included in bone segments

- Refactor FindBoneForPoint into FindParentBoneForBranch, taking position and scale
directly instead of deriving them from PointFacade internally; fixes the fused
point scale being used as the bone search radius, which caused incorrect parent
bone matches for non-trunk branches

- AddInterpolatedPointToCollection -- add BranchIndex parameter and use
GetBranchPointScale for the current point, fixing incorrect (parent's) scale
being interpolated when subdividing segments starting at the fused point

#Synced-CL 52379040

diff --git a/Engine/Plugins/Experimental/ProceduralVegetationEditor/Source/ProceduralVegetation/Private/Facades/PVBoneFacade.cpp b/Engine/Plugins/Experimental/ProceduralVegetationEditor/Source/ProceduralVegetation/Private/Facades/PVBoneFacade.cpp
--- a/Engine/Plugins/Experimental/ProceduralVegetationEditor/Source/ProceduralVegetation/Private/Facades/PVBoneFacade.cpp
+++ b/Engine/Plugins/Experimental/ProceduralVegetationEditor/Source/ProceduralVegetation/Private/Facades/PVBoneFacade.cpp
@@ -8,6 +8,7 @@
 #include "Facades/PVPointFacade.h"
 #include "GeometryCollection/Facades/CollectionConstraintOverrideFacade.h"
 #include "Helpers/PVUtilities.h"
+#include "Helpers/PVAttributesHelper.h"
 
 class FGeometryCollection;
 
@@ -103,16 +104,13 @@ namespace PV::Facades
 		return ParentBranchBones;
 	}
 
-	FBoneNode* FindBoneForPoint(const FPointFacade& PointFacade, TArray<FBoneNode>& BoneNodes, int32 BonePoint, FString OutMessage, bool bFallbackToClosestBone = false, float MaxDistance = 1.0)
+	FBoneNode* FindParentBoneForBranch(TArray<FBoneNode>& BoneNodes, FVector BonePosition, float BoneScale,
+	                            FString& OutMessage, bool bFallbackToClosestBone = false)
 	{
-		const TManagedArray<FVector3f>& Positions = PointFacade.GetPositions();
-		FVector BonePosition = FVector(Positions[BonePoint]);
-		float BonePScale = PointFacade.GetPointScale(BonePoint);
-				
 		FBoneNode* FoundBone = nullptr;
 		FBoneNode* ClosestBone = nullptr;
 		float MinDist = FLT_MAX;
-				
+
 		for (int32 NodeIdx = 0; NodeIdx < BoneNodes.Num(); ++NodeIdx)
 		{
 			FVector NodePosition = BoneNodes[NodeIdx].AbsolutePosition;
@@ -124,7 +122,7 @@ namespace PV::Facades
 
 			float Dist = FMath::PointDistToSegment(BonePosition, NodePosition, NextPosition);
 			MinDist = FMath::Min(MinDist, Dist);
-			if (Dist <= BonePScale * 2)
+			if (Dist <= BoneScale * 2)
 			{
 				FoundBone = &BoneNodes[NodeIdx];
 			}
@@ -137,10 +135,10 @@ namespace PV::Facades
 
 		if (!FoundBone && bFallbackToClosestBone)
 		{
-			OutMessage = FString::Format(TEXT("FindBoneForPoint fallback to closest bone for point index {0}") , {BonePoint});
+			OutMessage = "FindParentBoneForBranch fallback to closest bone";
 			return ClosestBone;
 		}
-		
+
 		return FoundBone;
 	}
 
@@ -192,8 +190,12 @@ namespace PV::Facades
 		}
 	}
 
-	void SetBranchSimulationGroupIndex(const FPointFacade& PointFacade, FBranchFacade& BranchFacade)
+	void SetBranchSimulationGroupIndex(const FManagedArrayCollection& Collection, const FPointFacade& PointFacade, FBranchFacade& BranchFacade)
 	{
+		const auto PointScaleAttribute = PV::FPointScaleAttribute::GetAttribute(Collection);
+		const auto BranchPointsAttribute = PV::FBranchPointsAttribute::GetAttribute(Collection);
+		const auto BranchParentNumberAttribute = PV::FBranchParentNumberAttribute::GetAttribute(Collection);
+
 		TArray<float> BranchRootPScales;
 		int32 NumBranches = BranchFacade.GetElementCount();
 		BranchRootPScales.Reserve(NumBranches);
@@ -205,9 +207,8 @@ namespace PV::Facades
 		{
 			const TArray<int32>& BranchPoints = BranchFacade.GetPoints(BranchIdx);
 			check(BranchPoints.Num() > 0);
-			const int32 BranchRootIndex = BranchPoints[0];
 
-			float Pscale = PointFacade.GetPointScale(BranchRootIndex);
+			float Pscale = PV::AttributesHelper::GetBranchPointScale(PointScaleAttribute, BranchPointsAttribute, BranchParentNumberAttribute, BranchIdx, 0);
 			BranchRootPScales.Emplace(Pscale);
 
 			Min = FMath::Min(Min, Pscale);
@@ -244,7 +245,7 @@ namespace PV::Facades
 		FPointFacade PointFacade = FPointFacade(*BoneName.GetCollection());
 		FBranchFacade BranchFacade = FBranchFacade(*BoneName.GetCollection());
 		
-		SetBranchSimulationGroupIndex(PointFacade, BranchFacade);
+		SetBranchSimulationGroupIndex(Collection, PointFacade, BranchFacade);
 	}
 	
 	bool FBoneFacade::CreateBoneData(TArray<FBoneNode>& Bones,const float ReductionStrength,bool bSetWindData /*= true*/)
@@ -253,6 +254,10 @@ namespace PV::Facades
 		FBranchFacade BranchFacade(VertexPointIds.GetConstCollection());
 		FPlantFacade PlantFacade(VertexPointIds.GetConstCollection());
 
+		const auto PointScaleAttribute = FPointScaleAttribute::GetAttribute(VertexPointIds.GetConstCollection());
+		const auto BranchPointsAttribute = FBranchPointsAttribute::GetAttribute(VertexPointIds.GetConstCollection());
+		const auto BranchParentNumberAttribute = FBranchParentNumberAttribute::GetAttribute(VertexPointIds.GetConstCollection());
+
 		if (!(Utilities::IsValidGrowthData(VertexPointIds.GetConstCollection()) && PlantFacade.IsValid()))
 		{
 			return false;	
@@ -306,7 +311,9 @@ namespace PV::Facades
 				TArray<FBoneNode> ParentBoneNodes = GetParentBranchBones(Bones, ParentIndex);
 			
 				FString Message;
-				FBoneNode* FoundBoneNode = FindBoneForPoint(PointFacade, ParentBoneNodes, BranchPoints[0], Message, true);
+				const FVector BranchRootPosition = FVector(Positions[BranchPoints[0]]);
+				const float BranchRootScale = PV::AttributesHelper::GetBranchPointScale(PointScaleAttribute, BranchPointsAttribute, BranchParentNumberAttribute, BranchIndex, 0);
+				FBoneNode* FoundBoneNode = FindParentBoneForBranch(ParentBoneNodes, BranchRootPosition, BranchRootScale, Message, true);
 
 				if (FoundBoneNode == nullptr && i != 0)
 				{
@@ -331,7 +338,7 @@ namespace PV::Facades
 					const int32 PointIdx = BranchPoints[Idx];
 					check(Positions.Num() > PointIdx && PointIdx >= 0);
 					const FVector PointPosition = FVector(Positions[PointIdx]);
-					const float PScale = PointFacade.GetPointScale(PointIdx);
+					const float PScale = PV::AttributesHelper::GetBranchPointScale(PointScaleAttribute, BranchPointsAttribute, BranchParentNumberAttribute, BranchIndex, Idx);
 				
 					FVector NextPosition = PointPosition;
 					if (Idx + 1 < BranchPoints.Num())
diff --git a/Engine/Plugins/Experimental/ProceduralVegetationEditor/Source/ProceduralVegetation/Private/Helpers/PVUtilities.cpp b/Engine/Plugins/Experimental/ProceduralVegetationEditor/Source/ProceduralVegetation/Private/Helpers/PVUtilities.cpp
--- a/Engine/Plugins/Experimental/ProceduralVegetationEditor/Source/ProceduralVegetation/Private/Helpers/PVUtilities.cpp
+++ b/Engine/Plugins/Experimental/ProceduralVegetationEditor/Source/ProceduralVegetation/Private/Helpers/PVUtilities.cpp
@@ -1,6 +1,7 @@
 ﻿// Copyright Epic Games, Inc. All Rights Reserved.
  
 #include "Helpers/PVUtilities.h"
+#include "Helpers/PVAttributesHelper.h"
 #include "ProceduralVegetationModule.h"
 #include "ProceduralVegetationPreset.h"
 #include "Helpers/PVFoliageHelper.h"
@@ -335,7 +336,7 @@ namespace PV::Utilities
 		return OutArray;
 	}
 
-	int32 AddInterpolatedPointToCollection(FManagedArrayCollection& OutCollection, const int32 CurrentPointIndex, const int32 NextPointIndex,
+	int32 AddInterpolatedPointToCollection(FManagedArrayCollection& OutCollection, int32 BranchIndex, const int32 CurrentPointIndex, const int32 NextPointIndex,
 	                           const float Alpha, const int32 BudNumber, float NjordPixelIdx)
 	{
 		Facades::FPointFacade PointFacade(OutCollection);
@@ -348,7 +349,20 @@ namespace PV::Utilities
 		const float NextPointLengthFromRoot = PointFacade.GetLengthFromRoot(NextPointIndex);
 		const float LengthFromRoot = FMath::Lerp(CurrentPointLengthFromRoot, NextPointLengthFromRoot, Alpha);
 
-		const float CurrentPointScale = PointFacade.GetPointScale(CurrentPointIndex);
+		const auto PointScaleAttribute = PV::FPointScaleAttribute::GetAttribute(OutCollection);
+		const auto BranchPointsAttribute = PV::FBranchPointsAttribute::GetAttribute(OutCollection);
+		const auto BranchParentNumberAttribute = PV::FBranchParentNumberAttribute::GetAttribute(OutCollection);
+		
+		const int32 CurrentBranchPointIndex = BranchPointsAttribute.IsValidIndex(BranchIndex)
+			? BranchPointsAttribute[BranchIndex].IndexOfByKey(CurrentPointIndex)
+			: INDEX_NONE;
+		
+		ensure(CurrentBranchPointIndex != INDEX_NONE);
+		
+		const float CurrentPointScale = CurrentBranchPointIndex != INDEX_NONE
+			? PV::AttributesHelper::GetBranchPointScale(PointScaleAttribute, BranchPointsAttribute, BranchParentNumberAttribute, BranchIndex, CurrentBranchPointIndex)
+			: PointFacade.GetPointScale(CurrentPointIndex);
+		
 		const float NextPointScale = PointFacade.GetPointScale(NextPointIndex);
 		const float PointScale = FMath::Lerp(CurrentPointScale, NextPointScale, Alpha);
 		
diff --git a/Engine/Plugins/Experimental/ProceduralVegetationEditor/Source/ProceduralVegetation/Private/Implementations/PVGrafter.cpp b/Engine/Plugins/Experimental/ProceduralVegetationEditor/Source/ProceduralVegetation/Private/Implementations/PVGrafter.cpp
--- a/Engine/Plugins/Experimental/ProceduralVegetationEditor/Source/ProceduralVegetation/Private/Implementations/PVGrafter.cpp
+++ b/Engine/Plugins/Experimental/ProceduralVegetationEditor/Source/ProceduralVegetation/Private/Implementations/PVGrafter.cpp
@@ -843,14 +843,15 @@ void FPVGrafter::DistributeGraftWithHormoneBasedSettings(
 		{
 			AttachmentPointIndex = PV::Utilities::AddInterpolatedPointToCollection(
 				OutCollection,
+				AttachmentPoint.BranchIndex,
 				AttachmentPoint.PointIndexA,
 				AttachmentPoint.PointIndexB,
 				AttachmentPoint.PointAlpha,
 				CurrentBudNumber,
 				CurrentBudNumber);
-			
+
 			CurrentBudNumber++;
-			
+
 			auto& Points = BranchPointsAttribute[AttachmentPoint.BranchIndex];
 			Points.Add(AttachmentPointIndex);
 			Points.Sort([PointLengthFromRootAttribute](const int32& A, const int32& B)
@@ -991,14 +992,15 @@ void FPVGrafter::DistributeGraftWithParametricSettings(FManagedArrayCollection&
 		{
 			AttachmentPointIndex = PV::Utilities::AddInterpolatedPointToCollection(
 				OutCollection,
+				AttachmentPoint.BranchIndex,
 				AttachmentPoint.PointIndexA,
 				AttachmentPoint.PointIndexB,
 				AttachmentPoint.PointAlpha,
 				CurrentBudNumber,
 				CurrentBudNumber);
-			
+
 			CurrentBudNumber++;
-			
+
 			auto& Points = BranchPointsAttribute[AttachmentPoint.BranchIndex];
 			Points.Add(AttachmentPointIndex);
 			Points.Sort([PointLengthFromRootAttribute](const int32& A, const int32& B)
diff --git a/Engine/Plugins/Experimental/ProceduralVegetationEditor/Source/ProceduralVegetation/Private/Implementations/PVMeshBuilder.cpp b/Engine/Plugins/Experimental/ProceduralVegetationEditor/Source/ProceduralVegetation/Private/Implementations/PVMeshBuilder.cpp
--- a/Engine/Plugins/Experimental/ProceduralVegetationEditor/Source/ProceduralVegetation/Private/Implementations/PVMeshBuilder.cpp
+++ b/Engine/Plugins/Experimental/ProceduralVegetationEditor/Source/ProceduralVegetation/Private/Implementations/PVMeshBuilder.cpp
@@ -305,6 +305,7 @@ void FPVMeshBuilder::AddPointsToSkeleton(FManagedArrayCollection& OutCollection)
 				const float Alpha = (SectionLength * (j + 1)) / DistanceBetweenPoints;
 				const int32 NewPointIndex = PV::Utilities::AddInterpolatedPointToCollection(
 					OutCollection,
+					BranchIndex,
 					CurrentPointIndex,
 					NextPointIndex,
 					Alpha,
diff --git a/Engine/Plugins/Experimental/ProceduralVegetationEditor/Source/ProceduralVegetation/Private/Implementations/PVRemoveBranches.cpp b/Engine/Plugins/Experimental/ProceduralVegetationEditor/Source/ProceduralVegetation/Private/Implementations/PVRemoveBranches.cpp
--- a/Engine/Plugins/Experimental/ProceduralVegetationEditor/Source/ProceduralVegetation/Private/Implementations/PVRemoveBranches.cpp
+++ b/Engine/Plugins/Experimental/ProceduralVegetationEditor/Source/ProceduralVegetation/Private/Implementations/PVRemoveBranches.cpp
@@ -5,6 +5,7 @@
 #include "Facades/PVPlantFacade.h"
 #include "Facades/PVPointFacade.h"
 #include "Facades/PVTreeFacade.h"
+#include "Helpers/PVAttributesHelper.h"
 
 namespace PV::Facades
 {
@@ -12,9 +13,13 @@ namespace PV::Facades
 	static constexpr int BudDevelopmentGenerationIndex = 0;
 	static constexpr int BudDevelopmentAgeIndex = 2;
 
-	void GatherBranchRemovalData(const FBranchFacade& InBranchFacade, const FPointFacade& InPointFacade, const FPlantFacade& InPlantFacade,
+	void GatherBranchRemovalData(const FManagedArrayCollection& Collection, const FBranchFacade& InBranchFacade, const FPointFacade& InPointFacade, const FPlantFacade& InPlantFacade,
 	                             const ERemoveBranchesBasis InBasis, TArray<float>& OutRemovalData)
 	{
+		const auto PointScaleAttribute = PV::FPointScaleAttribute::GetAttribute(Collection);
+		const auto BranchPointsAttribute = PV::FBranchPointsAttribute::GetAttribute(Collection);
+		const auto BranchParentNumberAttribute = PV::FBranchParentNumberAttribute::GetAttribute(Collection);
+
 		const int32 TotalNumberOfBranches = InBranchFacade.GetElementCount();
 		OutRemovalData.Init(-1.0f, TotalNumberOfBranches);
 
@@ -66,7 +71,7 @@ namespace PV::Facades
 				}
 				else if (InBasis == ERemoveBranchesBasis::Radius)
 				{
-					OutRemovalData[BranchIndex] = InPointFacade.GetPointScale(BranchPoints[0]);
+					OutRemovalData[BranchIndex] = PV::AttributesHelper::GetBranchPointScale(PointScaleAttribute, BranchPointsAttribute, BranchParentNumberAttribute, BranchIndex, 0);
 				}
 				else if (InBasis == ERemoveBranchesBasis::Light)
 				{
@@ -143,7 +148,7 @@ void FPVRemoveBranches::ApplyRemoveBranches(const ERemoveBranchesBasis InBasis,
 	}
 
 	TArray<float> OutRemovalData;
-	GatherBranchRemovalData(BranchFacade, PointFacade, PlantFacade, InBasis, OutRemovalData);
+	GatherBranchRemovalData(OutCollection, BranchFacade, PointFacade, PlantFacade, InBasis, OutRemovalData);
 
 	TArray<int> OutBranchesToRemove;
 	FindBranchesToRemove(BranchFacade, OutRemovalData, InThreshold, OutBranchesToRemove);
diff --git a/Engine/Plugins/Experimental/ProceduralVegetationEditor/Source/ProceduralVegetation/Public/Helpers/PVUtilities.h b/Engine/Plugins/Experimental/ProceduralVegetationEditor/Source/ProceduralVegetation/Public/Helpers/PVUtilities.h
--- a/Engine/Plugins/Experimental/ProceduralVegetationEditor/Source/ProceduralVegetation/Public/Helpers/PVUtilities.h
+++ b/Engine/Plugins/Experimental/ProceduralVegetationEditor/Source/ProceduralVegetation/Public/Helpers/PVUtilities.h
@@ -39,7 +39,7 @@ namespace PV::Utilities
 	template <typename T>
 	TArray<T> LerpArrayElements(const TArray<T>& InArray1, const TArray<T>& InArray2, const float Alpha);
 
-	int32 AddInterpolatedPointToCollection(FManagedArrayCollection& OutCollection, const int32 CurrentPointIndex, const int32 NextPointIndex,
+	int32 AddInterpolatedPointToCollection(FManagedArrayCollection& OutCollection, int32 BranchIndex, const int32 CurrentPointIndex, const int32 NextPointIndex,
 	                                  const float Alpha, const int32 BudNumber, float NjordPixelIdx=-1.0f);
 }
 
