[PVE] Fix bud directions for fused branch root points

Add GetBudDirection helper to PV::AttributesHelper that correctly computes the
bud direction frame at any branch point. For non-trunk branches, ring 0 is the
fused point shared with the parent; the stored frame reflects the parent's
orientation. The helper recomputes the frame by deriving the apical direction
from the ring0->ring1 position delta, then rotating ring 1's axillary and up
vectors through the dihedral to produce a consistent child-branch frame at the
fused point.

Use GetBudDirection in:
- PVMeshBuilder::GenerateBranchMeshData -- fixes PreviousUpVector initialization
using the parent's frame at the fused point, causing incorrect mesh ring
orientation for non-trunk branches
- PVParametricDistributionHelper -- fixes initial apical/up vector using the
parent's frame, causing incorrect parametric distribution orientation at
branch roots
- PVHormoneDistributionHelper -- same fix for hormone distribution initial frame
- PVTreeFacade -- use BranchSourceBudNumberAttribute to look up the correct
source bud direction instead of reading directly from BranchPoints[0]

#Synced-CL 52379040

diff --git a/Engine/Plugins/Experimental/ProceduralVegetationEditor/Source/ProceduralVegetation/Private/Helpers/PVAttributesHelper.cpp b/Engine/Plugins/Experimental/ProceduralVegetationEditor/Source/ProceduralVegetation/Private/Helpers/PVAttributesHelper.cpp
--- a/Engine/Plugins/Experimental/ProceduralVegetationEditor/Source/ProceduralVegetation/Private/Helpers/PVAttributesHelper.cpp
+++ b/Engine/Plugins/Experimental/ProceduralVegetationEditor/Source/ProceduralVegetation/Private/Helpers/PVAttributesHelper.cpp
@@ -76,6 +76,56 @@ float PV::AttributesHelper::GetBranchPointScale(
 	return PointScaleAttribute[PointIndex];
 }
 
+TArray<FVector3f> PV::AttributesHelper::GetBudDirection(
+	PV::FBudDirectionAttributeConstView BudDirectionAttribute,
+	PV::FBranchPointsAttributeConstView BranchPointsAttribute,
+	PV::FBranchParentNumberAttributeConstView BranchParentNumberAttribute,
+	PV::FPointPositionAttributeConstView PointPositionAttribute,
+	const int32 BranchIndex,
+	const int32 BranchPointIndex)
+{
+	if (!BranchPointsAttribute.IsValidIndex(BranchIndex))
+	{
+		return TBudDirectionView<FVector3f>::DefaultData;
+	}
+
+	const TArray<int32>& BranchPoints = BranchPointsAttribute[BranchIndex];
+	const bool bIsTrunk = BranchParentNumberAttribute.IsValidIndex(BranchIndex) && BranchParentNumberAttribute[BranchIndex] == 0;
+	const bool bBranchFusedPoint = !bIsTrunk && BranchPointIndex == 0;
+
+	if (bBranchFusedPoint && BranchPoints.IsValidIndex(1))
+	{
+		const int32 CurrentPointIndex = BranchPoints[0];
+		const int32 NextPointIndex = BranchPoints[1];
+		
+		if (ensure(PointPositionAttribute.IsValidIndex(CurrentPointIndex) && PointPositionAttribute.IsValidIndex(NextPointIndex) && BudDirectionAttribute.IsValidIndex(NextPointIndex)))
+		{
+			const FVector3f Dir = (PointPositionAttribute[NextPointIndex] - PointPositionAttribute[CurrentPointIndex]).GetSafeNormal();
+			const FBudDirectionConstView NextBudDirection = BudDirectionAttribute[NextPointIndex];
+			const FQuat4f Rotation = FQuat4f::FindBetweenNormals(NextBudDirection.Apical, Dir);
+
+			TArray<FVector3f> Result = NextBudDirection.Array;
+			FBudDirectionView ResultView(Result);
+			ResultView.Apical = Dir;
+			ResultView.Axillary = Rotation.RotateVector(NextBudDirection.Axillary);
+			ResultView.UpVector = Rotation.RotateVector(NextBudDirection.UpVector);
+			return Result;
+		}
+		else
+		{
+			return TBudDirectionView<FVector3f>::DefaultData;
+		}
+	}
+	
+	if (!BranchPoints.IsValidIndex(BranchPointIndex))
+	{
+		return TBudDirectionView<FVector3f>::DefaultData;
+	}
+
+	const int32 PointIndex = BranchPoints[BranchPointIndex];
+	return BudDirectionAttribute.IsValidIndex(PointIndex) ? BudDirectionAttribute[PointIndex].Array : TBudDirectionView<FVector3f>::DefaultData;
+}
+
 void PV::AttributesHelper::ComputeBudNumbers(const FComputeBudNumbersAttributes& InAttributes)
 {
 	for (int32 PointIndex = 0; PointIndex < InAttributes.PointBudNumber.Num(); ++PointIndex)
diff --git a/Engine/Plugins/Experimental/ProceduralVegetationEditor/Source/ProceduralVegetation/Private/Helpers/PVAttributesHelper.h b/Engine/Plugins/Experimental/ProceduralVegetationEditor/Source/ProceduralVegetation/Private/Helpers/PVAttributesHelper.h
--- a/Engine/Plugins/Experimental/ProceduralVegetationEditor/Source/ProceduralVegetation/Private/Helpers/PVAttributesHelper.h
+++ b/Engine/Plugins/Experimental/ProceduralVegetationEditor/Source/ProceduralVegetation/Private/Helpers/PVAttributesHelper.h
@@ -19,8 +19,19 @@ namespace PV::AttributesHelper
 		int32 BranchIndex,
 		int32 BranchPointIndex
 	);
-	
-	
+
+	// Returns the bud directions for the given branch point. For non-trunk branches, the 0th
+	// branch ring is the fused point shared with the parent; the frame is recomputed by rotating
+	// the first child point's axillary and up vectors through the dihedral from apical_next to dir,
+	// so callers always receive the correct child-branch frame without extra fused-point handling.
+	TArray<FVector3f> GetBudDirection(
+		PV::FBudDirectionAttributeConstView BudDirectionAttribute,
+		PV::FBranchPointsAttributeConstView BranchPointsAttribute,
+		PV::FBranchParentNumberAttributeConstView BranchParentNumberAttribute,
+		PV::FPointPositionAttributeConstView PointPositionAttribute,
+		const int32 BranchIndex,
+		const int32 BranchPointIndex
+	);
 
 	struct FComputeBudNumbersAttributes
 	{
diff --git a/Engine/Plugins/Experimental/ProceduralVegetationEditor/Source/ProceduralVegetation/Private/Helpers/PVHormoneDistributionHelper.cpp b/Engine/Plugins/Experimental/ProceduralVegetationEditor/Source/ProceduralVegetation/Private/Helpers/PVHormoneDistributionHelper.cpp
--- a/Engine/Plugins/Experimental/ProceduralVegetationEditor/Source/ProceduralVegetation/Private/Helpers/PVHormoneDistributionHelper.cpp
+++ b/Engine/Plugins/Experimental/ProceduralVegetationEditor/Source/ProceduralVegetation/Private/Helpers/PVHormoneDistributionHelper.cpp
@@ -13,6 +13,7 @@ void PV::HormoneDistributionHelper::ComputeHormoneBasedAttachmentPoints(
 	FRandomStream& RandomStream, TArray<DistributionHelper::FAttachmentPoint>& OutAttachmentPoints)
 {
 	auto BranchPointsAttribute = PV::FBranchPointsAttribute::GetAttribute(OutCollection);
+	auto BranchParentNumberAttribute = PV::FBranchParentNumberAttribute::GetAttribute(OutCollection);
 	auto BudDirectionAttribute = PV::FBudDirectionAttribute::GetAttribute(OutCollection);
 	auto BudHormoneLevels = PV::FBudHormoneLevelsAttribute::GetAttribute(OutCollection);
 	auto PointLengthFromRootAttribute = PV::FPointLengthFromRootAttribute::GetAttribute(OutCollection);
@@ -21,6 +22,7 @@ void PV::HormoneDistributionHelper::ComputeHormoneBasedAttachmentPoints(
 
 	if (!ValidateAttributeCollection(
 		BranchPointsAttribute,
+		BranchParentNumberAttribute,
 		BudDirectionAttribute,
 		BudHormoneLevels,
 		PointLengthFromRootAttribute,
@@ -60,6 +62,18 @@ void PV::HormoneDistributionHelper::ComputeHormoneBasedAttachmentPoints(
 		const float BranchLength = LastPointLengthFromRoot - FirstPointLengthFromRoot;
 
 		if (BranchLength <= KINDA_SMALL_NUMBER) continue;
+		
+		// Precompute resolved bud directions for each branch point
+		TArray<TArray<FVector3f>> BranchBudDirPerPoint;
+		
+		BranchBudDirPerPoint.Reserve(BranchPoints.Num());
+		
+		for (int32 i = 0; i < BranchPoints.Num(); ++i)
+		{
+			BranchBudDirPerPoint.Add(AttributesHelper::GetBudDirection(
+				BudDirectionAttribute, BranchPointsAttribute, BranchParentNumberAttribute,
+				PointPositionAttribute, BranchIndex, i));
+		}
 
 		// Calculate Normalized Attachment Lookups
 		TArray<float> NormalizedAttachmentPoints;
@@ -88,9 +102,11 @@ void PV::HormoneDistributionHelper::ComputeHormoneBasedAttachmentPoints(
 		// Compute Position, Angles, Scale and LengthFromRoot for attachment points
 
 		// Compute based on first point attributes
-		const FVector3f InitialApicalDirection(BudDirectionAttribute[BranchPoints[0]].Apical);
-		FVector3f InitialAxillaryDirection(BudDirectionAttribute[BranchPoints[0]].Axillary);
-		const FVector3f InitialUpVector(BudDirectionAttribute[BranchPoints[0]].UpVector);
+		const FBudDirectionConstView InitialBudDirection(BranchBudDirPerPoint[0]);
+		
+		const FVector3f InitialApicalDirection(InitialBudDirection.Apical);
+		FVector3f InitialAxillaryDirection(InitialBudDirection.Axillary);
+		const FVector3f InitialUpVector(InitialBudDirection.UpVector);
 
 		if (bResetPhyllotaxyLeaf)
 		{
@@ -146,8 +162,11 @@ void PV::HormoneDistributionHelper::ComputeHormoneBasedAttachmentPoints(
 						? ScaleSettings.ScaleRamp->GetRichCurveConst()->Eval(1.0f - BranchGradient)
 						: 1.0f;
 
-					const float BranchScale = FMath::Lerp(PointScaleAttribute[CurrentPointIndex], PointScaleAttribute[NextPointIndex], DistanceAlpha);
+					const float CurrentScale = PV::AttributesHelper::GetBranchPointScale(PointScaleAttribute, BranchPointsAttribute, BranchParentNumberAttribute, BranchIndex, BranchPointIndex);
+					const float NextScale = PV::AttributesHelper::GetBranchPointScale(PointScaleAttribute, BranchPointsAttribute, BranchParentNumberAttribute, BranchIndex, BranchPointIndex + 1);
+					const float BranchScale = FMath::Lerp(CurrentScale, NextScale, DistanceAlpha);
 					const float BranchScaledValue = BranchGradientRampValue * (BranchScale + 0.1f);
+					
 					BranchGradientRampValue = FMath::Clamp(FMath::Lerp(BranchGradientRampValue, BranchScaledValue, ScaleSettings.BranchScaleImpact),
 						0.0f, 1.0f);
 
@@ -156,19 +175,17 @@ void PV::HormoneDistributionHelper::ComputeHormoneBasedAttachmentPoints(
 					Scale *= RandomStream.FRandRange(ScaleSettings.RandomScaleMin, ScaleSettings.RandomScaleMax);
 
 					// Conditions
-					const FVector3f BudApicalDirectionVector = FMath::Lerp(BudDirectionAttribute[CurrentPointIndex].Apical,
-						BudDirectionAttribute[NextPointIndex].Apical, DistanceAlpha);
+					const FBudDirectionConstView CurrentBudDir(BranchBudDirPerPoint[BranchPointIndex]);
+					const FBudDirectionConstView NextBudDir(BranchBudDirPerPoint[BranchPointIndex + 1]);
+					const FVector3f BudApicalDirectionVector = FMath::Lerp(CurrentBudDir.Apical, NextBudDir.Apical, DistanceAlpha);
+					
+					const bool bIsTip = (BranchPointIndex == BranchPoints.Num() - 2) && FMath::IsNearlyEqual(DistanceAlpha, 1.0f, 1e-3f);
 
 					DistributionConditionUtils::ConditionPicker::FPointConditionSample ConditionSample;
 					ConditionSample.UpAlignment = DistributionConditionUtils::UpAlignmentNormalized(BudApicalDirectionVector);
-					ConditionSample.Light = FMath::Lerp(NormalizedPointCaches.Light[CurrentPointIndex], NormalizedPointCaches.Light[NextPointIndex],
-						DistanceAlpha);
-					ConditionSample.Health = FMath::Lerp(NormalizedPointCaches.Health[CurrentPointIndex],
-						NormalizedPointCaches.Health[NextPointIndex], DistanceAlpha);
-					ConditionSample.Scale = FMath::Lerp(NormalizedPointCaches.PScales[CurrentPointIndex],
-						NormalizedPointCaches.PScales[NextPointIndex], DistanceAlpha);
-
-					const bool bIsTip = (BranchPointIndex == BranchPoints.Num() - 2) && FMath::IsNearlyEqual(DistanceAlpha, 1.0f, 1e-3f);
+					ConditionSample.Light = FMath::Lerp(NormalizedPointCaches.Light[CurrentPointIndex], NormalizedPointCaches.Light[NextPointIndex], DistanceAlpha);
+					ConditionSample.Health = FMath::Lerp(NormalizedPointCaches.Health[CurrentPointIndex],NormalizedPointCaches.Health[NextPointIndex], DistanceAlpha);
+					ConditionSample.Scale = FMath::Lerp(NormalizedPointCaches.PScales[CurrentPointIndex], NormalizedPointCaches.PScales[NextPointIndex], DistanceAlpha);
 					ConditionSample.Tip = DistributionConditionUtils::Tip(bIsTip);
 
 					// Tip detection logic shared by Grafter
@@ -182,8 +199,7 @@ void PV::HormoneDistributionHelper::ComputeHormoneBasedAttachmentPoints(
 					TArray<FVector3f> ComputedUpVectors;
 					TArray<FVector3f> ComputedNormalVectors;
 
-					const FVector3f BudLightOptimalVector = FMath::Lerp(BudDirectionAttribute[CurrentPointIndex].LightOptimal,
-						BudDirectionAttribute[NextPointIndex].LightOptimal, DistanceAlpha);
+					const FVector3f BudLightOptimalVector = FMath::Lerp(CurrentBudDir.LightOptimal, NextBudDir.LightOptimal, DistanceAlpha);
 					const FQuat4f QuatRotation = FQuat4f::FindBetweenNormals(BudApicalDirectionVector.GetSafeNormal(),
 						BudLightOptimalVector.GetSafeNormal());
 					const FVector3f BaseNormalVector = QuatRotation.RotateVector(BudApicalDirectionVector);
@@ -198,8 +214,10 @@ void PV::HormoneDistributionHelper::ComputeHormoneBasedAttachmentPoints(
 						float AxillaryRotationDegree = FMath::DegreesToRadians(PhyllotaxyFormationLeaf * AxillaryRotationIteration);
 						const FQuat4f RotationQuat(InitialApicalDirection.GetSafeNormal(), AxillaryRotationDegree);
 						FVector3f NewAxillaryDirection = RotationQuat.RotateVector(InitialAxillaryDirection.GetSafeNormal());
+						
 						FQuat4f ApicalCorrection = FQuat4f::FindBetweenNormals(InitialApicalDirection.GetSafeNormal(),
 							BudApicalDirectionVector.GetSafeNormal());
+						
 						NewAxillaryDirection = ApicalCorrection.RotateVector(NewAxillaryDirection).GetSafeNormal();
 
 						NumBuds = RandomStream.RandRange(MinBudsLeaf, MaxBudsLeaf);
@@ -216,12 +234,15 @@ void PV::HormoneDistributionHelper::ComputeHormoneBasedAttachmentPoints(
 								float AxilAngleRampValue = VectorSettings.AxilAngleRamp
 									? FMath::Clamp(VectorSettings.AxilAngleRamp->GetRichCurveConst()->Eval(1.0f - BranchGradient), 0.0f, 1.0f)
 									: 1.0f;
+								
 								AxilAngleRampValue *= VectorSettings.AxilAngleRampUpperValue;
+								
 								const float AxilAngleBlended = FMath::Lerp(VectorSettings.AxilAngle, AxilAngleRampValue,
 									VectorSettings.AxilAngleRampEffect);
 
 								const FVector3f RotationAxis = FVector3f::CrossProduct(BudApicalDirectionVector.GetSafeNormal(),
 									AxillaryDirectionModified.GetSafeNormal());
+								
 								const FQuat4f Rotation(RotationAxis.GetSafeNormal(), FMath::DegreesToRadians(AxilAngleBlended));
 								ComputedUpVectors.Add(Rotation.RotateVector(BudApicalDirectionVector.GetSafeNormal()));
 							}
diff --git a/Engine/Plugins/Experimental/ProceduralVegetationEditor/Source/ProceduralVegetation/Private/Helpers/PVParametricDistributionHelper.cpp b/Engine/Plugins/Experimental/ProceduralVegetationEditor/Source/ProceduralVegetation/Private/Helpers/PVParametricDistributionHelper.cpp
--- a/Engine/Plugins/Experimental/ProceduralVegetationEditor/Source/ProceduralVegetation/Private/Helpers/PVParametricDistributionHelper.cpp
+++ b/Engine/Plugins/Experimental/ProceduralVegetationEditor/Source/ProceduralVegetation/Private/Helpers/PVParametricDistributionHelper.cpp
@@ -177,11 +177,20 @@ void PV::ParametricDistributionHelper::ComputeAttachmentPoints(
 		{
 			continue;
 		}
+		
+		TArray<TArray<FVector3f>> BranchBudDirPerPoint;                                                                                                            
+		BranchBudDirPerPoint.Reserve(BranchPoints.Num());                                                                                                          
+		
+		for (int32 i = 0; i < BranchPoints.Num(); ++i)                                                                                                      
+		{                                                                                                                                                          
+			BranchBudDirPerPoint.Add(PV::AttributesHelper::GetBudDirection(BudDirectionAttribute, BranchPointsAttribute, BranchParentNumberAttribute, PointPositionAttribute, BranchIndex, i));                                                                                                              
+		}
 
 #if DO_ENSURE
 		for (int32 i = 0; i < BranchPoints.Num() - 1; ++i)
 		{
-			ensureMsgf(PointLengthFromRootAttribute[BranchPoints[i]] <= PointLengthFromRootAttribute[BranchPoints[i + 1]],
+			ensureMsgf(
+				PointLengthFromRootAttribute[BranchPoints[i]] <= PointLengthFromRootAttribute[BranchPoints[i + 1]],
 				TEXT("BranchPoints are not ordered root-to-tip at index %d for BranchIndex %d"), i, BranchIndex);
 		}
 #endif
@@ -194,10 +203,11 @@ void PV::ParametricDistributionHelper::ComputeAttachmentPoints(
 
 		FVector3f PreviousAcceptedPosition = PointPositionAttribute[FirstPointIndex];
 
-		const FVector3f InitialApicalDirection = BudDirectionAttribute[FirstPointIndex].Apical.GetSafeNormal();
-		const FVector3f InitialUpVector = BudDirectionAttribute[FirstPointIndex].UpVector.GetSafeNormal();
+		const PV::FBudDirectionConstView InitialBudDirection(BranchBudDirPerPoint[0]);
+		const FVector3f InitialApicalDirection = InitialBudDirection.Apical.GetSafeNormal();
+		const FVector3f InitialUpVector = InitialBudDirection.UpVector.GetSafeNormal();
 
-		FVector3f InitialAxillaryDirection = BudDirectionAttribute[FirstPointIndex].UpVector.GetSafeNormal();
+		FVector3f InitialAxillaryDirection = InitialBudDirection.Axillary.GetSafeNormal();
 		if (bResetPhyllotaxyLeaf)
 		{
 			InitialAxillaryDirection = FVector3f::CrossProduct(InitialApicalDirection, InitialUpVector).GetSafeNormal();
@@ -217,7 +227,7 @@ void PV::ParametricDistributionHelper::ComputeAttachmentPoints(
 		}
 
 		int32 AxillaryRotationIteration = 0;
-
+		
 		for (int32 j = 0; j < LoopNumber; j++)
 		{
 			float BaseValue = static_cast<float>(j) / FMath::Max(LoopNumber - 1, 1);
@@ -227,6 +237,9 @@ void PV::ParametricDistributionHelper::ComputeAttachmentPoints(
 			FVector3f Pos(0.0f, 0.0f, 0.0f);
 			int32 PointIndexA = INDEX_NONE;
 			int32 PointIndexB = INDEX_NONE;
+			int32 PreviousBranchPointIndex = INDEX_NONE;
+			int32 NextBranchPointIndex = INDEX_NONE;
+			
 			float PointAlpha = 0.0f;
 
 			if (SpacingSettings.SpacingBasis == FParametricSettings::EDistributionBasis::Plant)
@@ -260,15 +273,18 @@ void PV::ParametricDistributionHelper::ComputeAttachmentPoints(
 					FVector2f(0.0f, 1.0f),
 					AttachmentPointLengthFromRoot);
 
+				
 				FindPointByNormalizedLengthFromRoot(
 					BranchGradient,
 					BranchPoints,
 					PointLengthFromRootAttribute,
 					LFRofRoot,
 					LFRofTip,
-					PointIndexA,
-					PointIndexB,
+					PreviousBranchPointIndex,
+					NextBranchPointIndex,
 					PointAlpha);
+				
+				
 			}
 			else // SpacingBasis is Branch
 			{
@@ -290,12 +306,14 @@ void PV::ParametricDistributionHelper::ComputeAttachmentPoints(
 					PointLengthFromRootAttribute,
 					LFRofRoot,
 					LFRofTip,
-					PointIndexA,
-					PointIndexB,
+					PreviousBranchPointIndex,
+					NextBranchPointIndex,
 					PointAlpha))
 				{
 					continue;
 				}
+				PointIndexA = BranchPoints[PreviousBranchPointIndex];
+				PointIndexB = BranchPoints[NextBranchPointIndex];
 
 				const FVector3f P0 = PointPositionAttribute[PointIndexA];
 				const FVector3f P1 = PointPositionAttribute[PointIndexB];
@@ -306,11 +324,14 @@ void PV::ParametricDistributionHelper::ComputeAttachmentPoints(
 					FVector2f(LFRofRoot, LFRofTip),
 					BranchGradient);
 			}
-
-			if (PointIndexA == INDEX_NONE || PointIndexB == INDEX_NONE)
+			
+			if (PreviousBranchPointIndex == INDEX_NONE || NextBranchPointIndex == INDEX_NONE)
 			{
 				continue;
 			}
+			
+			PointIndexA = BranchPoints[PreviousBranchPointIndex];
+			PointIndexB = BranchPoints[NextBranchPointIndex];
 
 			const float PlantGradient0 = 1.0f - PointPlantGradientAttribute[PointIndexA];
 			const float PlantGradient1 = 1.0f - PointPlantGradientAttribute[PointIndexB];
@@ -339,13 +360,11 @@ void PV::ParametricDistributionHelper::ComputeAttachmentPoints(
 				NumBuds = 1;
 			}
 
-			const FVector3f BudDirectionApicalA = BudDirectionAttribute[PointIndexA].Apical;
-			const FVector3f BudDirectionApicalB = BudDirectionAttribute[PointIndexB].Apical;
-			const FVector3f BudDirectionUpA = BudDirectionAttribute[PointIndexA].UpVector;
-			const FVector3f BudDirectionUpB = BudDirectionAttribute[PointIndexB].UpVector;
+			const FBudDirectionConstView BudDirA(BranchBudDirPerPoint[PreviousBranchPointIndex]);
+			const FBudDirectionConstView BudDirB(BranchBudDirPerPoint[NextBranchPointIndex]);
 
-			const FVector3f ApicalDirection = FMath::Lerp(BudDirectionApicalA, BudDirectionApicalB, PointAlpha);
-			const FVector3f BudDirectionUp = FMath::Lerp(BudDirectionUpA, BudDirectionUpB, PointAlpha);
+			const FVector3f ApicalDirection = FMath::Lerp(BudDirA.Apical, BudDirB.Apical, PointAlpha);
+			const FVector3f BudDirectionUp = FMath::Lerp(BudDirA.UpVector, BudDirB.UpVector, PointAlpha);
 
 			float SitePhaseDegrees = 0.0f;
 			if (j > 0)
@@ -596,8 +615,8 @@ float PV::ParametricDistributionHelper::SampleFloatByFloat(const float SampleVal
 
 bool PV::ParametricDistributionHelper::FindPointByNormalizedLengthFromRoot(float SampleValue, const TArray<int32>& BranchPoints,
                                                                            FPointLengthFromRootAttributeConstView PointLengthFromRootAttribute,
-                                                                           const float LFRofRoot, const float LFRofTip, int32& PreviousPointIndex,
-                                                                           int32& NextPointIndex, float& OutAlpha)
+                                                                           const float LFRofRoot, const float LFRofTip,
+                                                                           int32& PreviousBranchPointIndex, int32& NextBranchPointIndex, float& OutAlpha)
 {
 	if (BranchPoints.Num() < 2)
 	{
@@ -626,15 +645,15 @@ bool PV::ParametricDistributionHelper::FindPointByNormalizedLengthFromRoot(float
 				? 0.0f
 				: ((SampleLengthFromRoot - LFR0) / Denominator);
 			OutAlpha = FMath::Clamp(OutAlpha, 0.0f, 1.0f);
-			PreviousPointIndex = P;
-			NextPointIndex = N;
+			PreviousBranchPointIndex = i;
+			NextBranchPointIndex = i + 1;
 
 			return true;
 		}
 	}
 
-	PreviousPointIndex = BranchPoints[0];
-	NextPointIndex = BranchPoints.Last();
+	PreviousBranchPointIndex = 0;
+	NextBranchPointIndex = BranchPoints.Num() - 1;
 	OutAlpha = (SampleValue < 0.5f)
 		? 0.0f
 		: 1.0f;
diff --git a/Engine/Plugins/Experimental/ProceduralVegetationEditor/Source/ProceduralVegetation/Private/Helpers/PVParametricDistributionHelper.h b/Engine/Plugins/Experimental/ProceduralVegetationEditor/Source/ProceduralVegetation/Private/Helpers/PVParametricDistributionHelper.h
--- a/Engine/Plugins/Experimental/ProceduralVegetationEditor/Source/ProceduralVegetation/Private/Helpers/PVParametricDistributionHelper.h
+++ b/Engine/Plugins/Experimental/ProceduralVegetationEditor/Source/ProceduralVegetation/Private/Helpers/PVParametricDistributionHelper.h
@@ -112,8 +112,8 @@ namespace PV::ParametricDistributionHelper
 		FPointLengthFromRootAttributeConstView PointLengthFromRootAttribute,
 		const float LFRofRoot,
 		const float LFRofTip,
-		int32& PreviousPointIndex,
-		int32& NextPointIndex,
+		int32& PreviousBranchPointIndex,
+		int32& NextBranchPointIndex,
 		float& OutAlpha);
 
 	void ComputeRotationAnglesForAttachment(
diff --git a/Engine/Plugins/Experimental/ProceduralVegetationEditor/Source/ProceduralVegetation/Private/Helpers/PVUtilities.cpp b/Engine/Plugins/Experimental/ProceduralVegetationEditor/Source/ProceduralVegetation/Private/Helpers/PVUtilities.cpp
--- a/Engine/Plugins/Experimental/ProceduralVegetationEditor/Source/ProceduralVegetation/Private/Helpers/PVUtilities.cpp
+++ b/Engine/Plugins/Experimental/ProceduralVegetationEditor/Source/ProceduralVegetation/Private/Helpers/PVUtilities.cpp
@@ -352,18 +352,26 @@ namespace PV::Utilities
 		const auto PointScaleAttribute = PV::FPointScaleAttribute::GetAttribute(OutCollection);
 		const auto BranchPointsAttribute = PV::FBranchPointsAttribute::GetAttribute(OutCollection);
 		const auto BranchParentNumberAttribute = PV::FBranchParentNumberAttribute::GetAttribute(OutCollection);
+		const auto BudDirectionAttribute = PV::FBudDirectionAttribute::GetAttribute(OutCollection);
+		const auto PointPositionAttribute = PV::FPointPositionAttribute::GetAttribute(OutCollection);
 		
 		const int32 CurrentBranchPointIndex = BranchPointsAttribute.IsValidIndex(BranchIndex)
 			? BranchPointsAttribute[BranchIndex].IndexOfByKey(CurrentPointIndex)
 			: INDEX_NONE;
 		
-		ensure(CurrentBranchPointIndex != INDEX_NONE);
+		const int32 NextBranchPointIndex = BranchPointsAttribute.IsValidIndex(BranchIndex)
+			? BranchPointsAttribute[BranchIndex].IndexOfByKey(NextPointIndex)
+			: INDEX_NONE;
+		
+		check(CurrentBranchPointIndex != INDEX_NONE);
+		check(NextBranchPointIndex != INDEX_NONE);
+
+		const float CurrentPointScale = AttributesHelper::GetBranchPointScale(PointScaleAttribute, BranchPointsAttribute, BranchParentNumberAttribute, 
+				BranchIndex, CurrentBranchPointIndex);
 		
-		const float CurrentPointScale = CurrentBranchPointIndex != INDEX_NONE
-			? PV::AttributesHelper::GetBranchPointScale(PointScaleAttribute, BranchPointsAttribute, BranchParentNumberAttribute, BranchIndex, CurrentBranchPointIndex)
-			: PointFacade.GetPointScale(CurrentPointIndex);
+		const float NextPointScale = AttributesHelper::GetBranchPointScale(PointScaleAttribute, BranchPointsAttribute, BranchParentNumberAttribute, 
+			BranchIndex, NextBranchPointIndex);
 		
-		const float NextPointScale = PointFacade.GetPointScale(NextPointIndex);
 		const float PointScale = FMath::Lerp(CurrentPointScale, NextPointScale, Alpha);
 		
 		float CurrentSeedPScale = 0.f;
@@ -405,8 +413,12 @@ namespace PV::Utilities
 		const TArray<float>& NextPointBudHormoneLevels = PointFacade.GetBudHormoneLevels(NextPointIndex);
 		const TArray<float> BudHormoneLevels = LerpArrayElements(CurrentPointBudHormoneLevels, NextPointBudHormoneLevels, Alpha);
 
-		const TArray<FVector3f>& CurrentPointBudDirections = PointFacade.GetBudDirection(CurrentPointIndex);
-		const TArray<FVector3f>& NextPointBudDirections = PointFacade.GetBudDirection(NextPointIndex);
+		const TArray<FVector3f> CurrentPointBudDirections = AttributesHelper::GetBudDirection(BudDirectionAttribute, BranchPointsAttribute, BranchParentNumberAttribute, 
+			PointPositionAttribute, BranchIndex, CurrentBranchPointIndex);
+		
+		const TArray<FVector3f>& NextPointBudDirections = AttributesHelper::GetBudDirection(BudDirectionAttribute, BranchPointsAttribute, BranchParentNumberAttribute, 
+			PointPositionAttribute, BranchIndex, NextBranchPointIndex);
+		
 		const TArray<FVector3f> BudDirections = LerpArrayElements(CurrentPointBudDirections, NextPointBudDirections, Alpha);
 
 		TArray<int> CurrentPointBudStatus;
diff --git a/Engine/Plugins/Experimental/ProceduralVegetationEditor/Source/ProceduralVegetation/Private/Implementations/PVMeshBuilder.cpp b/Engine/Plugins/Experimental/ProceduralVegetationEditor/Source/ProceduralVegetation/Private/Implementations/PVMeshBuilder.cpp
--- a/Engine/Plugins/Experimental/ProceduralVegetationEditor/Source/ProceduralVegetation/Private/Implementations/PVMeshBuilder.cpp
+++ b/Engine/Plugins/Experimental/ProceduralVegetationEditor/Source/ProceduralVegetation/Private/Implementations/PVMeshBuilder.cpp
@@ -270,7 +270,7 @@ void FPVMeshBuilder::AddPointsToSkeleton(FManagedArrayCollection& OutCollection)
 		return;
 	}
 	
-	const float DivisionUnitLength = LongestBranchLength / 500.0f;
+	const float DivisionUnitLength = LongestBranchLength / 50.0f;
 	
 	for (int32 BranchIndex = 0; BranchIndex < NumOfBranches; BranchIndex++)
 	{
@@ -911,9 +911,6 @@ void FPVMeshBuilder::GenerateBranchMeshData(const bool bPrimitiveIsTrunk, const
 		ProfilePoints = PlantProfileFacade.GetProfilePoints(0);
 	}
 
-	const TManagedArray<FVector3f>& PointPositions = PointFacade.GetPositions();
-	const TManagedArray<TArray<FVector3f>>& PointBudDirections = BudVectorFacade.GetBudDirections();
-
 	const auto PointScaleAttribute = PV::FPointScaleAttribute::GetAttribute(Collection);
 	const auto BranchPointsAttribute = PV::FBranchPointsAttribute::GetAttribute(Collection);
 	const auto BranchParentNumberAttribute = PV::FBranchParentNumberAttribute::GetAttribute(Collection);
@@ -930,10 +927,18 @@ void FPVMeshBuilder::GenerateBranchMeshData(const bool bPrimitiveIsTrunk, const
 		&& DisplacementData.TextureHeight > 0
 		&& GenerationNumber <= MeshBuilderParams.DisplacementGenerationUpperLimit;
 
-	FVector3f PreviousUpVector = PointBudDirections[PrimitivePoints[0]][5].GetUnsafeNormal();
+	const PV::FBudDirectionAttributeConstView BudDirectionAttribute = PV::FBudDirectionAttribute::FindAttribute(Collection);
+	const PV::FPointPositionAttributeConstView PointPositionAttribute = PV::FPointPositionAttribute::FindAttribute(Collection);
+
+	const TArray<FVector3f> RootBudDirectionData = PV::AttributesHelper::GetBudDirection(
+		BudDirectionAttribute, BranchPointsAttribute, BranchParentNumberAttribute, PointPositionAttribute, BranchIndex, 0);
+	
+	const PV::FBudDirectionConstView RootBudDirection(RootBudDirectionData);
+
+	FVector3f PreviousUpVector = RootBudDirection.UpVector.GetUnsafeNormal();
 
 	float Displacement_UV_V = 0.0f;
-	FVector3f PreviousPosition = PointPositions[PrimitivePoints[0]];
+	FVector3f PreviousPosition = PointPositionAttribute[PrimitivePoints[0]];
 
 	auto SafeDivide = [](const double InNumerator, const double InDenominator)
 	{
@@ -948,12 +953,12 @@ void FPVMeshBuilder::GenerateBranchMeshData(const bool bPrimitiveIsTrunk, const
 	{
 		const int32& PointIndex = PrimitivePoints[BranchPointIndex];
 		const float LengthFromRoot = PointFacade.GetLengthFromRoot(PointIndex);
-		const FVector3f& PointPosition = PointPositions[PointIndex];
+		const FVector3f& PointPosition = PointPositionAttribute[PointIndex];
 
 		const float PointScale = PV::AttributesHelper::GetBranchPointScale(PointScaleAttribute, BranchPointsAttribute, BranchParentNumberAttribute, BranchIndex, BranchPointIndex);
 		
 		const auto [TextureCoordV, TextureCoordUOffset, NewRange] =
-			ComputePointUVData(PointIndex, BranchPointIndex, bPrimitiveIsTrunk, PrimitivePoints, PointFacade, PointPositions, PointScale);
+			ComputePointUVData(PointIndex, BranchPointIndex, bPrimitiveIsTrunk, PrimitivePoints, PointFacade, PointPositionAttribute.ManagedArray->GetConstArray(), PointScale);
 		
 		const FVector2f OldRange = FVector2f(0, 1.0f);
 		
@@ -962,22 +967,22 @@ void FPVMeshBuilder::GenerateBranchMeshData(const bool bPrimitiveIsTrunk, const
 		if (PrimitivePoints.IsValidIndex(BranchPointIndex - 1) && PrimitivePoints.IsValidIndex(BranchPointIndex + 1))
 		[[likely]]
 		{
-			const FVector3f PreviousTangent = (PointPosition - PointPositions[PrimitivePoints[BranchPointIndex - 1]]).GetSafeNormal();
-			const FVector3f NextTangent = (PointPositions[PrimitivePoints[BranchPointIndex + 1]] - PointPosition).GetSafeNormal();
+			const FVector3f PreviousTangent = (PointPosition - PointPositionAttribute[PrimitivePoints[BranchPointIndex - 1]]).GetSafeNormal();
+			const FVector3f NextTangent = (PointPositionAttribute[PrimitivePoints[BranchPointIndex + 1]] - PointPosition).GetSafeNormal();
 			TangentVector = FVector3f::SlerpNormals(PreviousTangent, NextTangent, 0.5f);
 		}
 		else if (BranchPointIndex == PrimitivePoints.Num() - 1)
 		[[unlikely]]
 		{
 			const int32 PreviousPointIndex = PrimitivePoints[BranchPointIndex - 1];
-			const FVector3f& PreviousPointPosition = PointPositions[PreviousPointIndex];
+			const FVector3f& PreviousPointPosition = PointPositionAttribute[PreviousPointIndex];
 			TangentVector = (PointPosition - PreviousPointPosition).GetSafeNormal();
 		}
 		else if (BranchPointIndex == 0)
 		[[unlikely]]
 		{
 			const int32 NextPointIndex = PrimitivePoints[BranchPointIndex + 1];
-			const FVector3f& NextPointPosition = PointPositions[NextPointIndex];
+			const FVector3f& NextPointPosition = PointPositionAttribute[NextPointIndex];
 			TangentVector = (NextPointPosition - PointPosition).GetSafeNormal();
 		}
 		
@@ -986,7 +991,7 @@ void FPVMeshBuilder::GenerateBranchMeshData(const bool bPrimitiveIsTrunk, const
 		{
 			for (int32 SearchIndex = BranchPointIndex + 1; SearchIndex < PrimitivePoints.Num(); ++SearchIndex)
 			{
-				TangentVector = (PointPositions[PrimitivePoints[SearchIndex]] - PointPosition).GetSafeNormal();
+				TangentVector = (PointPositionAttribute[PrimitivePoints[SearchIndex]] - PointPosition).GetSafeNormal();
 				if (!TangentVector.IsNearlyZero())
 				{
 					break;
@@ -994,7 +999,9 @@ void FPVMeshBuilder::GenerateBranchMeshData(const bool bPrimitiveIsTrunk, const
 			}
 			if (TangentVector.IsNearlyZero())
 			{
-				TangentVector = PointBudDirections[PointIndex][0].GetSafeNormal(); // apical bud direction
+				const TArray<FVector3f> FallbackBudDirData = PV::AttributesHelper::GetBudDirection(
+					BudDirectionAttribute, BranchPointsAttribute, BranchParentNumberAttribute, PointPositionAttribute, BranchIndex, BranchPointIndex);
+				TangentVector = PV::FBudDirectionConstView(FallbackBudDirData).Apical.GetSafeNormal();
 			}
 			if (TangentVector.IsNearlyZero())
 			{
@@ -1023,10 +1030,11 @@ void FPVMeshBuilder::GenerateBranchMeshData(const bool bPrimitiveIsTrunk, const
 			FVector(TangentVector).FindBestAxisVectors(Axis1, Axis2);
 			UpVector = FVector3f(Axis1);
 		}
-
+		
 		PVE_DEBUG_LOCK()
 			PVE_PARAM_DEBUG_DIRECTION_PARAM("UpVector", PointPosition, UpVector);
 			PVE_PARAM_DEBUG_VECTOR_PARAM("TangentVector", PointPosition, PointPosition + TangentVector);
+			PVE_PARAM_DEBUG_TEXT_PARAM("PointScale", PointPosition, FText::FromString(FString::Printf(TEXT("%f"), PointScale)));
 		PVE_DEBUG_UNLOCK()
 
 		PreviousUpVector = UpVector;
