Add null-check to UPVSubgraphSettings::GetSubgraphInterface() to fix crash during live code

#Synced-CL 52672995

diff --git a/Engine/Plugins/Experimental/ProceduralVegetationEditor/Source/ProceduralVegetation/Private/Nodes/PVSubgraphSettings.cpp b/Engine/Plugins/Experimental/ProceduralVegetationEditor/Source/ProceduralVegetation/Private/Nodes/PVSubgraphSettings.cpp
--- a/Engine/Plugins/Experimental/ProceduralVegetationEditor/Source/ProceduralVegetation/Private/Nodes/PVSubgraphSettings.cpp
+++ b/Engine/Plugins/Experimental/ProceduralVegetationEditor/Source/ProceduralVegetation/Private/Nodes/PVSubgraphSettings.cpp
@@ -10,7 +10,7 @@
 
 UPCGGraphInterface* UPVSubgraphSettings::GetSubgraphInterface() const
 {
-	return SubgraphInstance->GetGraph();
+	return SubgraphInstance ? SubgraphInstance->GetGraph() : nullptr;
 }
 
 void FPVSubgraphElement::PostExecuteInternal(FPCGContext* InContext) const
