[PVE] noise implementation for MeshBuilder node

#Synced-CL 52778023

diff --git a/Engine/Plugins/Experimental/ProceduralVegetationEditor/Source/ProceduralVegetation/Private/Implementations/PVMeshBuilder.cpp b/Engine/Plugins/Experimental/ProceduralVegetationEditor/Source/ProceduralVegetation/Private/Implementations/PVMeshBuilder.cpp
--- a/Engine/Plugins/Experimental/ProceduralVegetationEditor/Source/ProceduralVegetation/Private/Implementations/PVMeshBuilder.cpp
+++ b/Engine/Plugins/Experimental/ProceduralVegetationEditor/Source/ProceduralVegetation/Private/Implementations/PVMeshBuilder.cpp
@@ -14,6 +14,7 @@
 #include "Facades/PVPointFacade.h"
 #include "Facades/PVTreeFacade.h"
 #include "Helpers/PVAttributesHelper.h"
+#include "Helpers/PVPlantTraversalHelper.h"
 
 #include "Helpers/PVUtilities.h"
 
@@ -69,7 +70,9 @@ void FPVMeshBuilder::GenerateGeometryCollection(const FManagedArrayCollection& I
 	BranchFacade.RecomputeBranchChildren();
 	
 	AddPointsToSkeleton(OutGeometryCollection);
-	
+	ApplyNoiseToSkeleton(OutGeometryCollection, MeshBuilderParams);
+
+	PV::AttributesHelper::ComputeLengthFromRoot(OutGeometryCollection);
 	PV::AttributesHelper::ComputeNjordPixelIndex(OutGeometryCollection);
 	PV::AttributesHelper::ComputePointPlantGradient(OutGeometryCollection);
 	PV::AttributesHelper::ComputePointGroundGradient(OutGeometryCollection);
@@ -323,6 +326,291 @@ void FPVMeshBuilder::AddPointsToSkeleton(FManagedArrayCollection& OutCollection)
 	}
 }
 
+void FPVMeshBuilder::ApplyNoiseToSkeleton(FManagedArrayCollection& OutCollection, const FPVMeshBuilderParams& MeshBuilderParams)
+{
+	if (!MeshBuilderParams.bApplySkeletonNoise || MeshBuilderParams.SkeletonNoiseEntries.IsEmpty())
+	{
+		return;
+	}
+
+	const PV::FBranchPointsAttributeConstView BranchPointsAttr = PV::FBranchPointsAttribute::FindAttribute(OutCollection);
+	const PV::FBranchParentsAttributeConstView BranchParentsAttr = PV::FBranchParentsAttribute::FindAttribute(OutCollection);
+	const PV::FBranchParentNumberAttributeView BranchParentNumberAttr = PV::FBranchParentNumberAttribute::FindAttribute(OutCollection);
+	const PV::FBranchNumberAttributeConstView BranchNumbersAttr = PV::FBranchNumberAttribute::FindAttribute(OutCollection);
+	const PV::FBranchHierarchyNumberAttributeConstView BranchHierarchyNumbersAttr = PV::FBranchHierarchyNumberAttribute::FindAttribute(OutCollection);
+	PV::FPointPositionAttributeView PointPositionsAttr = PV::FPointPositionAttribute::FindAttribute(OutCollection);
+
+	if (!BranchPointsAttr.IsValid() || !BranchParentsAttr.IsValid() || !BranchNumbersAttr.IsValid()
+		|| !BranchHierarchyNumbersAttr.IsValid() || !PointPositionsAttr.IsValid() || !BranchParentNumberAttr.IsValid())
+	{
+		return;
+	}
+
+	const int32 NumBranches = BranchPointsAttr.Num();
+	const int32 TotalPoints = PointPositionsAttr.Num();
+	if (NumBranches == 0 || TotalPoints == 0)
+	{
+		return;
+	}
+
+	// Sort entries once by generation ascending so the per-branch lookup can use a
+	// simple linear scan with early termination and always sees fallback candidates
+	// in order from lowest to highest generation.
+	TArray<FPVSkeletonNoiseEntry> SortedEntries = MeshBuilderParams.SkeletonNoiseEntries;
+	SortedEntries.Sort([](const FPVSkeletonNoiseEntry& A, const FPVSkeletonNoiseEntry& B)
+	{
+		return A.Generation < B.Generation;
+	});
+
+	// Returns the entry that governs a branch of the given hierarchy generation, or
+	// nullptr when no entry covers it (no noise should be applied to that branch).
+	//
+	// Priority rules:
+	//   1. An entry whose Generation == BranchGen is an exact match and wins outright.
+	//   2. Otherwise, the highest-generation entry whose Generation < BranchGen AND
+	//      bImpactRemainingGenerations == true acts as the fallback ("cascade down").
+	//   A higher-generation exact match always beats a lower-generation cascade entry.
+	auto FindEntryForGeneration = [&SortedEntries](const int32 BranchGen) -> const FPVSkeletonNoiseEntry*
+	{
+		const FPVSkeletonNoiseEntry* BestFallback = nullptr;
+		for (const FPVSkeletonNoiseEntry& Entry : SortedEntries)
+		{
+			if (Entry.Generation > BranchGen)
+			{
+				break; // Entries are sorted; nothing further can match.
+			}
+			if (Entry.Generation == BranchGen)
+			{
+				return &Entry; // Exact match — highest possible priority.
+			}
+			// Entry.Generation < BranchGen
+			if (Entry.bImpactRemainingGenerations)
+			{
+				BestFallback = &Entry; // Overwrite each time; last one is highest gen.
+			}
+		}
+		return BestFallback;
+	};
+
+	// ---- Pass 1: Compute the lateral noise displacement for every point ----
+	//
+	// Every point — including branch-first attachment points — gets a displacement
+	// derived from its own world position so the field is spatially coherent.
+	// Points whose branch has no applicable noise entry keep a zero displacement.
+	// Storing everything in a flat array (indexed by point index) lets pass 2
+	// replace attachment displacements in O(1).
+	TArray<FVector3f> Displacements;
+	Displacements.SetNumZeroed(TotalPoints);
+
+	// Pre-compute the set of point indices that are fused attachment points: a point is
+	// fused when it is BranchPoints[0] of a non-trunk branch AND the same index also
+	// appears in the parent branch's BranchPoints array (i.e. both branches share one
+	// physical point, not two coincident-but-distinct points as in the split topology).
+	//
+	// A fused point must be owned exclusively by the parent branch during noise
+	// processing:
+	//   Pass 1 — only the parent writes the displacement; the child skips index 0 to
+	//             avoid overwriting the parent's value with a different-generation noise.
+	//   Pass 3 — only the parent applies the displacement; the child skips index 0 to
+	//             avoid displacing the shared point a second time.
+	// In the split topology BranchPoints[0] of a child is a distinct index, so Pass 2's
+	// nearest-parent sync is sufficient and no special skipping is needed.
+	TSet<int32> FusedAttachmentPoints;
+	FusedAttachmentPoints.Reserve(NumBranches);
+	
+	for (int32 BranchIndex = 0; BranchIndex < NumBranches; ++BranchIndex)
+	{
+		const int32 ParentBranchIndex = PV::PlantTraversalHelper::GetBranchParentIndex(BranchParentNumberAttr, BranchNumbersAttr, BranchIndex);
+		if (ParentBranchIndex == INDEX_NONE)
+		{
+			continue;
+		}
+		const TArray<int32>& ChildBranchPoints = BranchPointsAttr[BranchIndex];
+		if (ChildBranchPoints.IsEmpty())
+		{
+			continue;
+		}
+		if (BranchPointsAttr[ParentBranchIndex].Contains(ChildBranchPoints[0]))
+		{
+			FusedAttachmentPoints.Add(ChildBranchPoints[0]);
+		}
+	}
+
+	for (int32 BranchIndex = 0; BranchIndex < NumBranches; ++BranchIndex)
+	{
+		// Resolve which noise entry governs this branch's hierarchy generation.
+		// Branches with no applicable entry contribute zero displacement.
+		const int32 BranchGen = BranchHierarchyNumbersAttr[BranchIndex];
+		const FPVSkeletonNoiseEntry* Entry = FindEntryForGeneration(BranchGen);
+		if (!Entry)
+		{
+			continue;
+		}
+
+		const TArray<int32>& BranchPoints = BranchPointsAttr[BranchIndex];
+		const int32 NumPoints = BranchPoints.Num();
+		if (NumPoints < 2)
+		{
+			continue;
+		}
+
+		const float Strength  = Entry->NoiseStrength;
+		const float Frequency = Entry->NoiseFrequency;
+		const float SeedScale = static_cast<float>(Entry->NoiseSeed);
+		const FVector SeedOffset(SeedScale * 127.1f, SeedScale * 311.7f, SeedScale * 74.7f);
+
+		// Snapshot original positions so the tangent used for lateral projection
+		// always comes from the unperturbed skeleton — this prevents compounding
+		// errors on densely interpolated branches.
+		TArray<FVector3f> OrigPos;
+		OrigPos.Reserve(NumPoints);
+		for (int32 i = 0; i < NumPoints; ++i)
+		{
+			OrigPos.Add(PointPositionsAttr[BranchPoints[i]]);
+		}
+
+		// For non-trunk branches, BranchPoints[0] is the attachment point.
+		// • Fused topology: that index is shared with the parent branch; writing
+		//   child-generation noise here would overwrite the correct parent value.
+		// • Split topology: the displacement will be replaced by the parent's in Pass 2.
+		// Starting from index 1 is safe and correct for both topologies; OrigPos[0] is
+		// still captured above so that the tangent at i=1 can use central-differencing.
+		const bool bHasParentBranch = (PV::PlantTraversalHelper::GetBranchParentIndex(BranchParentNumberAttr, BranchNumbersAttr, BranchIndex) != INDEX_NONE);
+		const int32 NoiseStartIndex = bHasParentBranch ? 1 : 0;
+
+		for (int32 i = NoiseStartIndex; i < NumPoints; ++i)
+		{
+			// Sample three independent Perlin channels at domain-offset coordinates
+			// for a coherent, non-axis-aligned 3D displacement field.
+			const FVector SamplePos = FVector(OrigPos[i]) * Frequency + SeedOffset;
+			FVector3f NoiseVec(
+				static_cast<float>(FMath::PerlinNoise3D(SamplePos)),
+				static_cast<float>(FMath::PerlinNoise3D(FVector(SamplePos.Y + 31.41, SamplePos.Z + 159.27, SamplePos.X + 26.53))),
+				static_cast<float>(FMath::PerlinNoise3D(FVector(SamplePos.Z + 271.83, SamplePos.X + 182.84, SamplePos.Y + 14.14)))
+			);
+
+			// Central-difference tangent (forward/backward at the ends).
+			FVector3f Tangent;
+			if (i > 0 && i < NumPoints - 1)
+			{
+				Tangent = (OrigPos[i + 1] - OrigPos[i - 1]).GetSafeNormal();
+			}
+			else if (i == NumPoints - 1)
+			{
+				Tangent = (OrigPos[i] - OrigPos[i - 1]).GetSafeNormal();
+			}
+			else
+			{
+				Tangent = (OrigPos[i + 1] - OrigPos[i]).GetSafeNormal();
+			}
+
+			// Project out the along-tangent component so displacement is purely
+			// lateral — this prevents the tangent direction from flipping between
+			// consecutive points, which is the primary cause of vertex-ring twisting.
+			if (!Tangent.IsNearlyZero())
+			{
+				NoiseVec -= Tangent * FVector3f::DotProduct(NoiseVec, Tangent);
+			}
+
+			Displacements[BranchPoints[i]] = NoiseVec * Strength;
+		}
+	}
+
+	// ---- Pass 2: Synchronise child attachment displacements with the parent ----
+	//
+	// The first point of a child branch is spatially coincident with a point on the
+	// parent branch.  If their displacements differ, the child tube visually detaches
+	// from the parent surface.  The fix: replace the child's first-point displacement
+	// with the displacement of the nearest parent-branch point.
+	//
+	// Branches are visited in hierarchy order (root → tips) so that when a chain of
+	// branches is present (grandparent → parent → child) the parent's first-point
+	// displacement is already the corrected value by the time the child is visited.
+	TArray<int32> HierarchySortedBranches;
+	HierarchySortedBranches.Reserve(NumBranches);
+	for (int32 Index = 0; Index < NumBranches; ++Index)
+	{
+		HierarchySortedBranches.Add(Index);
+	}
+	HierarchySortedBranches.Sort([&BranchParentsAttr](int32 Index1, int32 Index2)
+	{
+		const int32 Index1Hierarchy = BranchParentsAttr[Index1].Num();
+		const int32 Index2Hierarchy = BranchParentsAttr[Index2].Num();
+		if (Index1Hierarchy == Index2Hierarchy)
+		{
+			return Index1 < Index2;
+		}
+		return Index1Hierarchy < Index2Hierarchy;
+	});
+
+	for (const int32 BranchIndex : HierarchySortedBranches)
+	{
+		const int32 ParentBranchIndex = PV::PlantTraversalHelper::GetBranchParentIndex(BranchParentNumberAttr, BranchNumbersAttr, BranchIndex);
+		if (ParentBranchIndex == INDEX_NONE)
+		{
+			continue; // Trunk / root — no parent to sync with.
+		}
+
+		const TArray<int32>& ChildPoints = BranchPointsAttr[BranchIndex];
+		if (ChildPoints.Num() == 0)
+		{
+			continue;
+		}
+
+		const TArray<int32>& ParentPoints = BranchPointsAttr[ParentBranchIndex];
+		if (ParentPoints.Num() == 0)
+		{
+			continue;
+		}
+
+		const int32 ChildFirstPoint = ChildPoints[0];
+		const FVector3f ChildFirstPos = PointPositionsAttr[ChildFirstPoint];
+
+		// Find the parent point whose original position is closest to the child
+		// attachment so we match the displacement of the correct parent point.
+		const int32* ClosestParentPoint = Algo::MinElement(ParentPoints,
+			[&ChildFirstPos, &PointPositionsAttr](const int32 A, const int32 B)
+			{
+				return FVector3f::DistSquared(PointPositionsAttr[A], ChildFirstPos) <
+				       FVector3f::DistSquared(PointPositionsAttr[B], ChildFirstPos);
+			});
+
+		if (ClosestParentPoint)
+		{
+			// Overwrite the child first-point displacement with the parent's so
+			// both the parent tube and the child's root ring move identically.
+			Displacements[ChildFirstPoint] = Displacements[*ClosestParentPoint];
+		}
+	}
+
+	// ---- Pass 3: Apply displacements ----
+	for (int32 BranchIndex = 0; BranchIndex < NumBranches; ++BranchIndex)
+	{
+		const TArray<int32>& BranchPoints = BranchPointsAttr[BranchIndex];
+		const int32 NumPoints = BranchPoints.Num();
+		if (NumPoints < 2)
+		{
+			continue;
+		}
+
+		// For fused-topology branches BranchPoints[0] is the same point index as the
+		// parent's attachment point.  The parent branch's loop above has already applied
+		// the displacement to it; applying it here a second time would double-displace
+		// that shared point while all neighbouring points move only once.
+		// For split-topology branches BranchPoints[0] is a unique index that is not
+		// present in FusedAttachmentPoints, so it must be displaced here as normal.
+		const bool bFirstPointFused = (PV::PlantTraversalHelper::GetBranchParentIndex(BranchParentNumberAttr, BranchNumbersAttr, BranchIndex) != INDEX_NONE)
+		                           && FusedAttachmentPoints.Contains(BranchPoints[0]);
+		const int32 ApplyStartIndex = bFirstPointFused ? 1 : 0;
+
+		for (int32 i = ApplyStartIndex; i < NumPoints; ++i)
+		{
+			const int32 PointIndex = BranchPoints[i];
+			PointPositionsAttr[PointIndex] = PointPositionsAttr[PointIndex] + Displacements[PointIndex];
+		}
+	}
+}
+
 void FPVMeshBuilder::RemoveNonBudPoints(PV::Facades::FBranchFacade& BranchFacade, PV::Facades::FPointFacade& PointFacade)
 {
 	// Note that we only remove points and do not care about removing branches or 
diff --git a/Engine/Plugins/Experimental/ProceduralVegetationEditor/Source/ProceduralVegetation/Private/Implementations/PVMeshBuilder.h b/Engine/Plugins/Experimental/ProceduralVegetationEditor/Source/ProceduralVegetation/Private/Implementations/PVMeshBuilder.h
--- a/Engine/Plugins/Experimental/ProceduralVegetationEditor/Source/ProceduralVegetation/Private/Implementations/PVMeshBuilder.h
+++ b/Engine/Plugins/Experimental/ProceduralVegetationEditor/Source/ProceduralVegetation/Private/Implementations/PVMeshBuilder.h
@@ -26,6 +26,38 @@
 
 #include "PVMeshBuilder.generated.h"
 
+/** A single per-generation noise layer used by FPVMeshBuilderParams::SkeletonNoiseEntries. */
+USTRUCT()
+struct FPVSkeletonNoiseEntry
+{
+	GENERATED_BODY()
+
+	UPROPERTY(EditAnywhere, Category="Skeleton Noise", DisplayName="Generation",
+		meta=(ClampMin=1, UIMin=1, Tooltip=
+			"The branch-hierarchy generation this entry targets.\n\nGeneration 1 is the trunk. Generation 2 is the first tier of lateral branches, and so on. When Impact Remaining Generations is disabled the noise is applied only to branches whose hierarchy number equals this value."))
+	int32 Generation = 1;
+
+	UPROPERTY(EditAnywhere, Category="Skeleton Noise", DisplayName="Impact Remaining Generations",
+		meta=(Tooltip=
+			"When enabled, this entry also governs all higher-numbered generations that have no entry of their own.\n\nA more specific entry (exact generation match, or a higher-generation fallback entry) always takes precedence over this one. Disable to restrict the effect to the exact generation number above."))
+	bool bImpactRemainingGenerations = true;
+
+	UPROPERTY(EditAnywhere, Category="Skeleton Noise", DisplayName="Noise Strength",
+		meta=(ClampMin=0.0f, ClampMax=200.0f, UIMin=0.0f, UIMax=50.0f, Tooltip=
+			"Amplitude of the skeleton noise displacement in centimetres for this generation."))
+	float NoiseStrength = 5.0f;
+
+	UPROPERTY(EditAnywhere, Category="Skeleton Noise", DisplayName="Noise Frequency",
+		meta=(ClampMin=0.0001f, ClampMax=1.0f, UIMin=0.001f, UIMax=0.2f, Tooltip=
+			"Spatial frequency of the noise field for this generation.\n\nHigher values produce tighter, more frequent bends; lower values produce broad, sweeping curves."))
+	float NoiseFrequency = 0.05f;
+
+	UPROPERTY(EditAnywhere, Category="Skeleton Noise", DisplayName="Noise Seed",
+		meta=(Tooltip=
+			"Offsets the Perlin noise sampling position for this generation. Each unique value produces a distinct deformation pattern."))
+	int32 NoiseSeed = 0;
+};
+
 USTRUCT()
 struct FPVMeshBuilderParams
 {
@@ -185,6 +217,16 @@ struct FPVMeshBuilderParams
 			"Controls which generations the displacement is applied on \n\nControls which generations the displacement is applied on."))
 	int32 DisplacementGenerationUpperLimit = 1;
 
+	UPROPERTY(EditAnywhere, Category="Skeleton Noise",
+		meta=(PCG_NotOverridable, InlineEditConditionToggle, Tooltip=
+			"Enables Perlin noise displacement on skeleton points after interpolation.\n\nWhen enabled, each branch point receives a coherent lateral offset that adds organic variation to the tree silhouette without breaking branch connections."))
+	bool bApplySkeletonNoise = false;
+
+	UPROPERTY(EditAnywhere, Category="Skeleton Noise",
+		meta=(EditCondition="bApplySkeletonNoise", Tooltip=
+			"Per-generation noise layers.\n\nEach entry targets a specific branch-hierarchy generation. When Impact Remaining Generations is enabled on an entry, that entry also governs all higher generations that have no dedicated entry. A higher-generation entry always overrides a lower-generation fallback."))
+	TArray<FPVSkeletonNoiseEntry> SkeletonNoiseEntries;
+
 	FPVMeshBuilderParams()
 	{
 		InitializeLinearCurve(HullRetentionGradient);
@@ -302,7 +344,9 @@ struct FPVMeshBuilder
 	                                       FGeometryCollection& OutGeometryCollection);
 
 	static void AddPointsToSkeleton(FManagedArrayCollection& OutCollection);
-	
+
+	static void ApplyNoiseToSkeleton(FManagedArrayCollection& OutCollection, const FPVMeshBuilderParams& MeshBuilderParams);
+
 	static void RemoveNonBudPoints(PV::Facades::FBranchFacade& BranchFacade, PV::Facades::FPointFacade& PointFacade);
 
 	static void GenerateDynamicMesh(FManagedArrayCollection& Collection, const FManagedArrayCollection& InPlantProfileCollection, const FPVMeshBuilderParams& MeshBuilderParams,
