Fix visualization collections breaking in CL 52468466

#Synced-CL 52785170

#Synced-CL 52791601

diff --git a/Engine/Plugins/Experimental/ProceduralVegetationEditor/Source/ProceduralVegetationEditor/Private/Visualizations/DebugVisualization/PVDebugVisualizer.cpp b/Engine/Plugins/Experimental/ProceduralVegetationEditor/Source/ProceduralVegetationEditor/Private/Visualizations/DebugVisualization/PVDebugVisualizer.cpp
--- a/Engine/Plugins/Experimental/ProceduralVegetationEditor/Source/ProceduralVegetationEditor/Private/Visualizations/DebugVisualization/PVDebugVisualizer.cpp
+++ b/Engine/Plugins/Experimental/ProceduralVegetationEditor/Source/ProceduralVegetationEditor/Private/Visualizations/DebugVisualization/PVDebugVisualizer.cpp
@@ -25,48 +25,34 @@ FVisualizerDrawContext::FVisualizerDrawContext(
 	, SceneSetupParams(InSceneParams)
 {}
 
-void FPVDebugVisualizer::Draw(const UPVData* InData, FPCGSceneSetupParams& InOutParams, bool bDrawVisualizationCollections)
+void FPVDebugVisualizer::DrawDebugVisualizations(
+	const TArray<FPVVisualizationSettings>& VisualizationSettings,
+	const FManagedArrayCollection& Collection,
+	FPCGSceneSetupParams& InOutParams
+)
 {
-	check(InData);
-
-	const TArray<FPVVisualizationSettings>& VisualizationSettings = InData->GetDebugSettings().VisualizationSettings;
 	for (const FPVVisualizationSettings& VisualizationSetting : VisualizationSettings)
 	{
-		if (bDrawVisualizationCollections)
+		FVisualizerDrawContext Context(Collection, VisualizationSetting, InOutParams);
+		if (const FPVDebugVisualizationPtr VisualizationPtr = CreateVisualizer(VisualizationSetting))
 		{
-			for (const FPVVisualizationCollection& VisualizationCollection : InData->GetVisualizationCollections())
-			{
-				FVisualizerDrawContext Context(VisualizationCollection.Collection, VisualizationSetting, InOutParams);
-				if (const FPVDebugVisualizationPtr VisualizationPtr = CreateVisualizer(VisualizationSetting))
-				{
-					VisualizationPtr->Draw(Context);
-				}
-			}
-		}
-		else
-		{
-			FVisualizerDrawContext Context(InData->GetCollection(), VisualizationSetting, InOutParams);
-			if (const FPVDebugVisualizationPtr VisualizationPtr = CreateVisualizer(VisualizationSetting))
-			{
-				VisualizationPtr->Draw(Context);
-			}
+			VisualizationPtr->Draw(Context);
 		}
 	}
 }
 
-void FPVDebugVisualizer::DrawDebugParams(const UPVData* InData, FPCGSceneSetupParams& InOutParams)
+void FPVDebugVisualizer::DrawDebugParams(const TArray<FPVParamDebuggerSettings>& ParamDebugVisualizationSettings, bool bAutoFocusLoopDebug, FPCGSceneSetupParams& InOutParams)
 {
-	check(InData);
-	
-	TArray<TPair<FColor, FText>> LegendItems;
-	const TArray<FPVParamDebuggerSettings>& ParamDebugVisualizationSettings = InData->GetDebugSettings().ParamDebugVisualizationSettings;
-
 	TArray<FString> UniqueNames;
+	UniqueNames.Reserve(ParamDebugVisualizationSettings.Num());
 	for (const FPVParamDebuggerSettings& DebugVisualizationSetting : ParamDebugVisualizationSettings)
 	{
 		UniqueNames.AddUnique(DebugVisualizationSetting.ParamName.ToString());
 	}
 
+	TArray<TPair<FColor, FText>> LegendItems;
+	LegendItems.Reserve(UniqueNames.Num());
+
 	float Alpha = 0.0f;
 	const float AlphaStep = 1.0f / static_cast<float>(UniqueNames.Num());
 	TArray<FLinearColor> UniqueColors;
@@ -103,7 +89,7 @@ void FPVDebugVisualizer::DrawDebugParams(const UPVData* InData, FPCGSceneSetupPa
 		return;
 	}
 
-	if (InData->GetDebugSettings().bAutoFocusLoopDebug)
+	if (bAutoFocusLoopDebug)
 	{
 		if (const TSharedPtr<FEditorViewportClient> ViewportClient = ViewportWidget->GetViewportClient())
 		{
diff --git a/Engine/Plugins/Experimental/ProceduralVegetationEditor/Source/ProceduralVegetationEditor/Private/Visualizations/DebugVisualization/PVDebugVisualizer.h b/Engine/Plugins/Experimental/ProceduralVegetationEditor/Source/ProceduralVegetationEditor/Private/Visualizations/DebugVisualization/PVDebugVisualizer.h
--- a/Engine/Plugins/Experimental/ProceduralVegetationEditor/Source/ProceduralVegetationEditor/Private/Visualizations/DebugVisualization/PVDebugVisualizer.h
+++ b/Engine/Plugins/Experimental/ProceduralVegetationEditor/Source/ProceduralVegetationEditor/Private/Visualizations/DebugVisualization/PVDebugVisualizer.h
@@ -11,7 +11,6 @@ struct FManagedArrayCollection;
 struct FPVVisualizationSettings;
 class FPVDebugVisualizationBase;
 struct FPCGSceneSetupParams;
-class UPVData;
 
 struct FVisualizerDrawContext
 {
@@ -24,8 +23,8 @@ struct FVisualizerDrawContext
 class FPVDebugVisualizer
 {
 public:
-	static void Draw(const UPVData* InData, FPCGSceneSetupParams& InOutParams, bool bDrawVisualizationCollections = false);
-	static void DrawDebugParams(const UPVData* InData, FPCGSceneSetupParams& InOutParams);
+	static void DrawDebugVisualizations(const TArray<FPVVisualizationSettings>& VisualizationSettings, const FManagedArrayCollection& Collection, FPCGSceneSetupParams& InOutParams);
+	static void DrawDebugParams(const TArray<FPVParamDebuggerSettings>& ParamDebugVisualizationSettings, bool bAutoFocusLoopDebug, FPCGSceneSetupParams& InOutParams);
 
 private:
 	static void VisualizeDebuggerParam(
diff --git a/Engine/Plugins/Experimental/ProceduralVegetationEditor/Source/ProceduralVegetationEditor/Private/Visualizations/PVDataVisualization.cpp b/Engine/Plugins/Experimental/ProceduralVegetationEditor/Source/ProceduralVegetationEditor/Private/Visualizations/PVDataVisualization.cpp
--- a/Engine/Plugins/Experimental/ProceduralVegetationEditor/Source/ProceduralVegetationEditor/Private/Visualizations/PVDataVisualization.cpp
+++ b/Engine/Plugins/Experimental/ProceduralVegetationEditor/Source/ProceduralVegetationEditor/Private/Visualizations/PVDataVisualization.cpp
@@ -162,11 +162,14 @@ FPCGSetupSceneFunc FPVDataVisualization::GetViewportSetupFunc(const UPCGSettings
 
 		if (PV::Utilities::DebugModeEnabled() && RenderSettings->IsDebug())
 		{
-			FPVDebugVisualizer::Draw(WeakData.Get(), InOutParams);
-			FPVDebugVisualizer::DrawDebugParams(WeakData.Get(), InOutParams);
+			FPVDebugVisualizer::DrawDebugVisualizations(WeakData->GetDebugSettings().VisualizationSettings, WeakData->GetCollection(), InOutParams);
+			FPVDebugVisualizer::DrawDebugParams(WeakData->GetDebugSettings().ParamDebugVisualizationSettings, WeakData->GetDebugSettings().bAutoFocusLoopDebug, InOutParams);
+		}
+
+		for (const FPVVisualizationCollection& VisualizationCollection : WeakData->GetVisualizationCollections())
+		{
+			FPVDebugVisualizer::DrawDebugVisualizations(VisualizationCollection.VisualizationSettings, VisualizationCollection.Collection, InOutParams);
 		}
-		
-		FPVDebugVisualizer::Draw(WeakData.Get(), InOutParams, true);
  
 		InOutParams.FocusBounds = FBoxSphereBounds(Bounds.Origin, Bounds.BoxExtent, Bounds.SphereRadius);
 	};
