Fix performance issue in ConnectAndTraceBranches & RecomputeAllGrowthDataAttributes

- Fix bug which caused edge split to fail: Fix a bug in endpoint edge attachment that caused edge splits to fail by excluding endpoint vertices from candidate edges
- Speed up end point attachment edge search: Add scale-aware triangle search overloads to the dynamic mesh spatial query helpers for more accurate endpoint attachment radius queries
- Speed up tree traversal functions in PVTreeTraversalHelper.h: Fix performance issues in procedural vegetation branch traversal by introducing a pre-built FBranchNumberToIndexTable lookup structure that replaces O(N) linear searches with O(1) lookups. Also add CPU profiler scopes throughout attribute computation functions

#Synced-CL 52793511

#Synced-CL 52791601

diff --git a/Engine/Plugins/Experimental/ProceduralVegetationEditor/Source/ProceduralVegetation/Private/Helpers/PVAttributesHelper.cpp b/Engine/Plugins/Experimental/ProceduralVegetationEditor/Source/ProceduralVegetation/Private/Helpers/PVAttributesHelper.cpp
--- a/Engine/Plugins/Experimental/ProceduralVegetationEditor/Source/ProceduralVegetation/Private/Helpers/PVAttributesHelper.cpp
+++ b/Engine/Plugins/Experimental/ProceduralVegetationEditor/Source/ProceduralVegetation/Private/Helpers/PVAttributesHelper.cpp
@@ -128,6 +128,8 @@ TArray<FVector3f> PV::AttributesHelper::GetBudDirection(
 
 void PV::AttributesHelper::ComputeBudNumbers(const FComputeBudNumbersAttributes& InAttributes)
 {
+	TRACE_CPUPROFILER_EVENT_SCOPE(ComputeBudNumbers);
+
 	for (int32 PointIndex = 0; PointIndex < InAttributes.PointBudNumber.Num(); ++PointIndex)
 	{
 		InAttributes.PointBudNumber[PointIndex] = PointIndex + 1;
@@ -147,6 +149,8 @@ void PV::AttributesHelper::ComputeBudNumbers(const FComputeBudNumbersAttributes&
 
 void PV::AttributesHelper::ComputeNjordPixelIndex(const FComputeNjordPixelIndexAttributes& InAttributes)
 {
+	TRACE_CPUPROFILER_EVENT_SCOPE(ComputeNjordPixelIndex);
+
 	if (!InAttributes.IsValid())
 	{
 		return;
@@ -171,6 +175,8 @@ void PV::AttributesHelper::ComputeNjordPixelIndex(const FComputeNjordPixelIndexA
 
 void PV::AttributesHelper::ComputeLengthFromRoot(const FComputeLengthFromRootAttributes& InAttributes)
 {
+	TRACE_CPUPROFILER_EVENT_SCOPE(ComputeLengthFromRoot);
+
 	using namespace PV::PlantTraversalHelper;
 
 	if (!InAttributes.IsValid())
@@ -178,7 +184,8 @@ void PV::AttributesHelper::ComputeLengthFromRoot(const FComputeLengthFromRootAtt
 		return;
 	}
 
-	RecursiveWalkBranches(InAttributes.BranchParentNumber, InAttributes.BranchChildren, InAttributes.BranchNumber, [&](int32 BranchIndex)->EForEachResult
+	const FBranchNumberToIndexTable BranchNumberToIndex(InAttributes.BranchNumber);
+	RecursiveWalkBranches(InAttributes.BranchParentNumber, InAttributes.BranchChildren, BranchNumberToIndex, [&](int32 BranchIndex)->EForEachResult
 	{
 		const TArray<int32>& CurrentBranchPoints = InAttributes.BranchPoints[BranchIndex];
 		if (!ensure(CurrentBranchPoints.Num() > 0))
@@ -187,7 +194,7 @@ void PV::AttributesHelper::ComputeLengthFromRoot(const FComputeLengthFromRootAtt
 		}
 
 		const int32 FirstPointIndex = CurrentBranchPoints[0];
-		const int32 ParentBranchIndex = GetBranchParentIndex(InAttributes.BranchParentNumber, InAttributes.BranchNumber, BranchIndex);
+		const int32 ParentBranchIndex = GetBranchParentIndex(InAttributes.BranchParentNumber, BranchNumberToIndex, BranchIndex);
 
 		if (ParentBranchIndex == INDEX_NONE)
 		{
@@ -265,6 +272,8 @@ void PV::AttributesHelper::ComputeLengthFromRoot(const FComputeLengthFromRootAtt
 
 void PV::AttributesHelper::ComputeBudDevelopment(const FComputeBudDevelopmentAttributes& InAttributes)
 {
+	TRACE_CPUPROFILER_EVENT_SCOPE(ComputeBudDevelopment);
+
 	using namespace PV::PlantTraversalHelper;
 
 	if (!InAttributes.IsValid())
@@ -295,13 +304,14 @@ void PV::AttributesHelper::ComputeBudDevelopment(const FComputeBudDevelopmentAtt
 		}
 	);
 
+	const FBranchNumberToIndexTable BranchNumberToIndex(InAttributes.BranchNumber);
 	ForEachUniquePointOnBranches(
 		InAttributes.BranchPoints,
 		InAttributes.BranchParentNumber,
 		InAttributes.BranchNumber,
 		[&](int32 BranchIndex, int32 BranchPointIndex, int32 PointIndex)
 		{
-			const bool bIsRootBranch = GetBranchParentIndex(InAttributes.BranchParentNumber, InAttributes.BranchNumber, BranchIndex) == INDEX_NONE;
+			const bool bIsRootBranch = GetBranchParentIndex(InAttributes.BranchParentNumber, BranchNumberToIndex, BranchIndex) == INDEX_NONE;
 			const int32 RootPointIndex = bIsRootBranch ? InAttributes.BranchPoints[BranchIndex][0] : InAttributes.BranchPoints[BranchIndex][1];
 			const int32 BranchAge = InAttributes.BudDevelopment[RootPointIndex].BudAge;
 			InAttributes.BudDevelopment[PointIndex].BranchAge = BranchAge;
@@ -313,6 +323,8 @@ void PV::AttributesHelper::ComputeBudDevelopment(const FComputeBudDevelopmentAtt
 
 void PV::AttributesHelper::ComputeBudDirections(const FComputeBudDirectionsAttributes& InAttributes)
 {
+	TRACE_CPUPROFILER_EVENT_SCOPE(ComputeBudDirections);
+
 	using namespace PV::PlantTraversalHelper;
 
 	ForEachUniquePointOnBranches(
@@ -400,6 +412,8 @@ void PV::AttributesHelper::ComputeBudDirections(const FComputeBudDirectionsAttri
 
 void PV::AttributesHelper::ComputeBudHormoneLevels(const FComputeBudHormoneLevelsAttributes& InAttributes)
 {
+	TRACE_CPUPROFILER_EVENT_SCOPE(ComputeBudHormoneLevels);
+
 	using namespace PV::PlantTraversalHelper;
 		
 	for (int32 PointIndex = 0; PointIndex < InAttributes.BudHormoneLevels.Num(); PointIndex++)
@@ -418,6 +432,8 @@ void PV::AttributesHelper::ComputeBudHormoneLevels(const FComputeBudHormoneLevel
 
 void PV::AttributesHelper::ComputePointGroundGradient(const FComputePointGroundGradientAttributes& InAttributes)
 {
+	TRACE_CPUPROFILER_EVENT_SCOPE(ComputePointGroundGradient);
+
 	if (!InAttributes.IsValid())
 	{
 		return;
@@ -438,6 +454,8 @@ void PV::AttributesHelper::ComputePointGroundGradient(const FComputePointGroundG
 
 void PV::AttributesHelper::ComputePointScaleGradient(const FComputePointScaleGradientAttributes& InAttributes)
 {
+	TRACE_CPUPROFILER_EVENT_SCOPE(ComputePointScaleGradient);
+
 	if (!InAttributes.IsValid())
 	{
 		return;
@@ -459,6 +477,8 @@ void PV::AttributesHelper::ComputePointScaleGradient(const FComputePointScaleGra
 
 void PV::AttributesHelper::ComputePointPlantGradient(const FComputePointPlantGradientAttributes& InAttributes)
 {
+	TRACE_CPUPROFILER_EVENT_SCOPE(ComputePointPlantGradient);
+
 	using namespace PV::PlantTraversalHelper;
 
 	if (!InAttributes.IsValid())
@@ -498,6 +518,8 @@ void PV::AttributesHelper::ComputePointPlantGradient(const FComputePointPlantGra
 
 void PV::AttributesHelper::ComputePointHullGradient(const FComputePointHullGradientAttributes& InAttributes)
 {
+	TRACE_CPUPROFILER_EVENT_SCOPE(ComputePointHullGradient);
+
 	using namespace PV::PlantTraversalHelper;
 
 	if (!InAttributes.IsValid())
@@ -572,6 +594,8 @@ void PV::AttributesHelper::ComputePointHullGradient(const FComputePointHullGradi
 
 void PV::AttributesHelper::ComputePointMainTrunkGradient(const FComputePointMainTrunkGradientAttributes& InAttributes)
 {
+	TRACE_CPUPROFILER_EVENT_SCOPE(ComputePointMainTrunkGradient);
+
 	using namespace PV::PlantTraversalHelper;
 
 	if (!InAttributes.IsValid())
@@ -672,6 +696,8 @@ void PV::AttributesHelper::ComputePointMainTrunkGradient(const FComputePointMain
 
 void PV::AttributesHelper::ComputeBudStatus(const FComputeBudStatusAttributes& InAttributes)
 {
+	TRACE_CPUPROFILER_EVENT_SCOPE(ComputeBudStatus);
+
 	using namespace PV::PlantTraversalHelper;
 
 	if (!InAttributes.IsValid())
@@ -679,6 +705,7 @@ void PV::AttributesHelper::ComputeBudStatus(const FComputeBudStatusAttributes& I
 		return;
 	}
 
+	const FBranchNumberToIndexTable BranchNumberToIndex(InAttributes.BranchNumber);
 	ForEachUniquePointOnBranches(
 		InAttributes.BranchPoints,
 		InAttributes.BranchParentNumber,
@@ -706,7 +733,7 @@ void PV::AttributesHelper::ComputeBudStatus(const FComputeBudStatusAttributes& I
 				BudStatusView.Inactive = 0;
 			}
 			else if (InAttributes.BranchPoints[BranchIndex][0] == PointIndex
-				&& GetBranchParentIndex(InAttributes.BranchParentNumber, InAttributes.BranchNumber, BranchIndex) == INDEX_NONE) // Is plant root point
+				&& GetBranchParentIndex(InAttributes.BranchParentNumber, BranchNumberToIndex, BranchIndex) == INDEX_NONE) // Is plant root point
 			{
 				BudStatusView.Axillary = 1;
 				BudStatusView.Seed = 1;
@@ -719,7 +746,7 @@ void PV::AttributesHelper::ComputeBudStatus(const FComputeBudStatusAttributes& I
 				const int32 NumChildren = GetBranchPointChildBranchIndices(
 					InAttributes.BranchChildren,
 					InAttributes.BranchParentNumber,
-					InAttributes.BranchNumber,
+					BranchNumberToIndex,
 					InAttributes.BranchPoints,
 					InAttributes.PointLengthFromRoot,
 					BranchIndex, 
@@ -740,6 +767,8 @@ void PV::AttributesHelper::ComputeBudStatus(const FComputeBudStatusAttributes& I
 
 void PV::AttributesHelper::ComputeBudLateralMeristem(const FComputeBudLateralMeristemAttributes& InAttributes)
 {
+	TRACE_CPUPROFILER_EVENT_SCOPE(ComputeBudLateralMeristem);
+
 	using namespace PV::PlantTraversalHelper;
 
 	for (int32 i = 0; i < InAttributes.BudLateralMeristem.Num(); ++i)
@@ -825,6 +854,8 @@ void PV::AttributesHelper::ComputeBudLightDetected(PV::FBudLightDetectedAttribut
 
 void PV::AttributesHelper::ComputeBranchHierarchyNumbers(const FComputeBranchHierarchyNumberAttributes& InAttributes)
 {
+	TRACE_CPUPROFILER_EVENT_SCOPE(ComputeBranchHierarchyNumbers);
+
 	using namespace PV::PlantTraversalHelper;
 
 	RecursiveWalkBranches<int32>(
@@ -843,6 +874,8 @@ void PV::AttributesHelper::ComputeBranchHierarchyNumbers(const FComputeBranchHie
 
 void PV::AttributesHelper::ComputeSmoothedPointScales(const FComputeSmoothedPointScalesAttributes& InAttributes)
 {
+	TRACE_CPUPROFILER_EVENT_SCOPE(ComputeSmoothedPointScales);
+
 	using namespace PV::PlantTraversalHelper;
 
 	const bool bHasFusedPoints = PV::PlantTraversalHelper::HasFusedPoints(InAttributes.BranchPoints, InAttributes.BranchParentNumber, InAttributes.BranchNumber);
@@ -944,6 +977,8 @@ void PV::AttributesHelper::ComputeSmoothedPointScales(const FComputeSmoothedPoin
 
 bool PV::AttributesHelper::SplitFusedPoints(FSplitFusedPointsAttributes InAttributes)
 {
+	TRACE_CPUPROFILER_EVENT_SCOPE(SplitFusedPoints);
+
 	if (!InAttributes.IsValid())
 	{
 		return false;
@@ -1017,6 +1052,8 @@ bool PV::AttributesHelper::SplitFusedPoints(FSplitFusedPointsAttributes InAttrib
 
 bool PV::AttributesHelper::RecomputeAllGrowthDataAttributes(const FRecomputeGrowthDataAttributes& InAttributes)
 {
+	TRACE_CPUPROFILER_EVENT_SCOPE(RecomputeAllGrowthDataAttributes);
+
 	if (!InAttributes.IsValid())
 	{
 		return false;
diff --git a/Engine/Plugins/Experimental/ProceduralVegetationEditor/Source/ProceduralVegetation/Private/Helpers/PVDynamicMeshHelpers.cpp b/Engine/Plugins/Experimental/ProceduralVegetationEditor/Source/ProceduralVegetation/Private/Helpers/PVDynamicMeshHelpers.cpp
--- a/Engine/Plugins/Experimental/ProceduralVegetationEditor/Source/ProceduralVegetation/Private/Helpers/PVDynamicMeshHelpers.cpp
+++ b/Engine/Plugins/Experimental/ProceduralVegetationEditor/Source/ProceduralVegetation/Private/Helpers/PVDynamicMeshHelpers.cpp
@@ -5,7 +5,7 @@
 #include "DynamicMesh/DynamicMesh3.h"
 #include "DynamicMesh/DynamicMeshAABBTree3.h"
 
-TArray<int32> PV::DynamicMeshHelper::GetUniqueVertexIndices(const UE::Geometry::FDynamicMesh3& DynamicMesh, const TArray<int32>& TriangleIds)
+PV::DynamicMeshHelper::FFindMeshElemsResult PV::DynamicMeshHelper::GetUniqueVertexIndices(const UE::Geometry::FDynamicMesh3& DynamicMesh, const TArray<int32>& TriangleIds)
 {
 	TSet<int32> UniqueVertexIndices;
 	UniqueVertexIndices.Reserve(TriangleIds.Num() * 3);
@@ -17,12 +17,20 @@ TArray<int32> PV::DynamicMeshHelper::GetUniqueVertexIndices(const UE::Geometry::
 		UniqueVertexIndices.Add(VertexIndices[2]);
 	}
 
-	return UniqueVertexIndices.Array();
+	PV::DynamicMeshHelper::FFindMeshElemsResult OutResult;
+	OutResult.Reserve(UniqueVertexIndices.Num());
+
+	for (int32 VID : UniqueVertexIndices)
+	{
+		OutResult.Add(VID);
+	}
+
+	return OutResult;
 }
 
-TArray<int32> PV::DynamicMeshHelper::GetVertexNeighbours(const UE::Geometry::FDynamicMesh3& DynamicMesh, int32 VertexID)
+PV::DynamicMeshHelper::FFindMeshElemsResult PV::DynamicMeshHelper::GetVertexNeighbours(const UE::Geometry::FDynamicMesh3& DynamicMesh, int32 VertexID)
 {
-	TArray<int32> OutNeighbours;
+	FFindMeshElemsResult OutNeighbours;
 	OutNeighbours.Reserve(DynamicMesh.GetVtxEdgeCount(VertexID));
 	for (int32 EdgeID : DynamicMesh.VtxEdgesItr(VertexID))
 	{
@@ -33,9 +41,9 @@ TArray<int32> PV::DynamicMeshHelper::GetVertexNeighbours(const UE::Geometry::FDy
 	return OutNeighbours;
 }
 
-TArray<int32> PV::DynamicMeshHelper::FindVerticesInRadius(const UE::Geometry::FDynamicMeshAABBTree3& DynamicMeshAABBTree, const FVector& Pos, float Radius)
+PV::DynamicMeshHelper::FFindMeshElemsResult PV::DynamicMeshHelper::FindVerticesInRadius(const UE::Geometry::FDynamicMeshAABBTree3& DynamicMeshAABBTree, const FVector& Pos, float Radius)
 {
-	TArray<int32> OutVertexIndices;
+	FFindMeshElemsResult OutVertexIndices;
 
 	const UE::Geometry::FDynamicMesh3* DynamicMesh = DynamicMeshAABBTree.GetMesh();
 	if (DynamicMesh == nullptr)
@@ -72,9 +80,9 @@ TArray<int32> PV::DynamicMeshHelper::FindVerticesInRadius(const UE::Geometry::FD
 	return OutVertexIndices;
 }
 
-TArray<int32> PV::DynamicMeshHelper::FindTrianglesInRadius(const UE::Geometry::FDynamicMeshAABBTree3& DynamicMeshAABBTree, const FVector& Pos, float Radius)
+PV::DynamicMeshHelper::FFindMeshElemsResult PV::DynamicMeshHelper::FindTrianglesInRadius(const UE::Geometry::FDynamicMeshAABBTree3& DynamicMeshAABBTree, const FVector& Pos, float Radius)
 {
-	TArray<int32> OutTriangleIDs;
+	FFindMeshElemsResult OutTriangleIDs;
 
 	const UE::Geometry::FDynamicMesh3* DynamicMesh = DynamicMeshAABBTree.GetMesh();
 	if (DynamicMesh == nullptr)
@@ -100,7 +108,82 @@ TArray<int32> PV::DynamicMeshHelper::FindTrianglesInRadius(const UE::Geometry::F
 
 		if (FVector::DistSquared(Pos, Query.ClosestTrianglePoint) < RadiusSqr)
 		{
-			OutTriangleIDs.AddUnique(TID);
+			OutTriangleIDs.Add(TID);
+		}
+	};
+
+	DynamicMeshAABBTree.DoTraversal(TreeTraversal);
+
+	return OutTriangleIDs;
+}
+
+PV::DynamicMeshHelper::FFindMeshElemsResult PV::DynamicMeshHelper::FindTrianglesInRadius(
+	const UE::Geometry::FDynamicMeshAABBTree3& DynamicMeshAABBTree,
+	const FVector& Pos,
+	float Radius,
+	TFunctionRef<float(int32)> GetVertexScale)
+{
+	const UE::Geometry::FDynamicMesh3* DynamicMesh = DynamicMeshAABBTree.GetMesh();
+	if (DynamicMesh == nullptr)
+	{
+		return FFindMeshElemsResult();
+	}
+
+	float GlobalMaxVertexScale = 0.f;
+	for (int32 VID : DynamicMesh->VertexIndicesItr())
+	{
+		GlobalMaxVertexScale = FMath::Max(GlobalMaxVertexScale, GetVertexScale(VID));
+	}
+
+	return FindTrianglesInRadius(DynamicMeshAABBTree, Pos, Radius, GetVertexScale, GlobalMaxVertexScale);
+}
+
+PV::DynamicMeshHelper::FFindMeshElemsResult PV::DynamicMeshHelper::FindTrianglesInRadius(
+	const UE::Geometry::FDynamicMeshAABBTree3& DynamicMeshAABBTree,
+	const FVector& Pos,
+	float Radius,
+	TFunctionRef<float(int32)> GetVertexScale,
+	float GlobalMaxVertexScale)
+{
+	FFindMeshElemsResult OutTriangleIDs;
+
+	const UE::Geometry::FDynamicMesh3* DynamicMesh = DynamicMeshAABBTree.GetMesh();
+	if (DynamicMesh == nullptr)
+	{
+		return OutTriangleIDs;
+	}
+
+	// Conservative upper bound: a box must be visited if any triangle inside it
+	// could pass the per-triangle test. The per-triangle check below still
+	// re-evaluates each candidate precisely.
+	const double BoxRadiusSqr = FMath::Square((double)(GlobalMaxVertexScale + Radius));
+
+	UE::Geometry::FDynamicMeshAABBTree3::FTreeTraversal TreeTraversal;
+
+	TreeTraversal.NextBoxF = [&](const UE::Geometry::FAxisAlignedBox3d& Box, int32 Depth) -> bool
+	{
+		return Box.DistanceSquared(Pos) <= BoxRadiusSqr;
+	};
+
+	TreeTraversal.NextTriangleF = [&](int32 TID)
+	{
+		const UE::Geometry::FIndex3i Triangle = DynamicMesh->GetTriangle(TID);
+
+		// Use the maximum vertex scale in the triangle as a conservative
+		// upper bound for TargetScale at the closest point on the triangle.
+		float MaxVertexScale = 0.f;
+		for (int32 j = 0; j < 3; ++j)
+		{
+			MaxVertexScale = FMath::Max(MaxVertexScale, GetVertexScale(Triangle[j]));
+		}
+
+		using FDynamicMeshQueries3 = UE::Geometry::TMeshQueries<UE::Geometry::FDynamicMesh3>;
+		const UE::Geometry::FDistPoint3Triangle3d Query = FDynamicMeshQueries3::TriangleDistance(*DynamicMesh, TID, Pos);
+
+		const float AllowedDist = MaxVertexScale + Radius;
+		if (FVector::DistSquared(Pos, Query.ClosestTrianglePoint) <= FMath::Square(AllowedDist))
+		{
+			OutTriangleIDs.Add(TID);
 		}
 	};
 
diff --git a/Engine/Plugins/Experimental/ProceduralVegetationEditor/Source/ProceduralVegetation/Private/Helpers/PVDynamicMeshHelpers.h b/Engine/Plugins/Experimental/ProceduralVegetationEditor/Source/ProceduralVegetation/Private/Helpers/PVDynamicMeshHelpers.h
--- a/Engine/Plugins/Experimental/ProceduralVegetationEditor/Source/ProceduralVegetation/Private/Helpers/PVDynamicMeshHelpers.h
+++ b/Engine/Plugins/Experimental/ProceduralVegetationEditor/Source/ProceduralVegetation/Private/Helpers/PVDynamicMeshHelpers.h
@@ -14,11 +14,26 @@ namespace UE::Geometry
 
 namespace PV::DynamicMeshHelper
 {
-	TArray<int32> GetUniqueVertexIndices(const UE::Geometry::FDynamicMesh3& DynamicMesh, const TArray<int32>& TriangleIds);
+	using FFindMeshElemsResult = TArray<int32, TInlineAllocator<4>>;
 
-	TArray<int32> GetVertexNeighbours(const UE::Geometry::FDynamicMesh3& DynamicMesh, int32 VertexID);
+	FFindMeshElemsResult GetUniqueVertexIndices(const UE::Geometry::FDynamicMesh3& DynamicMesh, const TArray<int32>& TriangleIds);
 
-	TArray<int32> FindVerticesInRadius(const UE::Geometry::FDynamicMeshAABBTree3& DynamicMeshAABBTree, const FVector& Pos, float Radius);
+	FFindMeshElemsResult GetVertexNeighbours(const UE::Geometry::FDynamicMesh3& DynamicMesh, int32 VertexID);
 
-	TArray<int32> FindTrianglesInRadius(const UE::Geometry::FDynamicMeshAABBTree3& DynamicMeshAABBTree, const FVector& Pos, float Radius);
+	FFindMeshElemsResult FindVerticesInRadius(const UE::Geometry::FDynamicMeshAABBTree3& DynamicMeshAABBTree, const FVector& Pos, float Radius);
+
+	FFindMeshElemsResult FindTrianglesInRadius(const UE::Geometry::FDynamicMeshAABBTree3& DynamicMeshAABBTree, const FVector& Pos, float Radius);
+
+	FFindMeshElemsResult FindTrianglesInRadius(
+		const UE::Geometry::FDynamicMeshAABBTree3& DynamicMeshAABBTree,
+		const FVector& Pos,
+		float Radius,
+		TFunctionRef<float(int32)> GetVertexScale);
+
+	FFindMeshElemsResult FindTrianglesInRadius(
+		const UE::Geometry::FDynamicMeshAABBTree3& DynamicMeshAABBTree,
+		const FVector& Pos,
+		float Radius,
+		TFunctionRef<float(int32)> GetVertexScale,
+		float GlobalMaxVertexScale);
 };
\ No newline at end of file
diff --git a/Engine/Plugins/Experimental/ProceduralVegetationEditor/Source/ProceduralVegetation/Private/Helpers/PVImportHelpers.cpp b/Engine/Plugins/Experimental/ProceduralVegetationEditor/Source/ProceduralVegetation/Private/Helpers/PVImportHelpers.cpp
--- a/Engine/Plugins/Experimental/ProceduralVegetationEditor/Source/ProceduralVegetation/Private/Helpers/PVImportHelpers.cpp
+++ b/Engine/Plugins/Experimental/ProceduralVegetationEditor/Source/ProceduralVegetation/Private/Helpers/PVImportHelpers.cpp
@@ -118,7 +118,7 @@ void PV::ImportHelper::ComputeLengthFromRoot(
 		const float PrevLengthFromRoot = LengthFromRootAttribute.GetValue(VertexID);
 		const FVector PrevVertexPosition = DynamicMesh.GetVertex(VertexID);
 
-		const TArray<int32> Neighbours = PV::DynamicMeshHelper::GetVertexNeighbours(DynamicMesh, VertexID);
+		const PV::DynamicMeshHelper::FFindMeshElemsResult Neighbours = PV::DynamicMeshHelper::GetVertexNeighbours(DynamicMesh, VertexID);
 		for (int32 NeighbourVertexID : Neighbours)
 		{
 			if (!EndPointsAttribute.GetValue(NeighbourVertexID) // Ingore vertex filter for end points as that means we've reached the end of the branch
@@ -173,7 +173,7 @@ void PV::ImportHelper::TraceBranches(
 	
 		const FVector ThisVertexPosition = DynamicMesh.GetVertex(PointIndex);
 		const float ThisLengthFromRoot = LengthFromRootAttribute.GetValue(PointIndex);
-		const TArray<int32> Neighbours = GetVertexNeighbours(DynamicMesh, PointIndex);
+		const PV::DynamicMeshHelper::FFindMeshElemsResult Neighbours = GetVertexNeighbours(DynamicMesh, PointIndex);
 		for (int32 Neighbour : Neighbours)
 		{
 			if (RootPointAttribute.GetValue(Neighbour))
@@ -381,7 +381,7 @@ void PV::ImportHelper::ExtractBranchHierarchies(
 	const auto FindNextBranchPoints = [&](int32 VertexID)
 	{
 		TArray<int32> NextBranchPoints;
-		const TArray<int32> Neighbours = GetVertexNeighbours(DynamicMesh, VertexID);
+		const PV::DynamicMeshHelper::FFindMeshElemsResult Neighbours = GetVertexNeighbours(DynamicMesh, VertexID);
 
 		for (int32 Neighbour : Neighbours)
 		{
@@ -898,6 +898,8 @@ void PV::ImportHelper::FillAttributesFromBranchHierarchy(
 	TArrayView<const FPVBranchHierarchyDescription> InBranchHierarchies
 )
 {
+	TRACE_CPUPROFILER_EVENT_SCOPE(FillAttributesFromBranchHierarchy);
+
 	const static auto BranchIndexToBranchNumber = [](int32 BranchIndex) 
 	{ 
 		// Branch Number is BranchIndex + 1 since branch with parent number 0 is a root in PVE and our root has parent index -1.
@@ -1341,6 +1343,8 @@ bool PV::ImportHelper::GenerateGrowthDataFromBranchHierarchies(FManagedArrayColl
 {
 	using namespace PV::AttributesHelper;
 
+	TRACE_CPUPROFILER_EVENT_SCOPE(GenerateGrowthDataFromBranchHierarchies);
+
 	if (InBranchHierarchies.Num() < 1)
 	{
 		return false;
diff --git a/Engine/Plugins/Experimental/ProceduralVegetationEditor/Source/ProceduralVegetation/Private/Helpers/PVPlantTraversalHelper.h b/Engine/Plugins/Experimental/ProceduralVegetationEditor/Source/ProceduralVegetation/Private/Helpers/PVPlantTraversalHelper.h
--- a/Engine/Plugins/Experimental/ProceduralVegetationEditor/Source/ProceduralVegetation/Private/Helpers/PVPlantTraversalHelper.h
+++ b/Engine/Plugins/Experimental/ProceduralVegetationEditor/Source/ProceduralVegetation/Private/Helpers/PVPlantTraversalHelper.h
@@ -5,13 +5,127 @@
 
 namespace PV::PlantTraversalHelper
 {
+	// Lookup table that maps branch numbers to their array indices.
+	// Internally uses an offset array for dense ranges or a TMap for sparse ranges,
+	// chosen at construction time based on the ratio of element count to value range.
+	struct FBranchNumberToIndexTable
+	{
+		struct FArrayTable
+		{
+			int32 Offset;
+			TArray<int32> Table;
+		};
+
+		using FMapTable = TMap<int32, int32>;
+
+		TVariant<FEmptyVariantState, FArrayTable, FMapTable> Table;
+
+		FBranchNumberToIndexTable() = default;
+		FBranchNumberToIndexTable(PV::FBranchNumberAttributeView BranchNumberAttribute) { Init(BranchNumberAttribute); }
+		FBranchNumberToIndexTable(PV::FBranchNumberAttributeConstView BranchNumberAttribute) { Init(BranchNumberAttribute); }
+
+		void Init(PV::FBranchNumberAttributeConstView BranchNumberAttribute)
+		{
+			if (BranchNumberAttribute.Num() == 0)
+			{
+				return;
+			}
+
+			int32 Min = TNumericLimits<int32>::Max();
+			int32 Max = TNumericLimits<int32>::Min();
+
+			for (int32 i = 0; i < BranchNumberAttribute.Num(); ++i)
+			{
+				Min = FMath::Min(Min, BranchNumberAttribute[i]);
+				Max = FMath::Max(Max, BranchNumberAttribute[i]);
+			}
+
+			// Allow up to 8192 empty slots (~32KB) in the array table before falling back to TMap.
+			// Array lookup has better cache locality; 32KB of overhead is negligible per plant.
+			constexpr int32 ArrayTableLimit = 8192;
+			const int32 Range = Max - Min;
+			if (Range <= BranchNumberAttribute.Num() + ArrayTableLimit)
+			{
+				FArrayTable ArrayTable;
+				ArrayTable.Table.Init(INDEX_NONE, Range + 1);
+				ArrayTable.Offset = Min;
+
+				for (int32 i = 0; i < BranchNumberAttribute.Num(); ++i)
+				{
+					const int32 BranchNumber = BranchNumberAttribute[i];
+					ArrayTable.Table[BranchNumber - ArrayTable.Offset] = i;
+				}
+
+				Table.Emplace<FArrayTable>(MoveTemp(ArrayTable));
+			}
+			else
+			{
+				FMapTable MapTable;
+				MapTable.Reserve(BranchNumberAttribute.Num());
+				for (int32 i = 0; i < BranchNumberAttribute.Num(); ++i)
+				{
+					MapTable.Add(BranchNumberAttribute[i], i);
+				}
+
+				Table.Emplace<FMapTable>(MoveTemp(MapTable));
+			}
+		}
+
+		int32 Find(int32 BranchNumber) const
+		{
+			if (!ensure(IsValid()))
+			{
+				return INDEX_NONE;
+			}
+
+			if (Table.IsType<FArrayTable>())
+			{
+				const FArrayTable& ArrayTable = Table.Get<FArrayTable>();
+				const int32 TableIndex = BranchNumber - ArrayTable.Offset;
+				return ArrayTable.Table.IsValidIndex(TableIndex) ? ArrayTable.Table[TableIndex] : INDEX_NONE;
+			}
+			else
+			{
+				const FMapTable& MapTable = Table.Get<FMapTable>();
+				const int32* BranchIndex = MapTable.Find(BranchNumber);
+				return BranchIndex ? *BranchIndex : INDEX_NONE;
+			}
+		}
+
+		bool IsValid() const
+		{
+			return !Table.IsType<FEmptyVariantState>();
+		}
+	};
+
+	inline int32 GetBranchParentIndex(
+		FBranchParentNumberAttributeConstView BranchParentNumberAttribute,
+		const FBranchNumberToIndexTable& BranchNumberToIndex,
+		int32 BranchIndex
+	)
+	{
+		if (!BranchParentNumberAttribute.IsValidIndex(BranchIndex) || !BranchNumberToIndex.IsValid())
+		{
+			return INDEX_NONE;
+		}
+
+		return BranchNumberToIndex.Find(BranchParentNumberAttribute[BranchIndex]);
+	}
+
+	// Slow variant of GetBranchParentIndex. Prefer GetBranchParentIndex with a FBranchNumberToIndexTable when a FBranchNumberToIndexTable is already available.
+	// The difference is in how the parent branch index is resolved:
+	//   GetBranchParentIndex        — takes a pre-built FBranchNumberToIndexTable and resolves the parent in O(1).
+	//   GetBranchParentIndex (Slow) — takes the raw BranchNumberAttribute and resolves the parent via a linear Find(),
+	//                                 making repeated calls O(N) per lookup.
+	// Use this variant only when constructing the lookup table upfront is not practical (e.g. one-off lookups
+	// in contexts where the table is not already available and building it would be wasteful).
 	inline int32 GetBranchParentIndex(
 		FBranchParentNumberAttributeConstView BranchParentNumberAttribute,
 		FBranchNumberAttributeConstView BranchNumberAttribute,
 		int32 BranchIndex
 	)
 	{
-		return BranchParentNumberAttribute.IsValidIndex(BranchIndex) 
+		return BranchParentNumberAttribute.IsValidIndex(BranchIndex)
 			? BranchNumberAttribute.Find(BranchParentNumberAttribute[BranchIndex])
 			: INDEX_NONE;
 	}
@@ -62,6 +176,37 @@ namespace PV::PlantTraversalHelper
 		return PlantNumberElements.Array();
 	}
 
+	inline bool HasFusedPoints(
+		PV::FBranchPointsAttributeConstView BranchPointsAttribute,
+		PV::FBranchParentNumberAttributeConstView BranchParentNumberAttribute,
+		const FBranchNumberToIndexTable& BranchNumberToIndex
+	)
+	{
+		for (int32 BranchIndex = 0; BranchIndex < BranchPointsAttribute.Num(); ++BranchIndex)
+		{
+			const int32 BranchParentIndex = GetBranchParentIndex(BranchParentNumberAttribute, BranchNumberToIndex, BranchIndex);
+			if (BranchParentIndex != INDEX_NONE)
+			{
+				const TArray<int32>& BranchPoints = BranchPointsAttribute[BranchIndex];
+				const TArray<int32>& ParentBranchPoints = BranchPointsAttribute[BranchParentIndex];
+
+				if (BranchPoints.Num() > 0 && ParentBranchPoints.Num() > 0)
+				{
+					return ParentBranchPoints.Find(BranchPoints[0]) != INDEX_NONE;
+				}
+			}
+		}
+
+		return true;
+	}
+
+	// Slow variant of HasFusedPoints. Prefer HasFusedPoints with a FBranchNumberToIndexTable when a FBranchNumberToIndexTable is already available.
+	// The difference is in how the parent branch index is resolved:
+	//   HasFusedPoints        — takes a pre-built FBranchNumberToIndexTable and resolves each parent in O(1).
+	//   HasFusedPoints (Slow) — takes the raw BranchNumberAttribute and resolves each parent via a linear Find(),
+	//                           making the overall complexity O(N^2) in the number of branches.
+	// Use this variant only when constructing the lookup table upfront is not practical (e.g. one-off checks
+	// in contexts where the table is not already available and building it would be wasteful).
 	inline bool HasFusedPoints(
 		PV::FBranchPointsAttributeConstView BranchPointsAttribute,
 		PV::FBranchParentNumberAttributeConstView BranchParentNumberAttribute,
@@ -97,12 +242,12 @@ namespace PV::PlantTraversalHelper
 	void RecursiveWalkBranches(
 		FBranchParentNumberAttributeConstView BranchParentNumberAttribute,
 		FBranchChildrenAttributeConstView BranchChildrenAttribute,
-		FBranchNumberAttributeConstView BranchNumberAttribute,
+		const FBranchNumberToIndexTable& BranchNumberToIndex,
 		const StateType& RootBranchState,
 		const TFunction<EForEachResult(int32, StateType&)>& Body
 	)
 	{
-		if (!ValidateAttributeCollection(BranchParentNumberAttribute, BranchChildrenAttribute, BranchNumberAttribute))
+		if (!ValidateAttributeCollection(BranchParentNumberAttribute, BranchChildrenAttribute))
 		{
 			return;
 		}
@@ -118,13 +263,13 @@ namespace PV::PlantTraversalHelper
 			const TArray<int32>& BranchChildren = BranchChildrenAttribute[BranchIndex];
 			for (int32 ChildBranchNumber : BranchChildren)
 			{
-				const int32 ChildBranchIndex = BranchNumberAttribute.Find(ChildBranchNumber);
+				const int32 ChildBranchIndex = BranchNumberToIndex.Find(ChildBranchNumber);
 				if (!ensure(ChildBranchIndex != INDEX_NONE))
 				{
 					continue;
 				}
 
-				if (BranchParentNumberAttribute[ChildBranchIndex] != BranchNumberAttribute[BranchIndex])
+				if (BranchNumberToIndex.Find(BranchParentNumberAttribute[ChildBranchIndex]) != BranchIndex)
 				{
 					continue;
 				}
@@ -150,14 +295,14 @@ namespace PV::PlantTraversalHelper
 	inline void RecursiveWalkBranches(
 		FBranchParentNumberAttributeConstView BranchParentNumberAttribute,
 		FBranchChildrenAttributeConstView BranchChildrenAttribute,
-		FBranchNumberAttributeConstView BranchNumberAttribute,
+		const FBranchNumberToIndexTable& BranchNumberToIndex,
 		const TFunction<EForEachResult(int32 BranchIndex)>& Body
 	)
 	{
 		RecursiveWalkBranches<int32>(
 			BranchParentNumberAttribute,
 			BranchChildrenAttribute,
-			BranchNumberAttribute,
+			BranchNumberToIndex,
 			INDEX_NONE,
 			[&](int32 BranchIndex, int32&) -> EForEachResult
 			{
@@ -174,7 +319,7 @@ namespace PV::PlantTraversalHelper
 	inline void RecursiveWalkBranches_Reversed(
 		FBranchParentNumberAttributeConstView BranchParentNumberAttribute,
 		FBranchChildrenAttributeConstView BranchChildrenAttribute,
-		FBranchNumberAttributeConstView BranchNumberAttribute,
+		const FBranchNumberToIndexTable& BranchNumberToIndex,
 		const TFunction<EForEachResult(int32 BranchIndex)>& Body
 	)
 	{
@@ -184,7 +329,7 @@ namespace PV::PlantTraversalHelper
 		RecursiveWalkBranches(
 			BranchParentNumberAttribute,
 			BranchChildrenAttribute,
-			BranchNumberAttribute,
+			BranchNumberToIndex,
 			[&](int32 BranchIndex) -> EForEachResult
 			{
 				BranchIndices.Add(BranchIndex);
@@ -258,14 +403,14 @@ namespace PV::PlantTraversalHelper
 	inline TArray<int32> GetBranchPointChildBranchIndices_SplitPoints(
 		PV::FBranchChildrenAttributeConstView BranchChildrenAttribute,
 		PV::FBranchParentNumberAttributeConstView BranchParentNumberAttribute,
-		PV::FBranchNumberAttributeConstView BranchNumberAttribute,
+		const FBranchNumberToIndexTable& BranchNumberToIndex,
 		PV::FBranchPointsAttributeConstView BranchPointsAttribute,
 		PV::FPointLengthFromRootAttributeConstView PointLengthFromRootAttribute,
 		int32 BranchIndex,
 		int32 BranchPointIndex
 	)
 	{
-		if (!ValidateAttributeCollection(BranchChildrenAttribute, BranchParentNumberAttribute, BranchNumberAttribute, BranchPointsAttribute, PointLengthFromRootAttribute))
+		if (!ValidateAttributeCollection(BranchChildrenAttribute, BranchParentNumberAttribute, BranchPointsAttribute, PointLengthFromRootAttribute))
 		{
 			return TArray<int32>();
 		}
@@ -300,13 +445,13 @@ namespace PV::PlantTraversalHelper
 
 		for (const int32 BranchChild : BranchChildren)
 		{
-			const int32 ChildIndex = BranchNumberAttribute.Find(BranchChild);
+			const int32 ChildIndex = BranchNumberToIndex.Find(BranchChild);
 			if (ChildIndex == INDEX_NONE)
 			{
 				continue;
 			}
 
-			const int32 ParentIndex = GetBranchParentIndex(BranchParentNumberAttribute, BranchNumberAttribute, ChildIndex);
+			const int32 ParentIndex = GetBranchParentIndex(BranchParentNumberAttribute, BranchNumberToIndex, ChildIndex);
 			if (ParentIndex != BranchIndex)
 			{
 				continue;
@@ -334,13 +479,13 @@ namespace PV::PlantTraversalHelper
 	inline TArray<int32> GetBranchPointChildBranchIndices_FusedPoints(
 		PV::FBranchChildrenAttributeConstView BranchChildrenAttribute,
 		PV::FBranchParentNumberAttributeConstView BranchParentNumberAttribute,
-		PV::FBranchNumberAttributeConstView BranchNumberAttribute,
+		const FBranchNumberToIndexTable& BranchNumberToIndex,
 		PV::FBranchPointsAttributeConstView BranchPointsAttribute,
 		int32 BranchIndex,
 		int32 BranchPointIndex
 	)
 	{
-		if (!ValidateAttributeCollection(BranchChildrenAttribute, BranchParentNumberAttribute, BranchNumberAttribute, BranchPointsAttribute))
+		if (!ValidateAttributeCollection(BranchChildrenAttribute, BranchParentNumberAttribute, BranchPointsAttribute))
 		{
 			return TArray<int32>();
 		}
@@ -364,13 +509,13 @@ namespace PV::PlantTraversalHelper
 
 		for (const int32 BranchChild : BranchChildren)
 		{
-			const int32 ChildIndex = BranchNumberAttribute.Find(BranchChild);
+			const int32 ChildIndex = BranchNumberToIndex.Find(BranchChild);
 			if (ChildIndex == INDEX_NONE)
 			{
 				continue;
 			}
 
-			const int32 ParentIndex = GetBranchParentIndex(BranchParentNumberAttribute, BranchNumberAttribute, ChildIndex);
+			const int32 ParentIndex = GetBranchParentIndex(BranchParentNumberAttribute, BranchNumberToIndex, ChildIndex);
 			if (ParentIndex != BranchIndex)
 			{
 				continue;
@@ -395,19 +540,19 @@ namespace PV::PlantTraversalHelper
 	inline TArray<int32> GetBranchPointChildBranchIndices(
 		PV::FBranchChildrenAttributeConstView BranchChildrenAttribute,
 		PV::FBranchParentNumberAttributeConstView BranchParentNumberAttribute,
-		PV::FBranchNumberAttributeConstView BranchNumberAttribute,
+		const FBranchNumberToIndexTable BranchNumberToIndex,
 		PV::FBranchPointsAttributeConstView BranchPointsAttribute,
 		PV::FPointLengthFromRootAttributeConstView PointLengthFromRootAttribute,
 		int32 BranchIndex,
 		int32 BranchPointIndex
 	)
 	{
-		if (HasFusedPoints(BranchPointsAttribute, BranchParentNumberAttribute, BranchNumberAttribute))
+		if (HasFusedPoints(BranchPointsAttribute, BranchParentNumberAttribute, BranchNumberToIndex))
 		{
 			return GetBranchPointChildBranchIndices_FusedPoints(
 				BranchChildrenAttribute,
 				BranchParentNumberAttribute,
-				BranchNumberAttribute,
+				BranchNumberToIndex,
 				BranchPointsAttribute,
 				BranchIndex,
 				BranchPointIndex
@@ -418,7 +563,7 @@ namespace PV::PlantTraversalHelper
 			return GetBranchPointChildBranchIndices_SplitPoints(
 				BranchChildrenAttribute,
 				BranchParentNumberAttribute,
-				BranchNumberAttribute,
+				BranchNumberToIndex,
 				BranchPointsAttribute,
 				PointLengthFromRootAttribute,
 				BranchIndex,
@@ -441,7 +586,7 @@ namespace PV::PlantTraversalHelper
 	inline EForEachResult RecursiveWalkBranchPoints_SplitPoints(
 		PV::FBranchChildrenAttributeConstView BranchChildrenAttribute,
 		PV::FBranchParentNumberAttributeConstView BranchParentNumberAttribute,
-		PV::FBranchNumberAttributeConstView BranchNumberAttribute,
+		const FBranchNumberToIndexTable& BranchNumberToIndex,
 		PV::FBranchPointsAttributeConstView BranchPointsAttribute,
 		PV::FPointLengthFromRootAttributeConstView PointLengthFromRootAttribute,
 		int32 BranchIndex,
@@ -449,7 +594,7 @@ namespace PV::PlantTraversalHelper
 		const TFunction<EForEachResult(const FRecursiveWalkBranchPointsParams& Params, StateType& PointState)>& Body
 	)
 	{
-		if (!ValidateAttributeCollection(BranchChildrenAttribute, BranchParentNumberAttribute, BranchNumberAttribute, BranchPointsAttribute, PointLengthFromRootAttribute))
+		if (!ValidateAttributeCollection(BranchChildrenAttribute, BranchParentNumberAttribute, BranchPointsAttribute, PointLengthFromRootAttribute))
 		{
 			return EForEachResult::Break;
 		}
@@ -464,7 +609,7 @@ namespace PV::PlantTraversalHelper
 			const TArray<int32> ChildBranchIndices = GetBranchPointChildBranchIndices_SplitPoints(
 				BranchChildrenAttribute,
 				BranchParentNumberAttribute,
-				BranchNumberAttribute,
+				BranchNumberToIndex,
 				BranchPointsAttribute,
 				PointLengthFromRootAttribute,
 				BranchIndex,
@@ -535,14 +680,14 @@ namespace PV::PlantTraversalHelper
 	inline EForEachResult RecursiveWalkBranchPoints_FusedPoints(
 		PV::FBranchChildrenAttributeConstView BranchChildrenAttribute,
 		PV::FBranchParentNumberAttributeConstView BranchParentNumberAttribute,
-		PV::FBranchNumberAttributeConstView BranchNumberAttribute,
+		const FBranchNumberToIndexTable& BranchNumberToIndex,
 		PV::FBranchPointsAttributeConstView BranchPointsAttribute,
 		int32 BranchIndex,
 		const StateType& RootPointState,
 		const TFunction<EForEachResult(const FRecursiveWalkBranchPointsParams& Params, StateType& PointState)>& Body
 	)
 	{
-		if (!ValidateAttributeCollection(BranchChildrenAttribute, BranchParentNumberAttribute, BranchNumberAttribute, BranchPointsAttribute))
+		if (!ValidateAttributeCollection(BranchChildrenAttribute, BranchParentNumberAttribute, BranchPointsAttribute))
 		{
 			return EForEachResult::Break;
 		}
@@ -557,7 +702,7 @@ namespace PV::PlantTraversalHelper
 			const TArray<int32> ChildBranchIndices = GetBranchPointChildBranchIndices_FusedPoints(
 				BranchChildrenAttribute,
 				BranchParentNumberAttribute,
-				BranchNumberAttribute,
+				BranchNumberToIndex,
 				BranchPointsAttribute,
 				BranchIndex,
 				BranchPointIndex
@@ -627,7 +772,7 @@ namespace PV::PlantTraversalHelper
 	inline EForEachResult RecursiveWalkBranchPoints(
 		PV::FBranchChildrenAttributeConstView BranchChildrenAttribute,
 		PV::FBranchParentNumberAttributeConstView BranchParentNumberAttribute,
-		PV::FBranchNumberAttributeConstView BranchNumberAttribute,
+		const FBranchNumberToIndexTable& BranchNumberToIndex,
 		PV::FBranchPointsAttributeConstView BranchPointsAttribute,
 		PV::FPointLengthFromRootAttributeConstView PointLengthFromRootAttribute,
 		int32 BranchIndex,
@@ -635,12 +780,12 @@ namespace PV::PlantTraversalHelper
 		const TFunction<EForEachResult(const FRecursiveWalkBranchPointsParams& Params, StateType& PointState)>& Body
 	)
 	{
-		if (HasFusedPoints(BranchPointsAttribute, BranchParentNumberAttribute, BranchNumberAttribute))
+		if (HasFusedPoints(BranchPointsAttribute, BranchParentNumberAttribute, BranchNumberToIndex))
 		{
 			return RecursiveWalkBranchPoints_FusedPoints<StateType, bWalkAxillaryFirst>(
 				BranchChildrenAttribute,
 				BranchParentNumberAttribute,
-				BranchNumberAttribute,
+				BranchNumberToIndex,
 				BranchPointsAttribute,
 				BranchIndex,
 				RootPointState,
@@ -652,7 +797,7 @@ namespace PV::PlantTraversalHelper
 			return RecursiveWalkBranchPoints_SplitPoints<StateType, bWalkAxillaryFirst>(
 				BranchChildrenAttribute,
 				BranchParentNumberAttribute,
-				BranchNumberAttribute,
+				BranchNumberToIndex,
 				BranchPointsAttribute,
 				PointLengthFromRootAttribute,
 				BranchIndex,
@@ -666,65 +811,47 @@ namespace PV::PlantTraversalHelper
 	inline EForEachResult RecursiveWalkBranchPoints(
 		PV::FBranchChildrenAttributeConstView BranchChildrenAttribute,
 		PV::FBranchParentNumberAttributeConstView BranchParentNumberAttribute,
-		PV::FBranchNumberAttributeConstView BranchNumberAttribute,
+		const FBranchNumberToIndexTable& BranchNumberToIndex,
 		PV::FBranchPointsAttributeConstView BranchPointsAttribute,
 		PV::FPointLengthFromRootAttributeConstView PointLengthFromRootAttribute,
 		int32 BranchIndex,
 		const TFunction<EForEachResult(const FRecursiveWalkBranchPointsParams& Params)>& Body
 	)
 	{
-		if (HasFusedPoints(BranchPointsAttribute, BranchParentNumberAttribute, BranchNumberAttribute))
-		{
-			return RecursiveWalkBranchPoints_FusedPoints<int32, bWalkAxillaryFirst>(
-				BranchChildrenAttribute,
-				BranchParentNumberAttribute,
-				BranchNumberAttribute,
-				BranchPointsAttribute,
-				BranchIndex,
-				0,
-				[&](const FRecursiveWalkBranchPointsParams& Params, int32&) -> EForEachResult
-				{
-					return Body(Params);
-				}
-			);
-		}
-		else
-		{
-			return RecursiveWalkBranchPoints_SplitPoints<int32, bWalkAxillaryFirst>(
-				BranchChildrenAttribute,
-				BranchParentNumberAttribute,
-				BranchNumberAttribute,
-				BranchPointsAttribute,
-				PointLengthFromRootAttribute,
-				BranchIndex,
-				0,
-				[&](const FRecursiveWalkBranchPointsParams& Params, int32&) -> EForEachResult
-				{
-					return Body(Params);
-				}
-			);
-		}
+		return RecursiveWalkBranchPoints<int32, bWalkAxillaryFirst>(
+			BranchChildrenAttribute,
+			BranchParentNumberAttribute,
+			BranchNumberToIndex,
+			BranchPointsAttribute,
+			PointLengthFromRootAttribute,
+			BranchIndex,
+			0,
+			[&](const FRecursiveWalkBranchPointsParams& Params, int32&) -> EForEachResult
+			{
+				return Body(Params);
+			}
+		);
 	}
 
 	template <typename StateType, bool bWalkAxillaryBranchesFirst = false>
 	void RecursiveWalkPlantPoints(
 		PV::FBranchChildrenAttributeConstView BranchChildrenAttribute,
 		PV::FBranchParentNumberAttributeConstView BranchParentNumberAttribute,
-		PV::FBranchNumberAttributeConstView BranchNumberAttribute,
+		const FBranchNumberToIndexTable& BranchNumberToIndex,
 		PV::FBranchPointsAttributeConstView BranchPointsAttribute,
 		PV::FPointLengthFromRootAttributeConstView PointLengthFromRootAttribute,
 		const StateType& RootPointState,
 		const TFunction<EForEachResult(const FRecursiveWalkBranchPointsParams& Params, StateType& PointState)>& Body
 	)
 	{
-		if (HasFusedPoints(BranchPointsAttribute, BranchParentNumberAttribute, BranchNumberAttribute))
+		if (HasFusedPoints(BranchPointsAttribute, BranchParentNumberAttribute, BranchNumberToIndex))
 		{
 			for (const int32 TrunkIndex : GetTrunkIndices(BranchParentNumberAttribute))
 			{
 				auto Result = RecursiveWalkBranchPoints_FusedPoints<StateType, bWalkAxillaryBranchesFirst>(
 					BranchChildrenAttribute,
 					BranchParentNumberAttribute,
-					BranchNumberAttribute,
+					BranchNumberToIndex,
 					BranchPointsAttribute,
 					TrunkIndex,
 					RootPointState,
@@ -743,7 +870,7 @@ namespace PV::PlantTraversalHelper
 				auto Result = RecursiveWalkBranchPoints_SplitPoints<StateType, bWalkAxillaryBranchesFirst>(
 					BranchChildrenAttribute,
 					BranchParentNumberAttribute,
-					BranchNumberAttribute,
+					BranchNumberToIndex,
 					BranchPointsAttribute,
 					PointLengthFromRootAttribute,
 					TrunkIndex,
@@ -762,7 +889,7 @@ namespace PV::PlantTraversalHelper
 	inline void RecursiveWalkPlantPoints(
 		PV::FBranchChildrenAttributeConstView BranchChildrenAttribute,
 		PV::FBranchParentNumberAttributeConstView BranchParentNumberAttribute,
-		PV::FBranchNumberAttributeConstView BranchNumberAttribute,
+		const FBranchNumberToIndexTable& BranchNumberToIndex,
 		PV::FBranchPointsAttributeConstView BranchPointsAttribute,
 		PV::FPointLengthFromRootAttributeConstView PointLengthFromRootAttribute,
 		const TFunction<EForEachResult(const FRecursiveWalkBranchPointsParams& Params)>& Body
@@ -771,7 +898,7 @@ namespace PV::PlantTraversalHelper
 		RecursiveWalkPlantPoints<int32, bWalkAxillaryBranchesFirst>(
 			BranchChildrenAttribute,
 			BranchParentNumberAttribute,
-			BranchNumberAttribute,
+			BranchNumberToIndex,
 			BranchPointsAttribute,
 			PointLengthFromRootAttribute,
 			INDEX_NONE,
diff --git a/Engine/Plugins/Experimental/ProceduralVegetationEditor/Source/ProceduralVegetation/Private/Implementations/PVImporter_StaticMesh.cpp b/Engine/Plugins/Experimental/ProceduralVegetationEditor/Source/ProceduralVegetation/Private/Implementations/PVImporter_StaticMesh.cpp
--- a/Engine/Plugins/Experimental/ProceduralVegetationEditor/Source/ProceduralVegetation/Private/Implementations/PVImporter_StaticMesh.cpp
+++ b/Engine/Plugins/Experimental/ProceduralVegetationEditor/Source/ProceduralVegetation/Private/Implementations/PVImporter_StaticMesh.cpp
@@ -128,7 +128,7 @@ static void ComputeVerticesCenterAndRadius_Internal_V1(
 	// center point to lay between the current vertex, and the found vertex. It then performs another overlap, and repeats the process until it no longer overlaps any vertices.
 	// Once an opposing vertex has been found, an extra check is performed to make sure the normal of the closest triangle is facing away from the center point.
 	
-	const TArray<int32> ComponentVertexIndices = PV::DynamicMeshHelper::GetUniqueVertexIndices(DynamicMesh, TriangleIds);
+	const PV::DynamicMeshHelper::FFindMeshElemsResult ComponentVertexIndices = PV::DynamicMeshHelper::GetUniqueVertexIndices(DynamicMesh, TriangleIds);
 
 	FDynamicMeshAABBTree3 DynamicMeshAABBTree(&DynamicMesh, false);
 	DynamicMeshAABBTree.Build(TriangleIds);
@@ -240,7 +240,7 @@ static void ComputeVerticesCenterAndRadius_Internal_V2(
 	// 3. It also not work well if the mesh is built up using only faces (e.g flat geometry instead of a tube for the branches)
 	// 4. The normal of the vertex may not point directly towards the opposite side (see "Remove extreme radius values" section below for more information)
 
-	const TArray<int32> ComponentVertexIndices = PV::DynamicMeshHelper::GetUniqueVertexIndices(DynamicMesh, TriangleIds);
+	const PV::DynamicMeshHelper::FFindMeshElemsResult ComponentVertexIndices = PV::DynamicMeshHelper::GetUniqueVertexIndices(DynamicMesh, TriangleIds);
 
 	FDynamicMeshAABBTree3 DynamicMeshAABBTree(&DynamicMesh, false);
 	DynamicMeshAABBTree.Build(TriangleIds);
@@ -432,7 +432,7 @@ static void ExtractReductionCloud_Internal(
 	FDynamicMeshAABBTree3 DynamicMeshAABBTree(&DynamicMesh, false);
 	DynamicMeshAABBTree.Build(TriangleIds);
 
-	TArray<int32> VertexIndices = PV::DynamicMeshHelper::GetUniqueVertexIndices(DynamicMesh, TriangleIds);
+	PV::DynamicMeshHelper::FFindMeshElemsResult VertexIndices = PV::DynamicMeshHelper::GetUniqueVertexIndices(DynamicMesh, TriangleIds);
 
 	const int32 MaxVertexID = DynamicMesh.MaxVertexID();
 		
@@ -476,7 +476,7 @@ static void ExtractReductionCloud_Internal(
 		}
 
 		const float SearchRadius = VertexRadius * SearchRadiusInflate;
-		const TArray<int32> OverlappingVertices = PV::DynamicMeshHelper::FindVerticesInRadius(DynamicMeshAABBTree, VertexPosition, SearchRadius);
+		const PV::DynamicMeshHelper::FFindMeshElemsResult OverlappingVertices = PV::DynamicMeshHelper::FindVerticesInRadius(DynamicMeshAABBTree, VertexPosition, SearchRadius);
 		for (int32 j = 0; j < OverlappingVertices.Num(); j++)
 		{
 			const int32 OverlappingVertexIndex = OverlappingVertices[j];
@@ -745,7 +745,7 @@ static void ConnectBranchesToParent(
 			const float SearchRadius = 100.f; // TODO: Base this on the trunk radius?
 
 			// Potential optimization: FindTrianglesInRadius could return a large number of triangles, could use FDynamicMeshAABBTree3::FindNearestTriangle
-			const TArray<int32> CloseTriangles = PV::DynamicMeshHelper::FindTrianglesInRadius(DynamicMeshAABBTree, VertexPosition, SearchRadius);
+			const PV::DynamicMeshHelper::FFindMeshElemsResult CloseTriangles = PV::DynamicMeshHelper::FindTrianglesInRadius(DynamicMeshAABBTree, VertexPosition, SearchRadius);
 
 			TMap<int32, TArray<FIndex3i>> NodeToTriangles;
 			NodeToTriangles.Reserve(Nodes.Num());
@@ -1057,14 +1057,26 @@ static void ConnectAndTraceBranches(
 
 	// Build candidate edges
 	const FDynamicMeshAABBTree3 DynamicMeshAABBTree(&DynamicMesh);
+	float GlobalMaxVertexScale = 0.f;
+	for (int32 VID : DynamicMesh.VertexIndicesItr())
+	{
+		GlobalMaxVertexScale = FMath::Max(GlobalMaxVertexScale, VertexScaleAttribute.GetValue(VID));
+	}
+
 	for (int32 EndPointIndex = 0; EndPointIndex < NumEndPoints; ++EndPointIndex)
 	{
 		const int32 EndPointVertexIndex = EndPointVertexIndices[EndPointIndex];
 		const FVector VertexPosition    = DynamicMesh.GetVertex(EndPointVertexIndex);
 		const int32 EndPointNodeIndex   = VertexIndexToNodeIndex[EndPointVertexIndex];
-		const float SearchRadius        = 100.f;
-
-		const TArray<int32> CloseTriangles = PV::DynamicMeshHelper::FindTrianglesInRadius(DynamicMeshAABBTree, VertexPosition, SearchRadius);
+		const float ThisPointScale      = VertexScaleAttribute.GetValue(EndPointVertexIndex);
+		const float ScaleInflation      = 1.25f;
+		const PV::DynamicMeshHelper::FFindMeshElemsResult CloseTriangles = PV::DynamicMeshHelper::FindTrianglesInRadius(
+			DynamicMeshAABBTree,
+			VertexPosition,
+			ThisPointScale * ScaleInflation,
+			[&VertexScaleAttribute, ScaleInflation](int32 VID) { return VertexScaleAttribute.GetValue(VID) * ScaleInflation; },
+			GlobalMaxVertexScale * ScaleInflation
+		);
 
 		TMap<int32, TArray<FIndex3i>> NodeToTriangles;
 		for (const int32 TriangleID : CloseTriangles)
@@ -1089,7 +1101,12 @@ static void ConnectAndTraceBranches(
 					FIndex2i Edge = FIndex2i(Triangle[j], Triangle[(j + 1) % 3]);
 					Edge.Sort();
 
-					if (FillerVertexAttribute.GetValue(Edge.A) || FillerVertexAttribute.GetValue(Edge.B))
+					// Do not allow attchments to end points as end point vertices can be removed during fusing step further down.
+					// This could cause issues where we fail to connect certain branches if an end-point is the only valid connection point.
+					const bool bEdgeHasEndPoint = EndPointAttribute.GetValue(Edge.A) || EndPointAttribute.GetValue(Edge.B);
+
+					if (FillerVertexAttribute.GetValue(Edge.A) || FillerVertexAttribute.GetValue(Edge.B)
+						|| bEdgeHasEndPoint)
 					{
 						continue;
 					}
@@ -1105,10 +1122,9 @@ static void ConnectAndTraceBranches(
 					const float VertexScaleA    = VertexScaleAttribute.GetValue(Edge.A);
 					const float VertexScaleB    = VertexScaleAttribute.GetValue(Edge.B);
 					const float TargetScale     = FMath::Lerp(VertexScaleA, VertexScaleB, EdgeAlpha);
-					const float ThisPointScale  = VertexScaleAttribute.GetValue(EndPointVertexIndex);
 
 					const float DistToEdge = (VertexPosition - ClosestPoint).Size();
-					if (DistToEdge > (TargetScale + ThisPointScale) * 1.5f)
+					if (DistToEdge > (TargetScale + ThisPointScale) * ScaleInflation)
 					{
 						continue; // radii do not overlap — not a valid attachment
 					}
diff --git a/Engine/Plugins/Experimental/ProceduralVegetationEditor/Source/ProceduralVegetation/Private/Implementations/PVImporter_Texture2D.cpp b/Engine/Plugins/Experimental/ProceduralVegetationEditor/Source/ProceduralVegetation/Private/Implementations/PVImporter_Texture2D.cpp
--- a/Engine/Plugins/Experimental/ProceduralVegetationEditor/Source/ProceduralVegetation/Private/Implementations/PVImporter_Texture2D.cpp
+++ b/Engine/Plugins/Experimental/ProceduralVegetationEditor/Source/ProceduralVegetation/Private/Implementations/PVImporter_Texture2D.cpp
@@ -631,7 +631,7 @@ namespace Internal
 			 	check(DynamicMesh.IsVertex(PointIndex));
 			 		
 			 	int32 NumInternalNeighbours = 0;
-				const TArray<int32> PointNeighbours = PV::DynamicMeshHelper::GetVertexNeighbours(DynamicMesh, PointIndex);
+				const PV::DynamicMeshHelper::FFindMeshElemsResult PointNeighbours = PV::DynamicMeshHelper::GetVertexNeighbours(DynamicMesh, PointIndex);
 			 	for (int32 NeighbourIndex : PointNeighbours)
 			 	{
 			 		if (!PerimeterAttribute.GetValue(NeighbourIndex))
diff --git a/Engine/Plugins/Experimental/ProceduralVegetationEditor/Source/ProceduralVegetation/Private/Implementations/PVSimplifyPlantSkeleton.cpp b/Engine/Plugins/Experimental/ProceduralVegetationEditor/Source/ProceduralVegetation/Private/Implementations/PVSimplifyPlantSkeleton.cpp
--- a/Engine/Plugins/Experimental/ProceduralVegetationEditor/Source/ProceduralVegetation/Private/Implementations/PVSimplifyPlantSkeleton.cpp
+++ b/Engine/Plugins/Experimental/ProceduralVegetationEditor/Source/ProceduralVegetation/Private/Implementations/PVSimplifyPlantSkeleton.cpp
@@ -28,7 +28,8 @@ void PV::SimplifyPlantSkeleton(FManagedArrayCollection& InOutCollection, float S
 		return;
 	}
 
-	const bool bHasFusedPoints = HasFusedPoints(BranchPointsAttribute, BranchParentNumberAttribute, BranchNumberAttribute);
+	const FBranchNumberToIndexTable BranchNumberToIndex(BranchNumberAttribute);
+	const bool bHasFusedPoints = HasFusedPoints(BranchPointsAttribute, BranchParentNumberAttribute, BranchNumberToIndex);
 	const auto GetNumAxillaryBranches = [&](int32 BranchIndex, int32 BranchPointIndex)
 	{
 		if (bHasFusedPoints)
@@ -36,7 +37,7 @@ void PV::SimplifyPlantSkeleton(FManagedArrayCollection& InOutCollection, float S
 			return GetBranchPointChildBranchIndices_FusedPoints(
 				BranchChildrenAttribute,
 				BranchParentNumberAttribute,
-				BranchNumberAttribute,
+				BranchNumberToIndex,
 				BranchPointsAttribute,
 				BranchIndex,
 				BranchPointIndex
@@ -47,7 +48,7 @@ void PV::SimplifyPlantSkeleton(FManagedArrayCollection& InOutCollection, float S
 			return GetBranchPointChildBranchIndices_SplitPoints(
 				BranchChildrenAttribute,
 				BranchParentNumberAttribute,
-				BranchNumberAttribute,
+				BranchNumberToIndex,
 				BranchPointsAttribute,
 				PointLengthFromRootAttribute,
 				BranchIndex,
